[PCG] Fix crashes due to unitialized data coming from GPU Attribute Partition and make sure we return valid null data from Count Unique Attribute Values to fix downstream ensures.

#rnx

diff --git a/Engine/Plugins/PCG/Source/PCG/Private/Compute/BuiltInKernels/PCGCountUniqueAttributeValuesKernel.cpp b/Engine/Plugins/PCG/Source/PCG/Private/Compute/BuiltInKernels/PCGCountUniqueAttributeValuesKernel.cpp
--- a/Engine/Plugins/PCG/Source/PCG/Private/Compute/BuiltInKernels/PCGCountUniqueAttributeValuesKernel.cpp
+++ b/Engine/Plugins/PCG/Source/PCG/Private/Compute/BuiltInKernels/PCGCountUniqueAttributeValuesKernel.cpp
@@ -102,7 +102,7 @@ TSharedPtr<const FPCGDataCollectionDesc> UPCGCountUniqueAttributeValuesKernel::C
 	const FPCGKernelPin InputKernelPin(KernelIndex, PCGPinConstants::DefaultInputLabel, /*bIsInput=*/true);
 	const TSharedPtr<const FPCGDataCollectionDesc> InputDesc = InBinding->ComputeKernelPinDataDesc(InputKernelPin);
 
-	if (!ensure(InputDesc) || InputDesc->DataDescs.IsEmpty())
+	if (!ensure(InputDesc))
 	{
 		return nullptr;
 	}
@@ -161,6 +161,22 @@ TSharedPtr<const FPCGDataCollectionDesc> UPCGCountUniqueAttributeValuesKernel::C
 			OutDataDesc->DataDescs.Add(MoveTemp(CountDataDesc));
 		}
 	}
+	else
+	{
+		// Attribute not present on input data. We can't create a zero sized buffer on the GPU, so create a null output.
+		if (bOutputRawBuffer)
+		{
+			// Null output: One pair (string key 0 (invalid), value count 0).
+			OutDataDesc->DataDescs.Emplace(FPCGDataTypeInfoRawBuffer::AsId(), 2);
+		}
+		else
+		{
+			// Null output: One entry: string key 0, value count 0.
+			FPCGDataDesc CountDataDesc = OutDataDesc->DataDescs.Emplace_GetRef(EPCGDataType::Param, 1);
+			CountDataDesc.AddAttribute(FPCGKernelAttributeKey(PCGCountUniqueAttributeValuesConstants::ValueAttributeName, EPCGKernelAttributeType::Int), InBinding);
+			CountDataDesc.AddAttribute(FPCGKernelAttributeKey(PCGCountUniqueAttributeValuesConstants::ValueCountAttributeName, EPCGKernelAttributeType::Int), InBinding);
+		}
+	}
 
 	return OutDataDesc;
 }
diff --git a/Engine/Plugins/PCG/Source/PCG/Private/Compute/DataInterfaces/Elements/PCGStaticMeshSpawnerDataInterface.cpp b/Engine/Plugins/PCG/Source/PCG/Private/Compute/DataInterfaces/Elements/PCGStaticMeshSpawnerDataInterface.cpp
--- a/Engine/Plugins/PCG/Source/PCG/Private/Compute/DataInterfaces/Elements/PCGStaticMeshSpawnerDataInterface.cpp
+++ b/Engine/Plugins/PCG/Source/PCG/Private/Compute/DataInterfaces/Elements/PCGStaticMeshSpawnerDataInterface.cpp
@@ -27,7 +27,9 @@
 #include "ShaderCore.h"
 #include "ShaderParameterMetadataBuilder.h"
 #include "SystemTextures.h"
+#include "Engine/AssetManager.h"
 #include "Engine/StaticMesh.h"
+#include "Engine/StreamableManager.h"
 #include "Engine/World.h"
 #include "ComputeFramework/ShaderParamTypeDefinition.h"
 
@@ -407,11 +409,22 @@ bool UPCGStaticMeshSpawnerDataProvider::PrepareForExecute_GameThread(UPCGDataBin
 	{
 		CreatePrimitiveDescriptors(Context, InBinding);
 
+		TArray<FSoftObjectPath> PathsToLoad;
+		PathsToLoad.Reserve(PrimitiveDescriptors.Num());
+
 		for (const FPCGProceduralISMComponentDescriptor& PrimitiveDescriptor : PrimitiveDescriptors)
 		{
-			UE_LOG(LogPCG, Verbose, TEXT("Request '%s' to load."), *PrimitiveDescriptor.StaticMesh.ToString());
+			if (!PrimitiveDescriptor.StaticMesh.IsNull())
+			{
+				UE_LOG(LogPCG, Verbose, TEXT("Request '%s' to load."), *PrimitiveDescriptor.StaticMesh.ToString());
+				PathsToLoad.Add(PrimitiveDescriptor.StaticMesh.ToSoftObjectPath());
+			}
+		}
 
-			PrimitiveDescriptor.StaticMesh.LoadAsync({});
+		if (!PathsToLoad.IsEmpty())
+		{
+			ensure(!LoadHandle);
+			LoadHandle = UAssetManager::Get().GetStreamableManager().RequestAsyncLoad(MoveTemp(PathsToLoad));
 		}
 
 		bPrimitiveDescriptorsCreated = true;
@@ -500,6 +513,7 @@ void UPCGStaticMeshSpawnerDataProvider::Reset()
 	NumPrimitivesSetup = 0;
 	bRegisteredPrimitives = false;
 	bStaticMeshesLoaded = false;
+	LoadHandle.Reset();
 }
 
 void UPCGStaticMeshSpawnerDataProvider::CreatePrimitiveDescriptors(FPCGContext* InContext, UPCGDataBinding* InBinding)
@@ -776,7 +790,7 @@ bool UPCGStaticMeshSpawnerDataProvider::SetupPrimitives(FPCGContext* InContext,
 	// Ensure SMs are compiled, required to prevent ensure when partially built SM is accessed later in FastGeo code.
 	for (const FPCGProceduralISMComponentDescriptor& Desc : PrimitiveDescriptors)
 	{
-		if (Desc.StaticMesh->IsCompiling())
+		if (ensure(Desc.StaticMesh.IsValid()) && Desc.StaticMesh->IsCompiling())
 		{
 			return false;
 		}
diff --git a/Engine/Plugins/PCG/Source/PCG/Private/Elements/Metadata/PCGMetadataPartitionKernel.cpp b/Engine/Plugins/PCG/Source/PCG/Private/Elements/Metadata/PCGMetadataPartitionKernel.cpp
--- a/Engine/Plugins/PCG/Source/PCG/Private/Elements/Metadata/PCGMetadataPartitionKernel.cpp
+++ b/Engine/Plugins/PCG/Source/PCG/Private/Elements/Metadata/PCGMetadataPartitionKernel.cpp
@@ -108,8 +108,6 @@ TSharedPtr<const FPCGDataCollectionDesc> UPCGMetadataPartitionKernel::ComputeOut
 		return nullptr;
 	}
 
-	TSharedPtr<FPCGDataCollectionDesc> OutDataDesc = FPCGDataCollectionDesc::MakeSharedFrom(InputPinDesc);
-
 	const TArray<int32> AnalysisDataIndices = InBinding->GetPinInputDataIndices(this, PCGMetadataPartitionConstants::ElementCountsPinLabel);
 
 	if (AnalysisDataIndices.Num() != InputPinDesc->DataDescs.Num())
@@ -124,6 +122,8 @@ TSharedPtr<const FPCGDataCollectionDesc> UPCGMetadataPartitionKernel::ComputeOut
 		return nullptr;
 	}
 
+	TSharedPtr<FPCGDataCollectionDesc> OutDataDesc = FPCGDataCollectionDesc::MakeShared();
+
 	for (int32 DataIndex = 0; DataIndex < AnalysisDataIndices.Num(); ++DataIndex)
 	{
 		const int32 AnalysisDataIndex = AnalysisDataIndices[DataIndex];
diff --git a/Engine/Plugins/PCG/Source/PCG/Public/Compute/DataInterfaces/Elements/PCGStaticMeshSpawnerDataInterface.h b/Engine/Plugins/PCG/Source/PCG/Public/Compute/DataInterfaces/Elements/PCGStaticMeshSpawnerDataInterface.h
--- a/Engine/Plugins/PCG/Source/PCG/Public/Compute/DataInterfaces/Elements/PCGStaticMeshSpawnerDataInterface.h
+++ b/Engine/Plugins/PCG/Source/PCG/Public/Compute/DataInterfaces/Elements/PCGStaticMeshSpawnerDataInterface.h
@@ -15,6 +15,7 @@ class UPCGStaticMeshSpawnerSettings;
 class UPCGSettings;
 class UPCGStaticMeshSpawnerKernel;
 struct FPCGContext;
+struct FStreamableHandle;
 
 namespace PCGStaticMeshSpawnerDataInterface
 {
@@ -108,6 +109,9 @@ public:
 
 	/** Number of primitives created during primitive setup. */
 	int32 NumPrimitivesSetup = 0;
+
+	/** Keeps loaded objects alive. */
+	TSharedPtr<FStreamableHandle> LoadHandle;
 };
 
 class FPCGStaticMeshSpawnerDataProviderProxy : public FComputeDataProviderRenderProxy
