[PCG] Use Property handle flow to set values to follow the normal flow for custom details of FPCGDataTypeIdentifier

* Without that, BP instances are not updated

#rnx

diff --git a/Engine/Plugins/PCG/Source/PCGEditor/Private/Details/PCGDataTypeIdentifierDetails.cpp b/Engine/Plugins/PCG/Source/PCGEditor/Private/Details/PCGDataTypeIdentifierDetails.cpp
--- a/Engine/Plugins/PCG/Source/PCGEditor/Private/Details/PCGDataTypeIdentifierDetails.cpp
+++ b/Engine/Plugins/PCG/Source/PCGEditor/Private/Details/PCGDataTypeIdentifierDetails.cpp
@@ -204,21 +204,26 @@ TSharedRef<SWidget> FPCGDataTypeIdentifierDetails::OnGetMenuContent()
 			(
 				FExecuteAction::CreateLambda([this, ThisDataType = DataType]()
 				{
+					TArray<UObject*> OuterObjects;
+					PropertyHandle->GetOuterObjects(OuterObjects);
+					if (OuterObjects.IsEmpty())
+					{
+						return;
+					}
+
 					FPCGDataTypeIdentifier* Identifier = GetStruct();
 					if (!ensure(Identifier))
 					{
 						return;
 					}
 
-					FScopedTransaction Transaction(NSLOCTEXT("PCGDataTypeIdentifierDetails", "ModifyIdentifier", "Modify PCG Data Type Identifier."));
-
-					PropertyHandle->NotifyPreChange();
+					FPCGDataTypeIdentifier TempIdentifier{};
 
 					if (bSupportComposition && FSlateApplication::Get().GetModifierKeys().IsControlDown())
 					{
 						if (!PCGDataTypeIdentifierDetails::IsSuperSet(*Identifier, ThisDataType))
 						{
-							*Identifier = Identifier->Compose(ThisDataType);
+							TempIdentifier = Identifier->Compose(ThisDataType);
 						}
 						else
 						{
@@ -230,19 +235,26 @@ TSharedRef<SWidget> FPCGDataTypeIdentifierDetails::OnGetMenuContent()
 								.FilteredTypes = HiddenTypes
 							};
 							
-							FPCGDataTypeIdentifier TempIdentifier = FPCGModule::GetConstDataTypeRegistry().GetIdentifiersDifference(Params);
-							if (TempIdentifier.IsValid())
-							{
-								*Identifier = MoveTemp(TempIdentifier);
-							}
+							TempIdentifier = FPCGModule::GetConstDataTypeRegistry().GetIdentifiersDifference(Params);
 						}
 					}
 					else
 					{
-						*Identifier = ThisDataType;
+						TempIdentifier = ThisDataType;
 					}
 
-					PropertyHandle->NotifyPostChange(EPropertyChangeType::ValueSet);
+					if (TempIdentifier.IsValid())
+					{
+						// Need to export text and import text to go through the normal flow (like update Blueprint instances)
+						UObject* Outer = OuterObjects[0];
+						FString IdentifierAsText{};
+						FPCGDataTypeIdentifier::StaticStruct()->ExportText(IdentifierAsText, &TempIdentifier, /*Defaults=*/nullptr, Outer, PPF_None, /*ExportRootScope=*/nullptr);
+						if (!IdentifierAsText.IsEmpty())
+						{
+							FScopedTransaction Transaction(NSLOCTEXT("PCGDataTypeIdentifierDetails", "ModifyIdentifier", "Modify PCG Data Type Identifier."));
+							PropertyHandle->SetPerObjectValues({MoveTemp(IdentifierAsText)});
+						}
+					}
 				}),
 				FCanExecuteAction(),
 				FIsActionChecked::CreateLambda([this, ThisDataType = DataType]() -> bool
