[PCG] BugHawk fix

#rnx

diff --git a/Engine/Plugins/PCG/Source/PCG/Private/Compute/BuiltInKernels/PCGCountUniqueAttributeValuesKernel.cpp b/Engine/Plugins/PCG/Source/PCG/Private/Compute/BuiltInKernels/PCGCountUniqueAttributeValuesKernel.cpp
--- a/Engine/Plugins/PCG/Source/PCG/Private/Compute/BuiltInKernels/PCGCountUniqueAttributeValuesKernel.cpp
+++ b/Engine/Plugins/PCG/Source/PCG/Private/Compute/BuiltInKernels/PCGCountUniqueAttributeValuesKernel.cpp
@@ -172,7 +172,7 @@ TSharedPtr<const FPCGDataCollectionDesc> UPCGCountUniqueAttributeValuesKernel::C
 		else
 		{
 			// Null output: One entry: string key 0, value count 0.
-			FPCGDataDesc CountDataDesc = OutDataDesc->DataDescs.Emplace_GetRef(EPCGDataType::Param, 1);
+			FPCGDataDesc& CountDataDesc = OutDataDesc->DataDescs.Emplace_GetRef(EPCGDataType::Param, 1);
 			CountDataDesc.AddAttribute(FPCGKernelAttributeKey(PCGCountUniqueAttributeValuesConstants::ValueAttributeName, EPCGKernelAttributeType::Int), InBinding);
 			CountDataDesc.AddAttribute(FPCGKernelAttributeKey(PCGCountUniqueAttributeValuesConstants::ValueCountAttributeName, EPCGKernelAttributeType::Int), InBinding);
 		}
