[PCG] Range based Blueprint Point Processor element : PointLoop & IterationLoop

diff --git a/Engine/Plugins/PCG/Source/PCG/Private/Data/PCGBasePointData.cpp b/Engine/Plugins/PCG/Source/PCG/Private/Data/PCGBasePointData.cpp
--- a/Engine/Plugins/PCG/Source/PCG/Private/Data/PCGBasePointData.cpp
+++ b/Engine/Plugins/PCG/Source/PCG/Private/Data/PCGBasePointData.cpp
@@ -444,7 +444,7 @@ namespace PCGBasePointData
 
 		if (!PCGBasePointData::IsValidRange<PointProperty>(Range, ValueRange))
 		{
-			UE_LOG(LogPCG, Error, TEXT("GetInputRangeValues '%s' called with input range [%d, %d[ when valid range is [0, %d[."), *StaticEnum<EPCGPointNativeProperties>()->GetNameStringByValue((int64)PointProperty), Range.RangeStartIndex, Range.RangeSize, Range.PointData->GetNumPoints());
+			UE_LOG(LogPCG, Error, TEXT("GetInputRangeValues '%s' called with input range [%d, %d[ when valid range is [0, %d]."), *StaticEnum<EPCGPointNativeProperties>()->GetNameStringByValue((int64)PointProperty), Range.RangeStartIndex, Range.RangeSize, Range.PointData->GetNumPoints()-1);
 			return {};
 		}
 
@@ -484,7 +484,7 @@ namespace PCGBasePointData
 
 		if (!PCGBasePointData::IsValidRange<PointProperty>(Range, ValueRange))
 		{
-			UE_LOG(LogPCG, Error, TEXT("SetOutputRangeValues '%s' called with output range [%d, %d[ when valid range is [0, %d[."), *StaticEnum<EPCGPointNativeProperties>()->GetNameStringByValue((int64)PointProperty), Range.RangeStartIndex, Range.RangeSize, Range.PointData->GetNumPoints());
+			UE_LOG(LogPCG, Error, TEXT("SetOutputRangeValues '%s' called with output range [%d, %d[ when valid range is [0, %d]."), *StaticEnum<EPCGPointNativeProperties>()->GetNameStringByValue((int64)PointProperty), Range.RangeStartIndex, Range.RangeSize, Range.PointData->GetNumPoints()-1);
 			return;
 		}
 
@@ -1562,4 +1562,40 @@ void UPCGBasePointData::SetMetadataEntryValuesFromRange(UPARAM(ref) FPCGPointOut
 	PCGBasePointData::SetOutputRangeValues<EPCGPointNativeProperties::MetadataEntry>(OutputRange, Values);
 }
 
+FPCGPoint UPCGBasePointData::GetPointFromRange(const FPCGPointInputRange& InputRange, int32 Index)
+{
+	if (!InputRange.PointData)
+	{
+		UE_LOG(LogPCG, Error, TEXT("GetPointFromRange called with invalid point data."));
+		return FPCGPoint();
+	}
+
+	if (Index < 0 || Index >= InputRange.PointData->GetNumPoints())
+	{
+		UE_LOG(LogPCG, Error, TEXT("GetPointFromRange called with invalid index '%d' when valid range is [0, %d]."), Index, InputRange.PointData->GetNumPoints()-1);
+		return FPCGPoint();
+	}
+
+	const FConstPCGPointValueRanges PointRanges(InputRange.PointData);
+	return PointRanges.GetPoint(Index);
+}
+
+void UPCGBasePointData::SetPointFromRange(UPARAM(ref) FPCGPointOutputRange& OutputRange, int32 Index, const FPCGPoint& Point)
+{
+	if (!OutputRange.PointData)
+	{
+		UE_LOG(LogPCG, Error, TEXT("SetPointFromRange called with invalid point data."));
+		return;
+	}
+
+	if (Index < 0 || Index >= OutputRange.PointData->GetNumPoints())
+	{
+		UE_LOG(LogPCG, Error, TEXT("SetPointFromRange called with invalid index '%d' when valid range is [0, %d]."), Index, OutputRange.PointData->GetNumPoints()-1);
+		return;
+	}
+
+	FPCGPointValueRanges PointRanges(OutputRange.PointData, /*bAllocate=*/false);
+	PointRanges.SetFromPoint(Index, Point);
+}
+
 #undef LOCTEXT_NAMESPACE
diff --git a/Engine/Plugins/PCG/Source/PCG/Private/Elements/Blueprint/PCGBlueprintPointProcessorElement.cpp b/Engine/Plugins/PCG/Source/PCG/Private/Elements/Blueprint/PCGBlueprintPointProcessorElement.cpp
new file mode 100644
--- /dev/null
+++ b/Engine/Plugins/PCG/Source/PCG/Private/Elements/Blueprint/PCGBlueprintPointProcessorElement.cpp
@@ -0,0 +1,111 @@
+// Copyright Epic Games, Inc. All Rights Reserved.
+
+#include "Elements/Blueprint/PCGBlueprintPointProcessorElement.h"
+
+#include "Data/PCGBasePointData.h"
+#include "Helpers/PCGAsync.h"
+
+void UPCGBlueprintPointProcessorElement::PointLoopInitialize_Implementation(const UPCGBasePointData* InputPointData, UPCGBasePointData* OutputPointData)
+{
+	FPCGInitializeFromDataParams InitializeFromDataParams(InputPointData);
+	InitializeFromDataParams.bInheritMetadata = true;
+	InitializeFromDataParams.bInheritSpatialData = true;
+	OutputPointData->InitializeFromDataWithParams(InitializeFromDataParams);
+}
+
+void UPCGBlueprintPointProcessorElement::PointLoop(const UPCGBasePointData* InputPointData, UPCGBasePointData*& OutputPointData)
+{
+	if (!InputPointData)
+	{
+		return;
+	}
+
+	OutputPointData = FPCGContext::NewPointData_AnyThread(CurrentContext);
+	PointLoopInitialize(InputPointData, OutputPointData);
+
+	auto ProcessFunc = [this, InputPointData, OutputPointData](int32 StartReadIndex, int32 StartWriteIndex, int32 Count)
+	{
+		FPCGPointInputRange InputRange{ InputPointData, StartReadIndex, Count };
+		FPCGPointOutputRange OutputRange{ OutputPointData, StartWriteIndex, Count };
+
+		return PointLoopBody(InputRange, OutputRange);
+	};
+		
+	auto MoveDataRangeFunc = [OutputPointData](int32 RangeStartIndex, int32 MoveToIndex, int32 NumElements)
+	{
+		OutputPointData->MoveRange(RangeStartIndex, MoveToIndex, NumElements);
+	};
+
+	auto FinishedFunc = [OutputPointData](int32 NumWritten)
+	{
+		OutputPointData->SetNumPoints(NumWritten);
+	};
+
+	{
+		TGuardValue<bool> IsCallingBlueprint(CurrentContext->AsyncState.bIsCallingBlueprint, true);
+
+		FPCGAsync::AsyncProcessingRangeEx(
+			&CurrentContext->AsyncState,
+			InputPointData->GetNumPoints(),
+			/*InitializeFunc=*/[]() {},
+			ProcessFunc,
+			MoveDataRangeFunc,
+			FinishedFunc,
+			/*bEnableTimeSlicing=*/false
+		);
+	}
+}
+
+void UPCGBlueprintPointProcessorElement::IterationLoopInitialize_Implementation(int32 NumIterations, const UPCGBasePointData* InputPointDataA, const UPCGBasePointData* InputPointDataB, UPCGBasePointData* OutputPointData)
+{
+	const UPCGBasePointData* InheritFromData = InputPointDataA ? InputPointDataA : InputPointDataB;
+	if (InheritFromData)
+	{
+		FPCGInitializeFromDataParams InitializeFromDataParams(InputPointDataA);
+		InitializeFromDataParams.bInheritMetadata = true;
+		InitializeFromDataParams.bInheritSpatialData = false;
+		OutputPointData->InitializeFromDataWithParams(InitializeFromDataParams);
+	}
+	OutputPointData->SetNumPoints(NumIterations);
+}
+
+void UPCGBlueprintPointProcessorElement::IterationLoop(int32 NumIterations, const UPCGBasePointData* InputPointDataA, const UPCGBasePointData* InputPointDataB, UPCGBasePointData*& OutputPointData)
+{
+	OutputPointData = FPCGContext::NewPointData_AnyThread(CurrentContext);
+	IterationLoopInitialize(NumIterations, InputPointDataA, InputPointDataB, OutputPointData);
+
+	auto ProcessFunc = [this, InputPointDataA, InputPointDataB, OutputPointData](int32 StartReadIndex, int32 StartWriteIndex, int32 Count)
+	{
+		// Give full input ranges as we don't know how those will be accessed in the nested loop body
+		const FPCGPointInputRange InputRangeA{ InputPointDataA, 0, InputPointDataA->GetNumPoints() };
+		const FPCGPointInputRange InputRangeB{ InputPointDataB, 0, InputPointDataB->GetNumPoints() };
+		FPCGPointOutputRange OutputRange{ OutputPointData, StartWriteIndex, Count };
+
+		return IterationLoopBody(StartReadIndex, InputRangeA, InputRangeB, OutputRange);
+	};
+
+	auto MoveDataRangeFunc = [OutputPointData](int32 RangeStartIndex, int32 MoveToIndex, int32 NumElements)
+	{
+		OutputPointData->MoveRange(RangeStartIndex, MoveToIndex, NumElements);
+	};
+
+	auto FinishedFunc = [OutputPointData](int32 NumWritten)
+	{
+		OutputPointData->SetNumPoints(NumWritten);
+	};
+
+	{
+		TGuardValue<bool> IsCallingBlueprint(CurrentContext->AsyncState.bIsCallingBlueprint, true);
+
+		FPCGAsync::AsyncProcessingRangeEx(
+			&CurrentContext->AsyncState,
+			NumIterations,
+			/*InitializeFunc=*/[]() {},
+			ProcessFunc,
+			MoveDataRangeFunc,
+			FinishedFunc,
+			/*bEnableTimeSlicing=*/false
+		);
+	}
+}
+
diff --git a/Engine/Plugins/PCG/Source/PCG/Private/Elements/Blueprint/PCGBlueprintPointProcessorSimpleElement.cpp b/Engine/Plugins/PCG/Source/PCG/Private/Elements/Blueprint/PCGBlueprintPointProcessorSimpleElement.cpp
--- a/Engine/Plugins/PCG/Source/PCG/Private/Elements/Blueprint/PCGBlueprintPointProcessorSimpleElement.cpp
+++ b/Engine/Plugins/PCG/Source/PCG/Private/Elements/Blueprint/PCGBlueprintPointProcessorSimpleElement.cpp
@@ -49,9 +49,6 @@ void UPCGBlueprintPointProcessorSimpleElement::Execute_Implementation(const FPCG
 
 		TFunction<int32(int32, int32, int32)> RangeProcessFunc = [this, InputPointData, OutputPointData](int32 StartReadIndex, int32 StartWriteIndex, int32 Count)
 		{
-			const FConstPCGPointValueRanges InputValueRanges(InputPointData);
-			FPCGPointValueRanges OutputValueRanges(OutputPointData, /*bAllocate=*/false);
-
 			FPCGPointInputRange InputRange{ InputPointData, StartReadIndex, Count };
 			FPCGPointOutputRange OutputRange{ OutputPointData, StartWriteIndex, Count };
 
diff --git a/Engine/Plugins/PCG/Source/PCG/Public/Data/PCGBasePointData.h b/Engine/Plugins/PCG/Source/PCG/Public/Data/PCGBasePointData.h
--- a/Engine/Plugins/PCG/Source/PCG/Public/Data/PCGBasePointData.h
+++ b/Engine/Plugins/PCG/Source/PCG/Public/Data/PCGBasePointData.h
@@ -385,6 +385,12 @@ public:
 	UFUNCTION(BlueprintCallable, Category = SpatialData)
 	static void SetMetadataEntryValuesFromRange(UPARAM(ref) FPCGPointOutputRange& OutputRange, const TArray<int64>& Values);
 
+	UFUNCTION(BlueprintCallable, Category = SpatialData)
+	static FPCGPoint GetPointFromRange(const FPCGPointInputRange& InputRange, int32 Index);
+
+	UFUNCTION(BlueprintCallable, Category = SpatialData)
+	static void SetPointFromRange(UPARAM(ref) FPCGPointOutputRange& OutputRange, int32 Index, const FPCGPoint& Point);
+
 protected:
 
 	void RebuildOctreeIfNeeded() const
diff --git a/Engine/Plugins/PCG/Source/PCG/Public/Elements/Blueprint/PCGBlueprintDeprecatedElement.h b/Engine/Plugins/PCG/Source/PCG/Public/Elements/Blueprint/PCGBlueprintDeprecatedElement.h
--- a/Engine/Plugins/PCG/Source/PCG/Public/Elements/Blueprint/PCGBlueprintDeprecatedElement.h
+++ b/Engine/Plugins/PCG/Source/PCG/Public/Elements/Blueprint/PCGBlueprintDeprecatedElement.h
@@ -24,7 +24,11 @@ DECLARE_MULTICAST_DELEGATE_OneParam(FOnPCGBlueprintChanged, UPCGBlueprintElement
 
 #endif // WITH_EDITOR
 
-UCLASS(MinimalAPI, Abstract, BlueprintType, Blueprintable, hidecategories = (Object))
+/**
+ * Class is no longer supported and doesn't support PCGBasePointData. 
+ * Please sub-class UPCGBlueprintBaseElement instead or use one of its existing child classes (ex: PCGBlueprintPointProcessorElement, PCGBlueprintSimplePointProcessorElement).
+ */
+UCLASS(MinimalAPI, Abstract, BlueprintType, Blueprintable, hidecategories = (Object), meta = (DisplayName = "Deprecated Blueprint Element"))
 class UPCGBlueprintElement : public UPCGBlueprintBaseElement
 {
 	GENERATED_BODY()
diff --git a/Engine/Plugins/PCG/Source/PCG/Public/Elements/Blueprint/PCGBlueprintPointProcessorElement.h b/Engine/Plugins/PCG/Source/PCG/Public/Elements/Blueprint/PCGBlueprintPointProcessorElement.h
new file mode 100644
--- /dev/null
+++ b/Engine/Plugins/PCG/Source/PCG/Public/Elements/Blueprint/PCGBlueprintPointProcessorElement.h
@@ -0,0 +1,84 @@
+// Copyright Epic Games, Inc. All Rights Reserved.
+
+#pragma once
+
+#include "PCGBlueprintBaseElement.h"
+
+#include "PCGBlueprintPointProcessorElement.generated.h"
+
+class UPCGBasePointData;
+struct FPCGPointInputRange;
+struct FPCGPointOutputRange;
+
+/**
+ * Point processor element that supports the following two type of loops:
+ * - PointLoop: requires PointLoopBody to be overridden with optional override of InitalizePointLoop.
+ * - IterationLoop: requires IterationLoopBody to be overridden with optional override of InitializeIterationLoop.
+ */
+UCLASS(MinimalAPI, Abstract, BlueprintType, Blueprintable, hidecategories = (Object))
+class UPCGBlueprintPointProcessorElement : public UPCGBlueprintBaseElement
+{
+	GENERATED_BODY()
+
+public:
+	/**
+	 * Calls the PointLoopInitialize once and LoopBody a fixed number of times.
+	 * @param InputPointData	Input point data. Constant, must not be modified.
+	 * @param OutputPointData	PointLoop output point data. Can be modified.
+	 */
+	UFUNCTION(BlueprintCallable, BlueprintPure = false, Category = "PCG|Flow Control", meta = (HideSelfPin = "true"))
+	PCG_API void PointLoop(const UPCGBasePointData* InputPointData, UPCGBasePointData*& OutputPointData);
+
+	/**
+	 * Initialization of the PointLoop output point data. This is called once before before the PointLoopBody calls.
+	 * Native default implementation inherits spatial data & metadata from input.
+	 * 
+	 * @param InputPointData	Input point data. Constant, must not be modified.
+	 * @param OutputPointData	PointLoop output point data. Can be modified.
+	 */
+	UFUNCTION(BlueprintNativeEvent, Category = "PCG|Flow Control", meta = (ForceAsFunction))
+	void PointLoopInitialize(const UPCGBasePointData* InputPointData, UPCGBasePointData* OutputPointData);
+
+	/**
+	 * Multi-threaded range based loop that will be called in chunks a number of times.
+	 *
+	 * @param InputRange		Input range data created from InputPointData.
+	 * @param OutputRange		PointLoop output range data used to write to the output point data.
+	 */
+	UFUNCTION(BlueprintImplementableEvent, Category = "PCG|Flow Control")
+	int32 PointLoopBody(const FPCGPointInputRange& InputRange, UPARAM(Ref) FPCGPointOutputRange& OutputRange);
+
+	/** 
+	 * Calls the IterationLoopInitalize once and IterationLoopBody a fixed number of times. 
+	 * @param NumIteration		Number of iterations.
+	 * @param InputPointDataA	Optional input point data, can be null. Constant, must not be modified.
+	 * @param InputPointDataB	Optional input point data, can be null. Constant, must not be modified.
+	 * @param OutputPointData	IterationLoop output point data. Can be modified.
+	 */
+	UFUNCTION(BlueprintCallable, BlueprintPure = false, Category = "PCG|Flow Control", meta = (HideSelfPin = "true"))
+	PCG_API void IterationLoop(int32 NumIterations, const UPCGBasePointData* InputPointDataA, const UPCGBasePointData* InputPointDataB, UPCGBasePointData*& OutputPointData);
+
+	/** 
+	 * Initialization of the IterationLoop output point data. This is called once before before the IterationLoopBody calls. 
+	 * Native default implementation does not inherit spatial data and inherits metadata from its first non-null input if any.
+	 * 
+	 * @param NumIteration		Number of iterations.
+	 * @param InputPointDataA	Optional input point data, can be null. Constant, must not be modified.
+	 * @param InputPointDataB	Optional input point data, can be null. Constant, must not be modified.
+	 * @param OutputPointData	IterationLoop output point data to initialize. Can be modified.
+	 */
+	UFUNCTION(BlueprintNativeEvent, Category = "PCG|Flow Control", meta = (ForceAsFunction))
+	void IterationLoopInitialize(int32 NumIterations, const UPCGBasePointData* InputPointDataA, const UPCGBasePointData* InputPointDataB, UPCGBasePointData* OutputPointData);
+
+	/** 
+	 * Multi-threaded range based loop that will be called in chunks a number of times.
+	 * Use the OutputRange.RangeSize to get the number of iterations to process in the IterationLoopBody [IterationIndex, IterationIndex + OutputRange.RangeSize [
+	 * 
+	 * @param IterationIndex	Current IterationLoopBody index.
+	 * @param InputRangeA		Optional input range data if IterationLoop was called with a valid InputPointDataA.
+	 * @param InputRangeB		Optional input range data if IterationLoop was called with a valid InputPointDataB.
+	 * @param OutputRange		IterationLoop output range data used to write to the output point data.
+	 */
+	UFUNCTION(BlueprintImplementableEvent, Category = "PCG|Flow Control")
+	int32 IterationLoopBody(int32 IterationIndex, const FPCGPointInputRange& InputRangeA, const FPCGPointInputRange& InputRangeB, UPARAM(Ref) FPCGPointOutputRange& OutputRange);
+};
