[PCG] Bughawk fixes CL 45776584

#rnx

diff --git a/Engine/Plugins/PCG/Source/PCG/Private/Elements/Blueprint/PCGBlueprintPointProcessorElement.cpp b/Engine/Plugins/PCG/Source/PCG/Private/Elements/Blueprint/PCGBlueprintPointProcessorElement.cpp
--- a/Engine/Plugins/PCG/Source/PCG/Private/Elements/Blueprint/PCGBlueprintPointProcessorElement.cpp
+++ b/Engine/Plugins/PCG/Source/PCG/Private/Elements/Blueprint/PCGBlueprintPointProcessorElement.cpp
@@ -61,7 +61,7 @@ void UPCGBlueprintPointProcessorElement::IterationLoopInitialize_Implementation(
 	const UPCGBasePointData* InheritFromData = InputPointDataA ? InputPointDataA : InputPointDataB;
 	if (InheritFromData)
 	{
-		FPCGInitializeFromDataParams InitializeFromDataParams(InputPointDataA);
+		FPCGInitializeFromDataParams InitializeFromDataParams(InheritFromData);
 		InitializeFromDataParams.bInheritMetadata = true;
 		InitializeFromDataParams.bInheritSpatialData = false;
 		OutputPointData->InitializeFromDataWithParams(InitializeFromDataParams);
@@ -77,8 +77,8 @@ void UPCGBlueprintPointProcessorElement::IterationLoop(int32 NumIterations, cons
 	auto ProcessFunc = [this, InputPointDataA, InputPointDataB, OutputPointData](int32 StartReadIndex, int32 StartWriteIndex, int32 Count)
 	{
 		// Give full input ranges as we don't know how those will be accessed in the nested loop body
-		const FPCGPointInputRange InputRangeA{ InputPointDataA, 0, InputPointDataA->GetNumPoints() };
-		const FPCGPointInputRange InputRangeB{ InputPointDataB, 0, InputPointDataB->GetNumPoints() };
+		const FPCGPointInputRange InputRangeA{ InputPointDataA, 0, InputPointDataA ? InputPointDataA->GetNumPoints() : 0 };
+		const FPCGPointInputRange InputRangeB{ InputPointDataB, 0, InputPointDataB ? InputPointDataB->GetNumPoints() : 0 };
 		FPCGPointOutputRange OutputRange{ OutputPointData, StartWriteIndex, Count };
 
 		return IterationLoopBody(StartReadIndex, InputRangeA, InputRangeB, OutputRange);
diff --git a/Engine/Plugins/PCG/Source/PCG/Public/Elements/Blueprint/PCGBlueprintPointProcessorElement.h b/Engine/Plugins/PCG/Source/PCG/Public/Elements/Blueprint/PCGBlueprintPointProcessorElement.h
--- a/Engine/Plugins/PCG/Source/PCG/Public/Elements/Blueprint/PCGBlueprintPointProcessorElement.h
+++ b/Engine/Plugins/PCG/Source/PCG/Public/Elements/Blueprint/PCGBlueprintPointProcessorElement.h
@@ -11,7 +11,7 @@ struct FPCGPointInputRange;
 struct FPCGPointOutputRange;
 
 /**
- * Point processor element that supports the following two type of loops:
+ * Point processor element that supports the following two types of loop:
  * - PointLoop: requires PointLoopBody to be overridden with optional override of InitalizePointLoop.
  * - IterationLoop: requires IterationLoopBody to be overridden with optional override of InitializeIterationLoop.
  */
@@ -22,7 +22,7 @@ class UPCGBlueprintPointProcessorElement : public UPCGBlueprintBaseElement
 
 public:
 	/**
-	 * Calls the PointLoopInitialize once and LoopBody a fixed number of times.
+	 * Calls the PointLoopInitialize once and PointLoopBody a fixed number of times.
 	 * @param InputPointData	Input point data. Constant, must not be modified.
 	 * @param OutputPointData	PointLoop output point data. Can be modified.
 	 */
@@ -30,7 +30,7 @@ public:
 	PCG_API void PointLoop(const UPCGBasePointData* InputPointData, UPCGBasePointData*& OutputPointData);
 
 	/**
-	 * Initialization of the PointLoop output point data. This is called once before before the PointLoopBody calls.
+	 * Initialization of the PointLoop output point data. This is called once before the PointLoopBody calls.
 	 * Native default implementation inherits spatial data & metadata from input.
 	 * 
 	 * @param InputPointData	Input point data. Constant, must not be modified.
@@ -49,7 +49,7 @@ public:
 	int32 PointLoopBody(const FPCGPointInputRange& InputRange, UPARAM(Ref) FPCGPointOutputRange& OutputRange);
 
 	/** 
-	 * Calls the IterationLoopInitalize once and IterationLoopBody a fixed number of times. 
+	 * Calls the IterationLoopInitialize once and IterationLoopBody a fixed number of times. 
 	 * @param NumIteration		Number of iterations.
 	 * @param InputPointDataA	Optional input point data, can be null. Constant, must not be modified.
 	 * @param InputPointDataB	Optional input point data, can be null. Constant, must not be modified.
@@ -59,7 +59,7 @@ public:
 	PCG_API void IterationLoop(int32 NumIterations, const UPCGBasePointData* InputPointDataA, const UPCGBasePointData* InputPointDataB, UPCGBasePointData*& OutputPointData);
 
 	/** 
-	 * Initialization of the IterationLoop output point data. This is called once before before the IterationLoopBody calls. 
+	 * Initialization of the IterationLoop output point data. This is called once before the IterationLoopBody calls. 
 	 * Native default implementation does not inherit spatial data and inherits metadata from its first non-null input if any.
 	 * 
 	 * @param NumIteration		Number of iterations.
