[PCG] Fix use-after-free for the data type registry

* When the registry is deleted, the mapping of InstancedStruct can't be deleted correctly as the struct metaclass is already dead

#rnx

diff --git a/Engine/Plugins/PCG/Source/PCG/Private/Data/Registry/PCGDataTypeRegistry.cpp b/Engine/Plugins/PCG/Source/PCG/Private/Data/Registry/PCGDataTypeRegistry.cpp
--- a/Engine/Plugins/PCG/Source/PCG/Private/Data/Registry/PCGDataTypeRegistry.cpp
+++ b/Engine/Plugins/PCG/Source/PCG/Private/Data/Registry/PCGDataTypeRegistry.cpp
@@ -338,6 +338,19 @@ void FPCGDataTypeRegistry::RegisterEntry(const FPCGDataTypeBaseId& BaseID)
 	HierarchyTree->AddEntry(BaseID);
 }
 
+void FPCGDataTypeRegistry::Shutdown()
+{
+	Mapping.Empty();
+	HierarchyTree.Reset();
+
+#if WITH_EDITOR
+	CombinationPinColors.Empty();
+	SingleTypePinColors.Empty();
+	CombinationPinIcons.Empty();
+	SingleTypePinIcons.Empty();
+#endif // WITH_EDITOR
+}
+
 FPCGDataTypeRegistry::~FPCGDataTypeRegistry() = default;
 
 EPCGDataTypeCompatibilityResult FPCGDataTypeRegistry::IsCompatible(const FPCGDataTypeIdentifier& InType, const FPCGDataTypeIdentifier& OutType, FText* OptionalOutCompatibilityMessage) const
diff --git a/Engine/Plugins/PCG/Source/PCG/Private/PCGModule.cpp b/Engine/Plugins/PCG/Source/PCG/Private/PCGModule.cpp
--- a/Engine/Plugins/PCG/Source/PCG/Private/PCGModule.cpp
+++ b/Engine/Plugins/PCG/Source/PCG/Private/PCGModule.cpp
@@ -129,6 +129,8 @@ void FPCGModule::PreExit()
 	AttributeAccessorFactory.UnregisterMethods<UPCGBasePointData>();
 	AttributeAccessorFactory.UnregisterDefaultMethods();
 
+	DataTypeRegistry.Shutdown();
+
 #if WITH_EDITOR
 	DeregisterNativeElementDeterminismTests();
 
diff --git a/Engine/Plugins/PCG/Source/PCG/Public/Data/Registry/PCGDataTypeRegistry.h b/Engine/Plugins/PCG/Source/PCG/Public/Data/Registry/PCGDataTypeRegistry.h
--- a/Engine/Plugins/PCG/Source/PCG/Public/Data/Registry/PCGDataTypeRegistry.h
+++ b/Engine/Plugins/PCG/Source/PCG/Public/Data/Registry/PCGDataTypeRegistry.h
@@ -105,6 +105,9 @@ private:
 	void RegisterEntry(const FPCGDataTypeBaseId& BaseID);
 	
 	bool IsRegistered(const FPCGDataTypeBaseId& BaseID) const { return Mapping.Contains(BaseID); }
+
+	// To be called by the FPCGModule on PreExit as we have to make sure the structs FPCGDataTypeInfo are still alive.
+	void Shutdown();
 	
 	TMap<FPCGDataTypeBaseId, TInstancedStruct<FPCGDataTypeInfo>> Mapping;
 	TUniquePtr<PCGDataTypeHierarchy::FTree> HierarchyTree;
