[PCG] Be more lenient in RGS with null/lost original compnent to fix crashes hit during BP recompile/save while generation is in progress.

#rnx

diff --git a/Engine/Plugins/PCG/Source/PCG/Private/RuntimeGen/PCGRuntimeGenScheduler.cpp b/Engine/Plugins/PCG/Source/PCG/Private/RuntimeGen/PCGRuntimeGenScheduler.cpp
--- a/Engine/Plugins/PCG/Source/PCG/Private/RuntimeGen/PCGRuntimeGenScheduler.cpp
+++ b/Engine/Plugins/PCG/Source/PCG/Private/RuntimeGen/PCGRuntimeGenScheduler.cpp
@@ -994,12 +994,16 @@ void FPCGRuntimeGenScheduler::TickScheduleGeneration(TMap<FGridGenerationKey, do
 
 		const FIntVector GridCoords = Key.GetGridCoords();
 		UPCGComponent* OriginalComponent = Key.GetOriginalComponent();
-		check(OriginalComponent);
+		if (!OriginalComponent)
+		{
+			UE_LOG(LogPCG, Verbose, TEXT("TickScheduleGeneration: Invalid original component in an InOutComponentsToGenerate entry, generation will be skipped."));
+			continue;
+		}
 
 		// If the Grid is unbounded, we have a non-partitioned or unbounded component.
 		if (Grid == EPCGHiGenGrid::Unbounded)
 		{
-			if (ensure(OriginalComponent) && !OriginalComponent->IsGenerating())
+			if (!OriginalComponent->IsGenerating())
 			{
 				if (PCGRuntimeGenSchedulerHelpers::CVarRuntimeGenerationEnableDebugging.GetValueOnAnyThread() && OriginalComponent->GetOwner())
 				{
