[PCG] Per Execution Cache: expand to support bounds cache

diff --git a/Engine/Plugins/PCG/Source/PCG/Private/Compute/DataInterfaces/PCGTextureDataInterface.cpp b/Engine/Plugins/PCG/Source/PCG/Private/Compute/DataInterfaces/PCGTextureDataInterface.cpp
--- a/Engine/Plugins/PCG/Source/PCG/Private/Compute/DataInterfaces/PCGTextureDataInterface.cpp
+++ b/Engine/Plugins/PCG/Source/PCG/Private/Compute/DataInterfaces/PCGTextureDataInterface.cpp
@@ -344,16 +344,9 @@ void UPCGTextureDataProvider::BuildInfosFromDataDescription(UPCGDataBinding* InB
 		return;
 	}
 
-	// @todo_pcg: Instead of using this helper and relying on an actor (which may not exist for some cases), we should use the context execution source/state.
-	AActor* OriginalActor = UPCGBlueprintHelpers::GetOriginalComponent(*Context)->GetOwner();
-	if (!ensure(OriginalActor))
-	{
-		return;
-	}
-
-	const FBox OriginalActorLocalBounds = PCGHelpers::GetActorLocalBounds(OriginalActor);
-	FTransform Transform = OriginalActor->GetTransform();
-	Transform.SetScale3D(Transform.GetScale3D() * 0.5 * (OriginalActorLocalBounds.Max - OriginalActorLocalBounds.Min));
+	const FBox OriginalLocalBounds = Context->ExecutionSource->GetExecutionState().GetOriginalLocalSpaceBounds();
+	FTransform Transform = Context->ExecutionSource->GetExecutionState().GetOriginalTransform();
+	Transform.SetScale3D(Transform.GetScale3D() * 0.5 * (OriginalLocalBounds.Max - OriginalLocalBounds.Min));
 
 	for (int DataIndex = 0; DataIndex < PinDesc->DataDescs.Num(); ++DataIndex)
 	{
diff --git a/Engine/Plugins/PCG/Source/PCG/Private/Elements/PCGSpatialNoise.cpp b/Engine/Plugins/PCG/Source/PCG/Private/Elements/PCGSpatialNoise.cpp
--- a/Engine/Plugins/PCG/Source/PCG/Private/Elements/PCGSpatialNoise.cpp
+++ b/Engine/Plugins/PCG/Source/PCG/Private/Elements/PCGSpatialNoise.cpp
@@ -303,20 +303,20 @@ namespace PCGSpatialNoise
 		return 1.0 / (1.0 + FMath::Pow(Value / (1.0 - Value), -Contrast));
 	}
 
-	FLocalCoordinates2D CalcLocalCoordinates2D(const FBox& ActorLocalBox, const FTransform& ActorTransformInverse, FVector2D Scale, const FVector& InPosition)
+	FLocalCoordinates2D CalcLocalCoordinates2D(const FBox& LocalBox, const FTransform& TransformInverse, FVector2D Scale, const FVector& InPosition)
 	{
-		if (!ActorLocalBox.IsValid)
+		if (!LocalBox.IsValid)
 		{
 			return {};
 		}
 
-		const FVector2D LocalPosition = FVector2D(ActorTransformInverse.TransformPosition(InPosition));
+		const FVector2D LocalPosition = FVector2D(TransformInverse.TransformPosition(InPosition));
 
-		const double LeftDist = LocalPosition.X-ActorLocalBox.Min.X;
-		const double RightDist = LocalPosition.X-ActorLocalBox.Max.X;
+		const double LeftDist = LocalPosition.X-LocalBox.Min.X;
+		const double RightDist = LocalPosition.X-LocalBox.Max.X;
 
-		const double TopDist = LocalPosition.Y-ActorLocalBox.Min.Y;
-		const double BottomDist = LocalPosition.Y-ActorLocalBox.Max.Y;
+		const double TopDist = LocalPosition.Y-LocalBox.Min.Y;
+		const double BottomDist = LocalPosition.Y-LocalBox.Max.Y;
 
 		FLocalCoordinates2D Result;
 
@@ -325,8 +325,8 @@ namespace PCGSpatialNoise
 		Result.Y0 = TopDist * Scale.Y;
 		Result.Y1 = BottomDist * Scale.Y;
 
-		Result.FracX = FMath::Clamp(LeftDist / (ActorLocalBox.Max.X - ActorLocalBox.Min.X), 0.0, 1.0);
-		Result.FracY = FMath::Clamp(TopDist / (ActorLocalBox.Max.Y - ActorLocalBox.Min.Y), 0.0, 1.0);
+		Result.FracX = FMath::Clamp(LeftDist / (LocalBox.Max.X - LocalBox.Min.X), 0.0, 1.0);
+		Result.FracY = FMath::Clamp(TopDist / (LocalBox.Max.Y - LocalBox.Min.Y), 0.0, 1.0);
 
 		return Result;
 	};
@@ -349,8 +349,8 @@ namespace PCGSpatialNoise
 	struct FSharedParams
 	{
 		FPCGContext* Context = nullptr;
-		FTransform ActorTransformInverse;
-		FBox ActorLocalBox;
+		FTransform TransformInverse;
+		FBox LocalBox;
 		FTransform Transform;
 
 		double Brightness;
@@ -435,8 +435,8 @@ namespace PCGSpatialNoise
 					if (SharedParams.bTiling)
 					{
 						const FLocalCoordinates2D LocalCoords = CalcLocalCoordinates2D(
-							SharedParams.ActorLocalBox,
-							SharedParams.ActorTransformInverse,
+							SharedParams.LocalBox,
+							SharedParams.TransformInverse,
 							FVector2D(SharedParams.Transform.GetScale3D()),
 							PointPos
 						);
@@ -506,8 +506,8 @@ namespace PCGSpatialNoise
 					const FVector InPosition = InTransformRange[ReadIndex].GetTranslation();
 
 					const FLocalCoordinates2D LocalCoords = CalcLocalCoordinates2D(
-						SharedParams.ActorLocalBox,
-						SharedParams.ActorTransformInverse,
+						SharedParams.LocalBox,
+						SharedParams.TransformInverse,
 						FVector2D(SharedParams.Transform.GetScale3D()),
 						InPosition
 					);
@@ -631,12 +631,12 @@ namespace PCGSpatialNoise
 			const FInt32Point InteriorOffset = FloorInt(FVector2D(SharedParams.Transform.GetTranslation()));
 			const int32 TiledVoronoiResolution = FMath::Max(1, Settings.TiledVoronoiResolution);
 			const int32 MaxBlendingTileIndex = TiledVoronoiResolution - Settings.TiledVoronoiEdgeBlendCellCount;
-			const FVector LocalBoxSize = SharedParams.ActorLocalBox.GetSize();
+			const FVector LocalBoxSize = SharedParams.LocalBox.GetSize();
 
 			// copy this so the new transform converts to the cell index
 			FSharedParams UseSharedParams = SharedParams;
 
-			UseSharedParams.Transform = SharedParams.ActorTransformInverse * FTransform(
+			UseSharedParams.Transform = SharedParams.TransformInverse * FTransform(
 				FRotator::ZeroRotator,
 				FVector(double(TiledVoronoiResolution) * 0.5 + 0.5, double(TiledVoronoiResolution) * 0.5 + 0.5, 1.0), // offset to the center
 				FVector(double(TiledVoronoiResolution)/LocalBoxSize.X, double(TiledVoronoiResolution)/LocalBoxSize.Y, 1.0) // scale to the range of 0 .. TiledVoronoiResolution
@@ -689,27 +689,22 @@ bool FPCGSpatialNoise::ExecuteInternal(FPCGContext* Context) const
 	PCGSpatialNoise::FSharedParams SharedParams;
 	SharedParams.Context = Context;
 
-	UPCGComponent* SourceComponent = Cast<UPCGComponent>(Context->ExecutionSource.Get());
-	if (!SourceComponent)
+	IPCGGraphExecutionSource* ExecutionSource = Cast<IPCGGraphExecutionSource>(Context->ExecutionSource.Get());
+	if (!ExecutionSource)
 	{
-		PCGE_LOG(Error, GraphAndLog, LOCTEXT("NoSourceComponent", "No Source Component."));
-		return true;		
+		PCGE_LOG(Error, GraphAndLog, LOCTEXT("NoExecutionSource", "No Execution Source."));
+		return true;
 	}
 
 	{
-		AActor* Actor = SourceComponent->GetOwner();
-		if (!Actor)
-		{
-			PCGE_LOG(Error, GraphAndLog, LOCTEXT("NoActor", "Source Component has no actor"));
-			return true;		
-		}
+		SharedParams.LocalBox = ExecutionSource->GetExecutionState().GetLocalSpaceBounds();
 
-		SharedParams.ActorLocalBox = PCGHelpers::GetActorLocalBounds(Actor);
-		SharedParams.ActorLocalBox.Min *= Actor->GetTransform().GetScale3D();
-		SharedParams.ActorLocalBox.Max *= Actor->GetTransform().GetScale3D();
+		const FTransform Transform = ExecutionSource->GetExecutionState().GetTransform();
+		SharedParams.LocalBox.Min *= Transform.GetScale3D();
+		SharedParams.LocalBox.Max *= Transform.GetScale3D();
 
-		const FTransform ActorTransform = FTransform(Actor->GetTransform().Rotator(), Actor->GetTransform().GetTranslation(), FVector::One());
-		SharedParams.ActorTransformInverse = ActorTransform.Inverse();
+		const FTransform ActorTransform = FTransform(Transform.Rotator(), Transform.GetTranslation(), FVector::One());
+		SharedParams.TransformInverse = ActorTransform.Inverse();
 	}
 
 	SharedParams.Transform = Settings->Transform * FTransform(FRotator::ZeroRotator, RandomOffset, FVector(PCGSpatialNoise::MAGIC_SCALE_FACTOR));
diff --git a/Engine/Plugins/PCG/Source/PCG/Private/Elements/PCGTextureSampler.cpp b/Engine/Plugins/PCG/Source/PCG/Private/Elements/PCGTextureSampler.cpp
--- a/Engine/Plugins/PCG/Source/PCG/Private/Elements/PCGTextureSampler.cpp
+++ b/Engine/Plugins/PCG/Source/PCG/Private/Elements/PCGTextureSampler.cpp
@@ -207,12 +207,10 @@ bool FPCGTextureSamplerElement::ExecuteInternal(FPCGContext* InContext) const
 		}
 		else
 		{
-			AActor* OriginalActor = UPCGBlueprintHelpers::GetOriginalComponent(*Context)->GetOwner();
+			const FTransform OriginalTransform = Context->ExecutionSource->GetExecutionState().GetOriginalTransform();
+			Context->Transform = Settings->Transform * OriginalTransform;
 
-			FTransform OriginalActorTransform = OriginalActor->GetTransform();
-			Context->Transform = Settings->Transform * OriginalActorTransform;
-
-			FBox OriginalActorLocalBounds = PCGHelpers::GetActorLocalBounds(OriginalActor);
+			const FBox OriginalActorLocalBounds = Context->ExecutionSource->GetExecutionState().GetOriginalLocalSpaceBounds();
 			Context->Transform.SetScale3D(Context->Transform.GetScale3D() * 0.5 * (OriginalActorLocalBounds.Max - OriginalActorLocalBounds.Min));
 		}
 
diff --git a/Engine/Plugins/PCG/Source/PCG/Private/Elements/PCGWorldQuery.cpp b/Engine/Plugins/PCG/Source/PCG/Private/Elements/PCGWorldQuery.cpp
--- a/Engine/Plugins/PCG/Source/PCG/Private/Elements/PCGWorldQuery.cpp
+++ b/Engine/Plugins/PCG/Source/PCG/Private/Elements/PCGWorldQuery.cpp
@@ -183,10 +183,8 @@ bool FPCGWorldRayHitQueryElement::ExecuteInternal(FPCGContext* Context) const
 		return true;
 	}
 
-	UPCGComponent* OriginalComponent = SourceComponent->GetOriginalComponent();
-	AActor* Owner = OriginalComponent ? OriginalComponent->GetOwner() : nullptr;
-	FTransform Transform = Owner ? Owner->GetTransform() : FTransform::Identity;
-	FBox LocalBounds = Owner ? PCGHelpers::GetActorLocalBounds(Owner) : FBox(EForceInit::ForceInit);
+	const FTransform Transform = Context->ExecutionSource->GetExecutionState().GetOriginalTransform();
+	const FBox LocalBounds = Context->ExecutionSource->GetExecutionState().GetOriginalLocalSpaceBounds();
 
 	if (!QueryParams.bOverrideDefaultParams)
 	{
diff --git a/Engine/Plugins/PCG/Source/PCG/Private/Graph/PCGGraphExecutor.cpp b/Engine/Plugins/PCG/Source/PCG/Private/Graph/PCGGraphExecutor.cpp
--- a/Engine/Plugins/PCG/Source/PCG/Private/Graph/PCGGraphExecutor.cpp
+++ b/Engine/Plugins/PCG/Source/PCG/Private/Graph/PCGGraphExecutor.cpp
@@ -217,8 +217,6 @@ FPCGGraphExecutor::FPCGGraphExecutor()
 FPCGGraphExecutor::FPCGGraphExecutor(UWorld* InWorld)
 	: World(InWorld)
 {
-	static_assert((int8)(EExecutionCacheDataType::Count) == ExecutionCacheSize);
-
 	GameThreadHandler = MakeShared<FGameThreadHandler>(this);
 }
 
@@ -705,31 +703,6 @@ IPCGBaseSubsystem* FPCGGraphExecutor::GetSubsystem() const
 	}
 }
 
-UPCGData* FPCGGraphExecutor::GetExecutionCacheData(FPCGTaskId InGraphExecutionTaskId, EExecutionCacheDataType InExecutionCacheDataType)
-{
-	if (InGraphExecutionTaskId != InvalidPCGTaskId)
-	{
-		PCGGraphExecutor::TScopeLock ScopeLock(GraphExecutionCachesLock);
-		if (FGraphExecutionCache* GraphExecutionCache = GraphExecutionCaches.Find(InGraphExecutionTaskId))
-		{
-			return GraphExecutionCache->Data[(int8)InExecutionCacheDataType];
-		}
-	}
-
-	return nullptr;
-}
-
-void FPCGGraphExecutor::SetExecutionCacheData(FPCGTaskId InGraphExecutionTaskId, EExecutionCacheDataType InExecutionCacheDataType, UPCGData* InData)
-{
-	if (InGraphExecutionTaskId != InvalidPCGTaskId)
-	{
-		PCGGraphExecutor::TScopeLock ScopeLock(GraphExecutionCachesLock);
-		FGraphExecutionCache& GraphExecutionCache = GraphExecutionCaches.FindOrAdd(InGraphExecutionTaskId);
-		check(!GraphExecutionCache.Data[(int8)InExecutionCacheDataType]);
-		GraphExecutionCache.Data[(int8)InExecutionCacheDataType] = InData;
-	}
-}
-
 FPCGTaskId FPCGGraphExecutor::ScheduleGeneric(const FPCGScheduleGenericParams& Params)
 {
 	return ScheduleSingleElement(MakeShared<FPCGGenericElement>(Params.Operation, Params.AbortOperation, Params.bSupportBasePointDataInput, Params.bCanExecuteOnlyOnMainThread),
@@ -2448,8 +2421,7 @@ void FPCGGraphExecutor::ClearResults()
 		{
 			NextTaskId = 0;
 
-			PCGGraphExecutor::TScopeLock ChildScopeLock(GraphExecutionCachesLock);
-			GraphExecutionCaches.Empty();
+			PerExecutionCache.Clear();
 		}
 	}
 
@@ -2709,13 +2681,7 @@ void FPCGGraphExecutor::AddReferencedObjects(FReferenceCollector& Collector)
 		});
 	}
 
-	{
-		PCGGraphExecutor::TScopeLock ScopeLock(GraphExecutionCachesLock);
-		for (TMap<FPCGTaskId, FGraphExecutionCache>::TIterator It = GraphExecutionCaches.CreateIterator(); It; ++It)
-		{
-			It.Value().AddStructReferencedObjects(Collector);
-		}
-	}
+	PerExecutionCache.AddStructReferencedObjects(Collector);
 
 	{
 		PCGGraphExecutor::TScopeLock ScopeLock(ReleasedContextLock);
@@ -2731,22 +2697,6 @@ void FPCGGraphExecutor::AddReferencedObjects(FReferenceCollector& Collector)
 	GraphCompiler.AddReferencedObjects(Collector);
 }
 
-FPCGGraphExecutor::FGraphExecutionCache::FGraphExecutionCache()
-{
-	for (int32 i = 0; i < (int32)EExecutionCacheDataType::Count; ++i)
-	{
-		Data[i] = nullptr;
-	}
-}
-
-void FPCGGraphExecutor::FGraphExecutionCache::AddStructReferencedObjects(FReferenceCollector& Collector)
-{
-	for (int32 i = 0; i < (int32)EExecutionCacheDataType::Count; ++i)
-	{
-		Collector.AddReferencedObject(Data[i]);
-	}
-}
-
 FPCGElementPtr FPCGGraphExecutor::GetFetchInputElement()
 {
 	if (!FetchInputElementPtr)
@@ -2938,29 +2888,18 @@ bool FPCGPreGraphElement::ExecuteInternal(FPCGContext* Context) const
 
 	check(IsInGameThread());
 	const bool bIsValid = Context->ExecutionSource.IsValid();
-	UPCGComponent* Component = Cast<UPCGComponent>(Context->ExecutionSource.Get());
+	IPCGGraphExecutionSource* ExecutionSource = Context->ExecutionSource.Get();
 
-	// Early out if the component has been deleted/is invalid
-	if (!Component)
+	// Early out if the execution source has been deleted/is invalid
+	if (!ExecutionSource)
 	{
-		// If the component should exist but it doesn't (which is all the time here, previously we checked for it), then this should be cancelled
+		// If the execution source should exist but it doesn't (which is all the time here, previously we checked for it), then this should be cancelled
 		Context->OutputData.bCancelExecution = !bIsValid;
 		return true;
 	}
 
-#if WITH_EDITOR
-	Component->StartGenerationInProgress();
-#endif
+	ExecutionSource->GetExecutionState().ExecutePreGraph(Context);
 
-	ensureMsgf(Component->GetGenerationTaskId() != InvalidPCGTaskId, TEXT("Component was Scheduled for generation without having its CurrentGenerationTask assigned"));
-
-	{
-		PCG_EXECUTION_CACHE_VALIDATION_CREATE_SCOPE(Component);
-		
-		// Call getters which will create the data and cache it
-		Component->GetActorPCGData();
-		Component->GetOriginalActorPCGData();
-	}
 	return true;
 }
 
diff --git a/Engine/Plugins/PCG/Source/PCG/Private/Graph/PCGGraphExecutor.h b/Engine/Plugins/PCG/Source/PCG/Private/Graph/PCGGraphExecutor.h
--- a/Engine/Plugins/PCG/Source/PCG/Private/Graph/PCGGraphExecutor.h
+++ b/Engine/Plugins/PCG/Source/PCG/Private/Graph/PCGGraphExecutor.h
@@ -4,6 +4,7 @@
 
 #include "PCGContext.h"
 #include "PCGElement.h"
+#include "PCGGraphPerExecutionCache.h"
 #include "PCGPin.h"
 #include "PCGSettings.h"
 #include "PCGSubsystem.h"
@@ -161,36 +162,30 @@ private:
 
 	IPCGBaseSubsystem* GetSubsystem() const;
 
-	enum class EExecutionCacheDataType : int8
-	{
-		PCGData = 0,
-		InputData,
-		ActorData,
-		LandscapeData,
-		LandscapeHeightData,
-		OriginalActorData,
-		Count
-	};
-
-	static constexpr int8 ExecutionCacheSize = 6;
+	FPCGPerExecutionCache PerExecutionCache;
 
-	UPCGData* GetExecutionCacheData(FPCGTaskId InGraphExecutionTaskId, EExecutionCacheDataType InExecutionCacheDataType);
+	UPCGData* GetPCGData(FPCGTaskId InGraphExecutionTaskId, bool& bOutFound) { return PerExecutionCache.GetExecutionCachePCGData(InGraphExecutionTaskId, EPCGPerExecutionCacheDataType::PCGData, bOutFound); }
+	UPCGData* GetInputPCGData(FPCGTaskId InGraphExecutionTaskId, bool& bOutFound) { return PerExecutionCache.GetExecutionCachePCGData(InGraphExecutionTaskId, EPCGPerExecutionCacheDataType::InputData, bOutFound); }
+	UPCGData* GetActorPCGData(FPCGTaskId InGraphExecutionTaskId, bool& bOutFound) { return PerExecutionCache.GetExecutionCachePCGData(InGraphExecutionTaskId, EPCGPerExecutionCacheDataType::ActorData, bOutFound); }
+	UPCGData* GetLandscapePCGData(FPCGTaskId InGraphExecutionTaskId, bool& bOutFound) { return PerExecutionCache.GetExecutionCachePCGData(InGraphExecutionTaskId, EPCGPerExecutionCacheDataType::LandscapeData, bOutFound); }
+	UPCGData* GetLandscapeHeightPCGData(FPCGTaskId InGraphExecutionTaskId, bool& bOutFound) { return PerExecutionCache.GetExecutionCachePCGData(InGraphExecutionTaskId, EPCGPerExecutionCacheDataType::LandscapeHeightData, bOutFound); }
+	UPCGData* GetOriginalActorPCGData(FPCGTaskId InGraphExecutionTaskId, bool& bOutFound) { return PerExecutionCache.GetExecutionCachePCGData(InGraphExecutionTaskId, EPCGPerExecutionCacheDataType::OriginalActorData, bOutFound); }
+	TOptional<FBox> GetBounds(FPCGTaskId InGraphExecutionTaskId) { return PerExecutionCache.GetExecutionCacheBounds(InGraphExecutionTaskId, EPCGPerExecutionCacheDataType::Bounds); }
+	TOptional<FBox> GetOriginalBounds(FPCGTaskId InGraphExecutionTaskId) { return PerExecutionCache.GetExecutionCacheBounds(InGraphExecutionTaskId, EPCGPerExecutionCacheDataType::OriginalBounds); }
+	TOptional<FBox> GetLocalSpaceBounds(FPCGTaskId InGraphExecutionTaskId) { return PerExecutionCache.GetExecutionCacheBounds(InGraphExecutionTaskId, EPCGPerExecutionCacheDataType::LocalSpaceBounds); }
+	TOptional<FBox> GetOriginalLocalSpaceBounds(FPCGTaskId InGraphExecutionTaskId) { return PerExecutionCache.GetExecutionCacheBounds(InGraphExecutionTaskId, EPCGPerExecutionCacheDataType::OriginalLocalSpaceBounds); }
 
-	UPCGData* GetPCGData(FPCGTaskId InGraphExecutionTaskId) { return GetExecutionCacheData(InGraphExecutionTaskId, EExecutionCacheDataType::PCGData); }
-	UPCGData* GetInputPCGData(FPCGTaskId InGraphExecutionTaskId) { return GetExecutionCacheData(InGraphExecutionTaskId, EExecutionCacheDataType::InputData); }
-	UPCGData* GetActorPCGData(FPCGTaskId InGraphExecutionTaskId) { return GetExecutionCacheData(InGraphExecutionTaskId, EExecutionCacheDataType::ActorData); }
-	UPCGData* GetLandscapePCGData(FPCGTaskId InGraphExecutionTaskId) { return GetExecutionCacheData(InGraphExecutionTaskId, EExecutionCacheDataType::LandscapeData); }
-	UPCGData* GetLandscapeHeightPCGData(FPCGTaskId InGraphExecutionTaskId) { return GetExecutionCacheData(InGraphExecutionTaskId, EExecutionCacheDataType::LandscapeHeightData); }
-	UPCGData* GetOriginalActorPCGData(FPCGTaskId InGraphExecutionTaskId) { return GetExecutionCacheData(InGraphExecutionTaskId, EExecutionCacheDataType::OriginalActorData); }
-
-	void SetExecutionCacheData(FPCGTaskId InGraphExecutionTaskId, EExecutionCacheDataType InExecutionCacheDataType, UPCGData* InData);
-
-	void SetPCGData(FPCGTaskId InGraphExecutionTaskId, UPCGData* InData) { SetExecutionCacheData(InGraphExecutionTaskId, EExecutionCacheDataType::PCGData, InData); }
-	void SetInputPCGData(FPCGTaskId InGraphExecutionTaskId, UPCGData* InData) { SetExecutionCacheData(InGraphExecutionTaskId, EExecutionCacheDataType::InputData, InData); }
-	void SetActorPCGData(FPCGTaskId InGraphExecutionTaskId, UPCGData* InData) { SetExecutionCacheData(InGraphExecutionTaskId, EExecutionCacheDataType::ActorData, InData); }
-	void SetLandscapePCGData(FPCGTaskId InGraphExecutionTaskId, UPCGData* InData) { SetExecutionCacheData(InGraphExecutionTaskId, EExecutionCacheDataType::LandscapeData, InData); }
-	void SetLandscapeHeightPCGData(FPCGTaskId InGraphExecutionTaskId, UPCGData* InData) { SetExecutionCacheData(InGraphExecutionTaskId, EExecutionCacheDataType::LandscapeHeightData, InData); }
-	void SetOriginalActorPCGData(FPCGTaskId InGraphExecutionTaskId, UPCGData* InData) { SetExecutionCacheData(InGraphExecutionTaskId, EExecutionCacheDataType::OriginalActorData, InData); }
+		
+	void SetPCGData(FPCGTaskId InGraphExecutionTaskId, UPCGData* InData) { PerExecutionCache.SetExecutionCachePCGData(InGraphExecutionTaskId, EPCGPerExecutionCacheDataType::PCGData, InData); }
+	void SetInputPCGData(FPCGTaskId InGraphExecutionTaskId, UPCGData* InData) { PerExecutionCache.SetExecutionCachePCGData(InGraphExecutionTaskId, EPCGPerExecutionCacheDataType::InputData, InData); }
+	void SetActorPCGData(FPCGTaskId InGraphExecutionTaskId, UPCGData* InData) { PerExecutionCache.SetExecutionCachePCGData(InGraphExecutionTaskId, EPCGPerExecutionCacheDataType::ActorData, InData); }
+	void SetLandscapePCGData(FPCGTaskId InGraphExecutionTaskId, UPCGData* InData) { PerExecutionCache.SetExecutionCachePCGData(InGraphExecutionTaskId, EPCGPerExecutionCacheDataType::LandscapeData, InData); }
+	void SetLandscapeHeightPCGData(FPCGTaskId InGraphExecutionTaskId, UPCGData* InData) { PerExecutionCache.SetExecutionCachePCGData(InGraphExecutionTaskId, EPCGPerExecutionCacheDataType::LandscapeHeightData, InData); }
+	void SetOriginalActorPCGData(FPCGTaskId InGraphExecutionTaskId, UPCGData* InData) { PerExecutionCache.SetExecutionCachePCGData(InGraphExecutionTaskId, EPCGPerExecutionCacheDataType::OriginalActorData, InData); }
+	void SetBounds(FPCGTaskId InGraphExecutionTaskId, const FBox& InBounds) { PerExecutionCache.SetExecutionCacheBounds(InGraphExecutionTaskId, EPCGPerExecutionCacheDataType::Bounds, InBounds); }
+	void SetOriginalBounds(FPCGTaskId InGraphExecutionTaskId, const FBox& InBounds) { PerExecutionCache.SetExecutionCacheBounds(InGraphExecutionTaskId, EPCGPerExecutionCacheDataType::OriginalBounds, InBounds); }
+	void SetLocalSpaceBounds(FPCGTaskId InGraphExecutionTaskId, const FBox& InBounds) { PerExecutionCache.SetExecutionCacheBounds(InGraphExecutionTaskId, EPCGPerExecutionCacheDataType::LocalSpaceBounds, InBounds); }
+	void SetOriginalLocalSpaceBounds(FPCGTaskId InGraphExecutionTaskId, const FBox& InBounds) { PerExecutionCache.SetExecutionCacheBounds(InGraphExecutionTaskId, EPCGPerExecutionCacheDataType::OriginalLocalSpaceBounds, InBounds); }
 
 	static double GetTickBudgetInSeconds();
 		
@@ -356,21 +351,6 @@ private:
 	/** Monotonically increasing id. Should be reset once all tasks are executed, should be protected by the ScheduleLock */
 	FPCGTaskId NextTaskId = 0;
 
-	/** Struct holding different UPCGData caches that we want to compute only once per graph execution */
-	struct FGraphExecutionCache
-	{
-		FGraphExecutionCache();
-
-		void AddStructReferencedObjects(FReferenceCollector& Collector);
-
-		// Must be equal to 
-		TObjectPtr<UPCGData> Data[ExecutionCacheSize];
-	};
-
-	/** Per graph execution cache, gets emptied when executor has no more work to do */
-	FPCGLock GraphExecutionCachesLock;
-	TMap<FPCGTaskId, FGraphExecutionCache> GraphExecutionCaches;
-
 	std::atomic<bool> bNeedToExecuteTasksEnded = false;
 	std::atomic<int32> NewReadyTaskCount = 0;
 
diff --git a/Engine/Plugins/PCG/Source/PCG/Private/Graph/PCGGraphPerExecutionCache.cpp b/Engine/Plugins/PCG/Source/PCG/Private/Graph/PCGGraphPerExecutionCache.cpp
new file mode 100644
--- /dev/null
+++ b/Engine/Plugins/PCG/Source/PCG/Private/Graph/PCGGraphPerExecutionCache.cpp
@@ -0,0 +1,97 @@
+// Copyright Epic Games, Inc. All Rights Reserved.
+
+#include "PCGGraphPerExecutionCache.h"
+
+UPCGData* FPCGPerExecutionCache::GetExecutionCachePCGData(FPCGTaskId InGraphExecutionTaskId, EPCGPerExecutionCacheDataType InExecutionCacheDataType, bool& bOutFound)
+{
+	bOutFound = false;
+	if (InGraphExecutionTaskId != InvalidPCGTaskId)
+	{
+		UE::TScopeLock ScopeLock(Lock);
+		if (FPCGPerExecutionCacheEntry* PerExecutionCacheEntry = Entries.Find(InGraphExecutionTaskId))
+		{
+			TInstancedStruct<FPCGPerExecutionCacheData>& StructData = PerExecutionCacheEntry->Data[(uint8)InExecutionCacheDataType];
+			if (StructData.IsValid())
+			{
+				bOutFound = true;
+				FPCGPerExecutionCachePCGData& PCGData = StructData.GetMutable<FPCGPerExecutionCachePCGData>();
+				return PCGData.Data;
+			}
+		}
+	}
+
+	return nullptr;
+}
+
+void FPCGPerExecutionCache::SetExecutionCachePCGData(FPCGTaskId InGraphExecutionTaskId, EPCGPerExecutionCacheDataType InExecutionCacheDataType, UPCGData* InData)
+{
+	if (InGraphExecutionTaskId != InvalidPCGTaskId)
+	{
+		UE::TScopeLock ScopeLock(Lock);
+		FPCGPerExecutionCacheEntry& PerExecutionCacheEntry = Entries.FindOrAdd(InGraphExecutionTaskId);
+		check(!PerExecutionCacheEntry.Data[(uint8)InExecutionCacheDataType].IsValid());
+		PerExecutionCacheEntry.Data[(uint8)InExecutionCacheDataType].InitializeAs<FPCGPerExecutionCachePCGData>(InData);
+	}
+}
+
+TOptional<FBox> FPCGPerExecutionCache::GetExecutionCacheBounds(FPCGTaskId InGraphExecutionTaskId, EPCGPerExecutionCacheDataType InExecutionCacheDataType)
+{
+	if (InGraphExecutionTaskId != InvalidPCGTaskId)
+	{
+		UE::TScopeLock ScopeLock(Lock);
+		if (FPCGPerExecutionCacheEntry* PerExecutionCacheEntry = Entries.Find(InGraphExecutionTaskId))
+		{
+			TInstancedStruct<FPCGPerExecutionCacheData>& StructData = PerExecutionCacheEntry->Data[(uint8)InExecutionCacheDataType];
+			if (StructData.IsValid())
+			{
+				FPCGPerExecutionCacheBounds& BoundsData = StructData.GetMutable<FPCGPerExecutionCacheBounds>();
+				return BoundsData.Bounds;
+			}
+		}
+	}
+
+	return TOptional<FBox>();
+}
+
+void FPCGPerExecutionCache::SetExecutionCacheBounds(FPCGTaskId InGraphExecutionTaskId, EPCGPerExecutionCacheDataType InExecutionCacheDataType, const FBox& InBounds)
+{
+	if (InGraphExecutionTaskId != InvalidPCGTaskId)
+	{
+		UE::TScopeLock ScopeLock(Lock);
+		FPCGPerExecutionCacheEntry& PerExecutionCacheEntry = Entries.FindOrAdd(InGraphExecutionTaskId);
+		check(!PerExecutionCacheEntry.Data[(uint8)InExecutionCacheDataType].IsValid());
+		PerExecutionCacheEntry.Data[(uint8)InExecutionCacheDataType].InitializeAs<FPCGPerExecutionCacheBounds>(InBounds);
+	}
+}
+
+void FPCGPerExecutionCache::Clear()
+{
+	UE::TScopeLock ScopeLock(Lock);
+	Entries.Empty();
+}
+
+void FPCGPerExecutionCache::AddStructReferencedObjects(FReferenceCollector& Collector)
+{
+	UE::TScopeLock ScopeLock(Lock);
+	for (TMap<FPCGTaskId, FPCGPerExecutionCacheEntry>::TIterator It = Entries.CreateIterator(); It; ++It)
+	{
+		It.Value().AddStructReferencedObjects(Collector);
+	}
+}
+
+void FPCGPerExecutionCacheEntry::AddStructReferencedObjects(FReferenceCollector& Collector)
+{
+	for (uint8 DataIndex = 0; DataIndex < (uint8)EPCGPerExecutionCacheDataType::Count; ++DataIndex)
+	{
+		if (Data[DataIndex].IsValid())
+		{
+			Data[DataIndex].GetMutable().AddStructReferencedObjects(Collector);
+		}
+	}
+}
+
+void FPCGPerExecutionCachePCGData::AddStructReferencedObjects(FReferenceCollector& Collector)
+{
+	Collector.AddReferencedObject(Data);
+}
+
diff --git a/Engine/Plugins/PCG/Source/PCG/Private/Graph/PCGGraphPerExecutionCache.h b/Engine/Plugins/PCG/Source/PCG/Private/Graph/PCGGraphPerExecutionCache.h
new file mode 100644
--- /dev/null
+++ b/Engine/Plugins/PCG/Source/PCG/Private/Graph/PCGGraphPerExecutionCache.h
@@ -0,0 +1,94 @@
+// Copyright Epic Games, Inc. All Rights Reserved.
+
+#pragma once
+
+#include "PCGCommon.h"
+
+#include "Misc/SpinLock.h"
+#include "StructUtils/InstancedStruct.h"
+
+#include "PCGGraphPerExecutionCache.generated.h"
+
+class UPCGData;
+
+USTRUCT()
+struct FPCGPerExecutionCacheData
+{
+	GENERATED_BODY()
+
+	FPCGPerExecutionCacheData() = default;
+	virtual ~FPCGPerExecutionCacheData() = default;
+
+	virtual void AddStructReferencedObjects(FReferenceCollector& Collector) {}
+};
+
+USTRUCT()
+struct FPCGPerExecutionCachePCGData : public FPCGPerExecutionCacheData
+{
+	GENERATED_BODY()
+
+	FPCGPerExecutionCachePCGData() = default;
+	FPCGPerExecutionCachePCGData(UPCGData* InData)
+		: Data(InData)
+	{
+	}
+
+	virtual void AddStructReferencedObjects(FReferenceCollector& Collector) override;
+
+	TObjectPtr<UPCGData> Data = nullptr;
+};
+
+USTRUCT()
+struct FPCGPerExecutionCacheBounds : public FPCGPerExecutionCacheData
+{
+	GENERATED_BODY()
+
+	FPCGPerExecutionCacheBounds() = default;
+	FPCGPerExecutionCacheBounds(const FBox& InBounds)
+		: Bounds(InBounds)
+	{
+	}
+
+	FBox Bounds;
+};
+
+enum class EPCGPerExecutionCacheDataType : uint8
+{
+	PCGData = 0,
+	InputData,
+	ActorData,
+	LandscapeData,
+	LandscapeHeightData,
+	OriginalActorData,
+	Bounds,
+	OriginalBounds,
+	LocalSpaceBounds,
+	OriginalLocalSpaceBounds,
+	Count
+};
+
+struct FPCGPerExecutionCacheEntry
+{
+	void AddStructReferencedObjects(FReferenceCollector& Collector);
+
+	// Fixed size array for now, could become a map at some point
+	TInstancedStruct<FPCGPerExecutionCacheData> Data[(uint8)EPCGPerExecutionCacheDataType::Count];
+};
+
+struct FPCGPerExecutionCache
+{
+	using FPCGLock = UE::FSpinLock;
+
+	UPCGData* GetExecutionCachePCGData(FPCGTaskId InGraphExecutionTaskId, EPCGPerExecutionCacheDataType InExecutionCacheDataType, bool& bOutFound);
+	void SetExecutionCachePCGData(FPCGTaskId InGraphExecutionTaskId, EPCGPerExecutionCacheDataType InExecutionCacheDataType, UPCGData* InData);
+
+	TOptional<FBox> GetExecutionCacheBounds(FPCGTaskId InGraphExecutionTaskId, EPCGPerExecutionCacheDataType InExecutionCacheDataType);
+	void SetExecutionCacheBounds(FPCGTaskId InGraphExecutionTaskId, EPCGPerExecutionCacheDataType InExecutionCacheDataType, const FBox& InBounds);
+
+	void Clear();
+	void AddStructReferencedObjects(FReferenceCollector& Collector);
+
+	/** Per graph execution cache, gets emptied when executor has no more work to do */
+	FPCGLock Lock;
+	TMap<FPCGTaskId, FPCGPerExecutionCacheEntry> Entries;
+};
\ No newline at end of file
diff --git a/Engine/Plugins/PCG/Source/PCG/Private/PCGComponent.cpp b/Engine/Plugins/PCG/Source/PCG/Private/PCGComponent.cpp
--- a/Engine/Plugins/PCG/Source/PCG/Private/PCGComponent.cpp
+++ b/Engine/Plugins/PCG/Source/PCG/Private/PCGComponent.cpp
@@ -3121,8 +3121,9 @@ UPCGData* UPCGComponent::GetPCGData() const
 	UPCGData* Data = nullptr;
 	if (UPCGSubsystem* PCGSubsystem = GetSubsystem())
 	{
-		Data = PCGSubsystem->GetPCGData(CurrentGenerationTask);
-		if (!Data)
+		bool bFound = false;
+		Data = PCGSubsystem->GetPCGData(CurrentGenerationTask, bFound);
+		if (!bFound)
 		{
 			Data = CachedPCGData ? CachedPCGData.Get() : CreatePCGData();
 			CachedPCGData = Data;
@@ -3143,8 +3144,9 @@ UPCGData* UPCGComponent::GetInputPCGData() const
 	UPCGData* Data = nullptr;
 	if (UPCGSubsystem* PCGSubsystem = GetSubsystem())
 	{
-		Data = PCGSubsystem->GetInputPCGData(CurrentGenerationTask);
-		if (!Data)
+		bool bFound = false;
+		Data = PCGSubsystem->GetInputPCGData(CurrentGenerationTask, bFound);
+		if (!bFound)
 		{
 			Data = CachedInputData ? CachedInputData.Get() : CreateInputPCGData();
 			CachedInputData = Data;
@@ -3165,8 +3167,9 @@ UPCGData* UPCGComponent::GetActorPCGData() const
 	UPCGData* Data = nullptr;
 	if (UPCGSubsystem* PCGSubsystem = GetSubsystem())
 	{
-		Data = PCGSubsystem->GetActorPCGData(CurrentGenerationTask);
-		if (!Data)
+		bool bFound = false;
+		Data = PCGSubsystem->GetActorPCGData(CurrentGenerationTask, bFound);
+		if (!bFound)
 		{
 			PCG_EXECUTION_CACHE_VALIDATION_CHECK(this);
 			Data = CachedActorData && !IsLandscapeCachedDataDirty(CachedActorData) ? CachedActorData.Get() : CreateActorPCGData();
@@ -3188,8 +3191,9 @@ UPCGData* UPCGComponent::GetLandscapePCGData() const
 	UPCGData* Data = nullptr;
 	if (UPCGSubsystem* PCGSubsystem = GetSubsystem())
 	{
-		Data = PCGSubsystem->GetLandscapePCGData(CurrentGenerationTask);
-		if (!Data)
+		bool bFound = false;
+		Data = PCGSubsystem->GetLandscapePCGData(CurrentGenerationTask, bFound);
+		if (!bFound)
 		{
 			Data = CachedLandscapeData && !IsLandscapeCachedDataDirty(CachedLandscapeData) ? CachedLandscapeData.Get() : CreateLandscapePCGData(/*bHeightOnly=*/false);
 			CachedLandscapeData = Data;
@@ -3210,8 +3214,9 @@ UPCGData* UPCGComponent::GetLandscapeHeightPCGData() const
 	UPCGData* Data = nullptr;
 	if (UPCGSubsystem* PCGSubsystem = GetSubsystem())
 	{
-		Data = PCGSubsystem->GetLandscapeHeightPCGData(CurrentGenerationTask);
-		if (!Data)
+		bool bFound = false;
+		Data = PCGSubsystem->GetLandscapeHeightPCGData(CurrentGenerationTask, bFound);
+		if (!bFound)
 		{
 			Data = CachedLandscapeHeightData && !IsLandscapeCachedDataDirty(CachedLandscapeHeightData) ? CachedLandscapeHeightData.Get() : CreateLandscapePCGData(/*bHeightOnly=*/true);
 			CachedLandscapeHeightData = Data;
@@ -3238,8 +3243,9 @@ UPCGData* UPCGComponent::GetOriginalActorPCGData() const
 	UPCGData* Data = nullptr;
 	if (UPCGSubsystem* PCGSubsystem = GetSubsystem())
 	{
-		Data = PCGSubsystem->GetOriginalActorPCGData(CurrentGenerationTask);
-		if (!Data)
+		bool bFound = false;
+		Data = PCGSubsystem->GetOriginalActorPCGData(CurrentGenerationTask, bFound);
+		if (!bFound)
 		{
 			if (UPCGComponent* OriginalComponent = PartitionActor->GetOriginalComponent(this))
 			{
@@ -3371,6 +3377,27 @@ void UPCGComponent::RefreshSchedulingPolicy()
 #endif
 }
 
+void UPCGComponent::ExecutePreGraph(FPCGContext* InContext)
+{
+#if WITH_EDITOR
+	StartGenerationInProgress();
+#endif
+
+	ensureMsgf(GetGenerationTaskId() != InvalidPCGTaskId, TEXT("Component was Scheduled for generation without having its CurrentGenerationTask assigned"));
+
+	{
+		PCG_EXECUTION_CACHE_VALIDATION_CREATE_SCOPE(Component);
+
+		// Call getters which will create the data and cache it
+		GetActorPCGData();
+		GetOriginalActorPCGData();
+		GetGridBounds();
+		GetOriginalGridBounds();
+		GetLocalSpaceBounds();
+		GetOriginalLocalSpaceBounds();
+	}
+}
+
 UPCGData* UPCGComponent::CreatePCGData() const
 {
 	TRACE_CPUPROFILER_EVENT_SCOPE(UPCGComponent::CreatePCGData);
@@ -3512,7 +3539,144 @@ bool UPCGComponent::IsLandscapeCachedDataDirty(const UPCGData* Data) const
 
 FBox UPCGComponent::GetGridBounds() const
 {
-	return PCGHelpers::GetGridBounds(GetOwner(), this);
+	if (CurrentGenerationTask != InvalidPCGTaskId)
+	{
+		UPCGSubsystem* PCGSubsystem = GetSubsystem();
+		check(PCGSubsystem);
+		
+		TOptional<FBox> GridBounds = PCGSubsystem->GetBounds(CurrentGenerationTask);
+		
+		if (!GridBounds.IsSet())
+		{
+			PCG_EXECUTION_CACHE_VALIDATION_CHECK(this);
+			GridBounds = PCGHelpers::GetGridBounds(GetOwner(), this);
+
+			PCGSubsystem->SetBounds(CurrentGenerationTask, GridBounds.GetValue());
+			if (PCGSubsystem->IsGraphCacheDebuggingEnabled())
+			{
+				UE_LOG(LogPCG, Log, TEXT("         [%s] CACHE REFRESH GridBounds"), *GetOwner()->GetName());
+			}
+		}
+
+		check(GridBounds.IsSet());
+		return GridBounds.GetValue();
+	}
+	else
+	{
+		return PCGHelpers::GetGridBounds(GetOwner(), this);
+	}
+}
+
+FBox UPCGComponent::GetOriginalGridBounds() const
+{
+	const APCGPartitionActor* PartitionActor = Cast<APCGPartitionActor>(GetOwner());
+	if (!PartitionActor)
+	{
+		return GetGridBounds();
+	}
+
+	if (CurrentGenerationTask != InvalidPCGTaskId)
+	{
+		UPCGSubsystem* PCGSubsystem = GetSubsystem();
+		check(PCGSubsystem);
+
+		TOptional<FBox> OriginalBounds = PCGSubsystem->GetOriginalBounds(CurrentGenerationTask);
+
+		if (!OriginalBounds.IsSet())
+		{
+			PCG_EXECUTION_CACHE_VALIDATION_CHECK(this);
+			const UPCGComponent* OriginalComponent = PartitionActor->GetOriginalComponent(this);
+			const UPCGComponent* ValidComponent = OriginalComponent ? OriginalComponent : this;
+
+			OriginalBounds = PCGHelpers::GetGridBounds(ValidComponent->GetOwner(), ValidComponent);
+
+			PCGSubsystem->SetOriginalBounds(CurrentGenerationTask, OriginalBounds.GetValue());
+			if (PCGSubsystem->IsGraphCacheDebuggingEnabled())
+			{
+				UE_LOG(LogPCG, Log, TEXT("         [%s] CACHE REFRESH Original Bounds"), *GetOwner()->GetName());
+			}
+		}
+
+		check(OriginalBounds.IsSet());
+		return OriginalBounds.GetValue();
+	}
+	else
+	{
+		const UPCGComponent* OriginalComponent = PartitionActor->GetOriginalComponent(this);
+		const UPCGComponent* ValidComponent = OriginalComponent ? OriginalComponent : this;
+		return PCGHelpers::GetGridBounds(ValidComponent->GetOwner(), ValidComponent);
+	}
+}
+
+FBox UPCGComponent::GetLocalSpaceBounds() const
+{
+	if (CurrentGenerationTask != InvalidPCGTaskId)
+	{
+		UPCGSubsystem* PCGSubsystem = GetSubsystem();
+		check(PCGSubsystem);
+
+		TOptional<FBox> LocalSpaceBounds = PCGSubsystem->GetLocalSpaceBounds(CurrentGenerationTask);
+
+		if (!LocalSpaceBounds.IsSet())
+		{
+			PCG_EXECUTION_CACHE_VALIDATION_CHECK(this);
+			LocalSpaceBounds = PCGHelpers::GetActorLocalBounds(GetOwner());
+
+			PCGSubsystem->SetLocalSpaceBounds(CurrentGenerationTask, LocalSpaceBounds.GetValue());
+			if (PCGSubsystem->IsGraphCacheDebuggingEnabled())
+			{
+				UE_LOG(LogPCG, Log, TEXT("         [%s] CACHE REFRESH Local Space Bounds"), *GetOwner()->GetName());
+			}
+		}
+
+		check(LocalSpaceBounds.IsSet());
+		return LocalSpaceBounds.GetValue();
+	}
+	else
+	{
+		return PCGHelpers::GetActorLocalBounds(GetOwner());
+	}
+}
+
+FBox UPCGComponent::GetOriginalLocalSpaceBounds() const
+{
+	const APCGPartitionActor* PartitionActor = Cast<APCGPartitionActor>(GetOwner());
+	if (!PartitionActor)
+	{
+		return GetLocalSpaceBounds();
+	}
+
+	if (CurrentGenerationTask != InvalidPCGTaskId)
+	{
+		UPCGSubsystem* PCGSubsystem = GetSubsystem();
+		check(PCGSubsystem);
+
+		TOptional<FBox> OriginalLocalSpaceBounds = PCGSubsystem->GetOriginalLocalSpaceBounds(CurrentGenerationTask);
+
+		if (!OriginalLocalSpaceBounds.IsSet())
+		{
+			PCG_EXECUTION_CACHE_VALIDATION_CHECK(this);
+			const UPCGComponent* OriginalComponent = PartitionActor->GetOriginalComponent(this);
+			const UPCGComponent* ValidComponent = OriginalComponent ? OriginalComponent : this;
+
+			OriginalLocalSpaceBounds = PCGHelpers::GetActorLocalBounds(ValidComponent->GetOwner());
+
+			PCGSubsystem->SetOriginalLocalSpaceBounds(CurrentGenerationTask, OriginalLocalSpaceBounds.GetValue());
+			if (PCGSubsystem->IsGraphCacheDebuggingEnabled())
+			{
+				UE_LOG(LogPCG, Log, TEXT("         [%s] CACHE REFRESH Original Local Space Bounds"), *GetOwner()->GetName());
+			}
+		}
+
+		check(OriginalLocalSpaceBounds.IsSet());
+		return OriginalLocalSpaceBounds.GetValue();
+	}
+	else
+	{
+		const UPCGComponent* OriginalComponent = PartitionActor->GetOriginalComponent(this);
+		const UPCGComponent* ValidComponent = OriginalComponent ? OriginalComponent : this;
+		return PCGHelpers::GetActorLocalBounds(ValidComponent->GetOwner());
+	}
 }
 
 FBox UPCGComponent::GetGridBounds(const AActor* Actor) const
@@ -4296,6 +4460,24 @@ FBox FPCGComponentExecutionState::GetBounds() const
 	return Component->GetGridBounds();
 }
 
+FBox FPCGComponentExecutionState::GetOriginalBounds() const
+{
+	check(Component);
+	return Component->GetOriginalGridBounds();
+}
+
+FBox FPCGComponentExecutionState::GetLocalSpaceBounds() const
+{
+	check(Component);
+	return Component->GetLocalSpaceBounds();
+}
+
+FBox FPCGComponentExecutionState::GetOriginalLocalSpaceBounds() const
+{
+	check(Component);
+	return Component->GetOriginalLocalSpaceBounds();
+}
+
 bool FPCGComponentExecutionState::Use2DGrid() const
 {
 	check(Component);
@@ -4332,6 +4514,12 @@ bool FPCGComponentExecutionState::IsGenerating() const
 	return Component->IsGenerating();
 }
 
+void FPCGComponentExecutionState::ExecutePreGraph(FPCGContext* InContext)
+{
+	check(Component);
+	Component->ExecutePreGraph(InContext);
+}
+
 bool FPCGComponentExecutionState::IsManagedByRuntimeGenSystem() const
 {
 	check(Component);
diff --git a/Engine/Plugins/PCG/Source/PCG/Private/Subsystems/PCGSubsystem.cpp b/Engine/Plugins/PCG/Source/PCG/Private/Subsystems/PCGSubsystem.cpp
--- a/Engine/Plugins/PCG/Source/PCG/Private/Subsystems/PCGSubsystem.cpp
+++ b/Engine/Plugins/PCG/Source/PCG/Private/Subsystems/PCGSubsystem.cpp
@@ -1919,34 +1919,60 @@ void UPCGSubsystem::OnPCGGraphCleaned(UPCGComponent* InComponent)
 
 #endif // WITH_EDITOR
 
-UPCGData* UPCGSubsystem::GetPCGData(FPCGTaskId InGraphExecutionTaskId)
+UPCGData* UPCGSubsystem::GetPCGData(FPCGTaskId InGraphExecutionTaskId, bool& bOutFound)
 {
-	return GraphExecutor ? GraphExecutor->GetPCGData(InGraphExecutionTaskId) : nullptr;
+	bOutFound = false;
+	return GraphExecutor ? GraphExecutor->GetPCGData(InGraphExecutionTaskId, bOutFound) : nullptr;
 }
 
-UPCGData* UPCGSubsystem::GetInputPCGData(FPCGTaskId InGraphExecutionTaskId)
+UPCGData* UPCGSubsystem::GetInputPCGData(FPCGTaskId InGraphExecutionTaskId, bool& bOutFound)
 {
-	return GraphExecutor ? GraphExecutor->GetInputPCGData(InGraphExecutionTaskId) : nullptr;
+	bOutFound = false;
+	return GraphExecutor ? GraphExecutor->GetInputPCGData(InGraphExecutionTaskId, bOutFound) : nullptr;
 }
 
-UPCGData* UPCGSubsystem::GetActorPCGData(FPCGTaskId InGraphExecutionTaskId)
+UPCGData* UPCGSubsystem::GetActorPCGData(FPCGTaskId InGraphExecutionTaskId, bool& bOutFound)
 {
-	return GraphExecutor ? GraphExecutor->GetActorPCGData(InGraphExecutionTaskId) : nullptr;
+	bOutFound = false;
+	return GraphExecutor ? GraphExecutor->GetActorPCGData(InGraphExecutionTaskId, bOutFound) : nullptr;
 }
 
-UPCGData* UPCGSubsystem::GetLandscapePCGData(FPCGTaskId InGraphExecutionTaskId)
+UPCGData* UPCGSubsystem::GetLandscapePCGData(FPCGTaskId InGraphExecutionTaskId, bool& bOutFound)
 {
-	return GraphExecutor ? GraphExecutor->GetLandscapePCGData(InGraphExecutionTaskId) : nullptr;
+	bOutFound = false;
+	return GraphExecutor ? GraphExecutor->GetLandscapePCGData(InGraphExecutionTaskId, bOutFound) : nullptr;
 }
 
-UPCGData* UPCGSubsystem::GetLandscapeHeightPCGData(FPCGTaskId InGraphExecutionTaskId)
+UPCGData* UPCGSubsystem::GetLandscapeHeightPCGData(FPCGTaskId InGraphExecutionTaskId, bool& bOutFound)
 {
-	return GraphExecutor ? GraphExecutor->GetLandscapeHeightPCGData(InGraphExecutionTaskId) : nullptr;
+	bOutFound = false;
+	return GraphExecutor ? GraphExecutor->GetLandscapeHeightPCGData(InGraphExecutionTaskId, bOutFound) : nullptr;
 }
 
-UPCGData* UPCGSubsystem::GetOriginalActorPCGData(FPCGTaskId InGraphExecutionTaskId)
+UPCGData* UPCGSubsystem::GetOriginalActorPCGData(FPCGTaskId InGraphExecutionTaskId, bool& bOutFound)
 {
-	return GraphExecutor ? GraphExecutor->GetOriginalActorPCGData(InGraphExecutionTaskId) : nullptr;
+	bOutFound = false;
+	return GraphExecutor ? GraphExecutor->GetOriginalActorPCGData(InGraphExecutionTaskId, bOutFound) : nullptr;
+}
+
+TOptional<FBox> UPCGSubsystem::GetBounds(FPCGTaskId InGraphExecutionTaskId)
+{
+	return GraphExecutor ? GraphExecutor->GetBounds(InGraphExecutionTaskId) : TOptional<FBox>();
+}
+
+TOptional<FBox> UPCGSubsystem::GetOriginalBounds(FPCGTaskId InGraphExecutionTaskId)
+{
+	return GraphExecutor ? GraphExecutor->GetOriginalBounds(InGraphExecutionTaskId) : TOptional<FBox>();
+}
+
+TOptional<FBox> UPCGSubsystem::GetLocalSpaceBounds(FPCGTaskId InGraphExecutionTaskId)
+{
+	return GraphExecutor ? GraphExecutor->GetLocalSpaceBounds(InGraphExecutionTaskId) : TOptional<FBox>();
+}
+
+TOptional<FBox> UPCGSubsystem::GetOriginalLocalSpaceBounds(FPCGTaskId InGraphExecutionTaskId)
+{
+	return GraphExecutor ? GraphExecutor->GetOriginalLocalSpaceBounds(InGraphExecutionTaskId) : TOptional<FBox>();
 }
 
 void UPCGSubsystem::SetPCGData(FPCGTaskId InGraphExecutionTaskId, UPCGData* InData)
@@ -1997,6 +2023,38 @@ void UPCGSubsystem::SetOriginalActorPCGData(FPCGTaskId InGraphExecutionTaskId, U
 	}
 }
 
+void UPCGSubsystem::SetBounds(FPCGTaskId InGraphExecutionTaskId, const FBox& InBounds)
+{
+	if (GraphExecutor)
+	{
+		GraphExecutor->SetBounds(InGraphExecutionTaskId, InBounds);
+	}
+}
+
+void UPCGSubsystem::SetOriginalBounds(FPCGTaskId InGraphExecutionTaskId, const FBox& InBounds)
+{
+	if (GraphExecutor)
+	{
+		GraphExecutor->SetOriginalBounds(InGraphExecutionTaskId, InBounds);
+	}
+}
+
+void UPCGSubsystem::SetLocalSpaceBounds(FPCGTaskId InGraphExecutionTaskId, const FBox& InBounds)
+{
+	if (GraphExecutor)
+	{
+		GraphExecutor->SetLocalSpaceBounds(InGraphExecutionTaskId, InBounds);
+	}
+}
+
+void UPCGSubsystem::SetOriginalLocalSpaceBounds(FPCGTaskId InGraphExecutionTaskId, const FBox& InBounds)
+{
+	if (GraphExecutor)
+	{
+		GraphExecutor->SetOriginalLocalSpaceBounds(InGraphExecutionTaskId, InBounds);
+	}
+}
+
 void UPCGSubsystem::ExecuteBeginTickActions()
 {
 	TArray<FTickAction> Actions = MoveTemp(BeginTickActions);
diff --git a/Engine/Plugins/PCG/Source/PCG/Public/PCGComponent.h b/Engine/Plugins/PCG/Source/PCG/Public/PCGComponent.h
--- a/Engine/Plugins/PCG/Source/PCG/Public/PCGComponent.h
+++ b/Engine/Plugins/PCG/Source/PCG/Public/PCGComponent.h
@@ -103,12 +103,16 @@ public:
 	virtual UWorld* GetWorld() const override;
 	virtual bool HasAuthority() const override;
 	virtual FBox GetBounds() const override;
+	virtual FBox GetOriginalBounds() const override;
+	virtual FBox GetLocalSpaceBounds() const override;
+	virtual FBox GetOriginalLocalSpaceBounds() const override;
 	virtual bool Use2DGrid() const override;
 	virtual UPCGGraph* GetGraph() const override;
 	virtual UPCGGraphInstance* GetGraphInstance() const override;
 	virtual void OnGraphExecutionAborted(bool bQuiet, bool bCleanupUnusedResources) override;
 	virtual void Cancel() override;
 	virtual bool IsGenerating() const override;
+	virtual void ExecutePreGraph(FPCGContext* InContext) override;
 	virtual bool IsManagedByRuntimeGenSystem() const override;
 	virtual IPCGGraphExecutionSource* GetOriginalSource() const override;
 	virtual IPCGGraphExecutionSource* GetLocalSource(const FPCGGridDescriptor& GridDescriptor, const FIntVector& CellCoords) const override;
@@ -551,6 +555,10 @@ public:
 	UE_API UPCGSubsystem* GetSubsystem() const;
 
 	UE_API FBox GetGridBounds() const;
+	UE_API FBox GetOriginalGridBounds() const;
+	UE_API FBox GetLocalSpaceBounds() const;
+	UE_API FBox GetOriginalLocalSpaceBounds() const;
+
 	FBox GetLastGeneratedBounds() const { return LastGeneratedBounds; }
 
 	/** Builds the PCG data from a given actor and its PCG component, and places it in a data collection with appropriate tags */
@@ -654,6 +662,8 @@ private:
 	uint32 RuntimeGridDescriptorHash = 0;
 
 private:
+	void ExecutePreGraph(FPCGContext* InContext);
+
 	UE_API UPCGData* CreatePCGData() const;
 	UE_API UPCGData* CreateInputPCGData() const;
 	UE_API UPCGData* CreateActorPCGData() const;
diff --git a/Engine/Plugins/PCG/Source/PCG/Public/PCGGraphExecutionStateInterface.h b/Engine/Plugins/PCG/Source/PCG/Public/PCGGraphExecutionStateInterface.h
--- a/Engine/Plugins/PCG/Source/PCG/Public/PCGGraphExecutionStateInterface.h
+++ b/Engine/Plugins/PCG/Source/PCG/Public/PCGGraphExecutionStateInterface.h
@@ -53,9 +53,21 @@ public:
 	/** Returns a Transform if the ExecutionState is a spatial one. */
 	virtual FTransform GetTransform() const = 0;
 
+	/** Returns a the original source Transform if the ExecutionState is a local source. */
+	virtual FTransform GetOriginalTransform() const { return GetTransform(); }
+
 	/** Returns the ExecutionState bounds if the ExecutionState is a spatial one. */
 	virtual FBox GetBounds() const = 0;
 
+	/** Returns the ExecutionState's original source bounds if this ExecutionState is a local source. */
+	virtual FBox GetOriginalBounds() const { return GetBounds(); }
+
+	/** Returns the ExecutionState's local bounds if this ExecutionState is a spatial one. */
+	virtual FBox GetLocalSpaceBounds() const { return GetBounds(); }
+
+	/** Returns the ExecutionState's original source local space bounds if this ExecutionState is a local source. */
+	virtual FBox GetOriginalLocalSpaceBounds() const { return GetBounds(); }
+
 	/** Returns the UPCGGraph this ExecutionState is executing. */
 	virtual UPCGGraph* GetGraph() const = 0;
 
@@ -88,6 +100,9 @@ public:
 	/** Returns the subsystem associated with this source. By default, it will return the UPCGSubsystem if the source has a world, otherwise, the EditorSubsystem (in editor only, obviously). */
 	PCG_API virtual IPCGBaseSubsystem* GetSubsystem() const;
 
+	/** Cache any data that is not thread safe or that needs to be computed only once for the whole execution */
+	virtual void ExecutePreGraph(FPCGContext* InContext) {}
+
 #if WITH_EDITOR
 
 	/** Returns the FExtraCapture object for this ExecutionState */
diff --git a/Engine/Plugins/PCG/Source/PCG/Public/Subsystems/PCGSubsystem.h b/Engine/Plugins/PCG/Source/PCG/Public/Subsystems/PCGSubsystem.h
--- a/Engine/Plugins/PCG/Source/PCG/Public/Subsystems/PCGSubsystem.h
+++ b/Engine/Plugins/PCG/Source/PCG/Public/Subsystems/PCGSubsystem.h
@@ -371,19 +371,27 @@ private:
 	/** When registering PAs, we might not have a PCG World Actor already registered (at runtime). In that case we'll look for a World Actor in the same level.*/
 	UE_API APCGWorldActor* GetPCGWorldActorForPartitionActor(APCGPartitionActor* InActor);
 
-	UE_API UPCGData* GetPCGData(FPCGTaskId InGraphExecutionTaskId);
-	UE_API UPCGData* GetInputPCGData(FPCGTaskId InGraphExecutionTaskId);
-	UE_API UPCGData* GetActorPCGData(FPCGTaskId InGraphExecutionTaskId);
-	UE_API UPCGData* GetLandscapePCGData(FPCGTaskId InGraphExecutionTaskId);
-	UE_API UPCGData* GetLandscapeHeightPCGData(FPCGTaskId InGraphExecutionTaskId);
-	UE_API UPCGData* GetOriginalActorPCGData(FPCGTaskId InGraphExecutionTaskId);
-
-	UE_API void SetPCGData(FPCGTaskId InGraphExecutionTaskId, UPCGData* InData);
-	UE_API void SetInputPCGData(FPCGTaskId InGraphExecutionTaskId, UPCGData* InData);
-	UE_API void SetActorPCGData(FPCGTaskId InGraphExecutionTaskId, UPCGData* InData);
-	UE_API void SetLandscapePCGData(FPCGTaskId InGraphExecutionTaskId, UPCGData* InData);
-	UE_API void SetLandscapeHeightPCGData(FPCGTaskId InGraphExecutionTaskId, UPCGData* InData);
-	UE_API void SetOriginalActorPCGData(FPCGTaskId InGraphExecutionTaskId, UPCGData* InData);
+	UPCGData* GetPCGData(FPCGTaskId InGraphExecutionTaskId, bool& bOutFound);
+	UPCGData* GetInputPCGData(FPCGTaskId InGraphExecutionTaskId, bool& bOutFound);
+	UPCGData* GetActorPCGData(FPCGTaskId InGraphExecutionTaskId, bool& bOutFound);
+	UPCGData* GetLandscapePCGData(FPCGTaskId InGraphExecutionTaskId, bool& bOutFound);
+	UPCGData* GetLandscapeHeightPCGData(FPCGTaskId InGraphExecutionTaskId, bool& bOutFound);
+	UPCGData* GetOriginalActorPCGData(FPCGTaskId InGraphExecutionTaskId, bool& bOutFound);
+	TOptional<FBox> GetBounds(FPCGTaskId InGraphExecutionTaskId);
+	TOptional<FBox> GetOriginalBounds(FPCGTaskId InGraphExecutionTaskId);
+	TOptional<FBox> GetLocalSpaceBounds(FPCGTaskId InGraphExecutionTaskId);
+	TOptional<FBox> GetOriginalLocalSpaceBounds(FPCGTaskId InGraphExecutionTaskId);
+
+	void SetPCGData(FPCGTaskId InGraphExecutionTaskId, UPCGData* InData);
+	void SetInputPCGData(FPCGTaskId InGraphExecutionTaskId, UPCGData* InData);
+	void SetActorPCGData(FPCGTaskId InGraphExecutionTaskId, UPCGData* InData);
+	void SetLandscapePCGData(FPCGTaskId InGraphExecutionTaskId, UPCGData* InData);
+	void SetLandscapeHeightPCGData(FPCGTaskId InGraphExecutionTaskId, UPCGData* InData);
+	void SetOriginalActorPCGData(FPCGTaskId InGraphExecutionTaskId, UPCGData* InData);
+	void SetBounds(FPCGTaskId InGraphExecutionTaskId, const FBox& InBounds);
+	void SetOriginalBounds(FPCGTaskId InGraphExecutionTaskId, const FBox& InBounds);
+	void SetLocalSpaceBounds(FPCGTaskId InGraphExecutionTaskId, const FBox& InBounds);
+	void SetOriginalLocalSpaceBounds(FPCGTaskId InGraphExecutionTaskId, const FBox& InBounds);
 
 	UE_API void ExecuteBeginTickActions();
 
