[PCG] Fix inspection manager data lifetime

#rnx

diff --git a/Engine/Plugins/PCG/Source/PCGEditor/Private/Managers/PCGEditorInspectionDataManager.cpp b/Engine/Plugins/PCG/Source/PCGEditor/Private/Managers/PCGEditorInspectionDataManager.cpp
--- a/Engine/Plugins/PCG/Source/PCGEditor/Private/Managers/PCGEditorInspectionDataManager.cpp
+++ b/Engine/Plugins/PCG/Source/PCGEditor/Private/Managers/PCGEditorInspectionDataManager.cpp
@@ -7,9 +7,14 @@
 #include "PCGGraphExecutionStateInterface.h"
 #include "Nodes/PCGEditorGraphNodeBase.h"
 
+bool FPCGEditorInspectionDataEntry::IsValid() const
+{
+	return !InspectionData.TaggedData.IsEmpty();
+}
+
 void FPCGEditorInspectionDataEntry::Clear()
 {
-	InspectionData = nullptr;
+	InspectionData = FPCGDataCollection{};
 	Node.Reset();
 	PinName = NAME_None;
 	bIsOutputPin = true;
@@ -75,8 +80,11 @@ void FPCGEditorInspectionDataManager::OnSourceGenerated(IPCGGraphExecutionSource
 
 const FPCGDataCollection* FPCGEditorInspectionDataManager::GetInspectionData(const FPCGEditorInspectionDataEntrySetupParams& InParams) const
 {
-	const FPCGEditorInspectionDataEntry NewEntry = ConstructEntry(InParams);
-	return NewEntry.IsValid() ? NewEntry.InspectionData : nullptr;
+	const FPCGDataCollection* DataCollection = nullptr;
+	FPCGStack Stack{};
+	const UPCGPin* Pin = nullptr;
+
+	return GetInspectionData_Internal(InParams, DataCollection, Stack, Pin) ? DataCollection : nullptr;
 }
 
 void FPCGEditorInspectionDataManager::SetupInspectionEntry(const FPCGEditorInspectionDataEntrySetupParams& InParams)
@@ -91,6 +99,14 @@ const FPCGEditorInspectionDataEntry& FPCGEditorInspectionDataManager::GetInspect
 	return InspectionEntries[EntryIndex].IsValid() ? InspectionEntries[EntryIndex] : InspectionEntries[0];
 }
 
+void FPCGEditorInspectionDataManager::AddReferencedObjects(FReferenceCollector& Collector)
+{
+	for (FPCGEditorInspectionDataEntry& Entry : InspectionEntries)
+	{
+		Entry.InspectionData.AddReferences(Collector);
+	}
+}
+
 FPCGEditorInspectionDataEntry FPCGEditorInspectionDataManager::ConstructEntry(const FPCGEditorInspectionDataEntrySetupParams& InParams) const
 {
 	FPCGEditorInspectionDataEntry Result{};
@@ -104,10 +120,32 @@ FPCGEditorInspectionDataEntry FPCGEditorInspectionDataManager::ConstructEntry(co
 		return Result;
 	}
 
+	const FPCGDataCollection* DataCollection = nullptr;
+	FPCGStack Stack{};
+	const UPCGPin* Pin = nullptr;
+	
+	if (!GetInspectionData_Internal(InParams, DataCollection, Stack, Pin))
+	{
+		return Result;
+	}
+
+	check(DataCollection && Pin);
+
+	Result.InspectionData = *DataCollection;
+	Result.Node = InParams.Node;
+	Result.PinName = Pin->Properties.Label;
+	Result.bIsOutputPin = InParams.bIsOutputPin;
+	Result.Stack = MoveTemp(Stack);
+
+	return Result;
+}
+
+bool FPCGEditorInspectionDataManager::GetInspectionData_Internal(const FPCGEditorInspectionDataEntrySetupParams& InParams, FPCGDataCollection const*& OutData, FPCGStack& OutStack, UPCGPin const*& OutPin) const
+{
 	const UPCGNode* PCGNode = InParams.Node->GetPCGNode();
 	if (!PCGNode)
 	{
-		return Result;
+		return false;
 	}
 
 	const TArray<TObjectPtr<UPCGPin>>& Pins = InParams.bIsOutputPin ? PCGNode->GetOutputPins() : PCGNode->GetInputPins();
@@ -119,36 +157,26 @@ FPCGEditorInspectionDataEntry FPCGEditorInspectionDataManager::ConstructEntry(co
 	
 	if (!Pins.IsValidIndex(PinIndex))
 	{
-		return Result;
+		return false;
 	}
 
 	const UPCGPin* Pin = Pins[PinIndex];
 
-	PCGEditorGraphUtils::GetInspectablePin(PCGNode, Pin, PCGNode, Pin);
+	PCGEditorGraphUtils::GetInspectablePin(PCGNode, Pin, PCGNode, OutPin);
 
-	if (!PCGNode || !Pin)
+	if (!PCGNode || !OutPin)
 	{
-		return Result;
+		return false;
 	}
 
 	// Create a temporary stack with Node+Pin to query the exact DataCollection we are inspecting
-	FPCGStack Stack = StackBeingInspected;
-	TArray<FPCGStackFrame>& StackFrames = Stack.GetStackFramesMutable();
+	OutStack = StackBeingInspected;
+	TArray<FPCGStackFrame>& StackFrames = OutStack.GetStackFramesMutable();
 	StackFrames.Reserve(StackFrames.Num() + 2);
 	StackFrames.Emplace(PCGNode);
-	StackFrames.Emplace(Pin);
+	StackFrames.Emplace(OutPin);
 
-	const FPCGDataCollection* DataCollection = PCGSourceBeingInspected.Get()->GetExecutionState().GetInspection().GetInspectionData(Stack);
-	if (!DataCollection)
-	{
-		return Result;
-	}
+	OutData = PCGSourceBeingInspected.Get()->GetExecutionState().GetInspection().GetInspectionData(OutStack);
 
-	Result.InspectionData = DataCollection;
-	Result.Node = InParams.Node;
-	Result.PinName = Pin->Properties.Label;
-	Result.bIsOutputPin = InParams.bIsOutputPin;
-	Result.Stack = MoveTemp(Stack);
-
-	return Result;
+	return OutData != nullptr;
 }
diff --git a/Engine/Plugins/PCG/Source/PCGEditor/Private/PCGEditor.cpp b/Engine/Plugins/PCG/Source/PCGEditor/Private/PCGEditor.cpp
--- a/Engine/Plugins/PCG/Source/PCGEditor/Private/PCGEditor.cpp
+++ b/Engine/Plugins/PCG/Source/PCGEditor/Private/PCGEditor.cpp
@@ -663,6 +663,8 @@ void FPCGEditor::AddReferencedObjects(FReferenceCollector& Collector)
 	Collector.AddReferencedObject(PCGGraphBeingEdited);
 	Collector.AddReferencedObject(PCGDefaultExecutionSource);
 
+	InspectionDataManager.AddReferencedObjects(Collector);
+
 	for (TSharedPtr<SPCGEditorGraphAttributeListView> ALV : AttributesWidgets)
 	{
 		if (ALV)
diff --git a/Engine/Plugins/PCG/Source/PCGEditor/Private/Widgets/SPCGEditorGraphAttributeListView.cpp b/Engine/Plugins/PCG/Source/PCGEditor/Private/Widgets/SPCGEditorGraphAttributeListView.cpp
--- a/Engine/Plugins/PCG/Source/PCGEditor/Private/Widgets/SPCGEditorGraphAttributeListView.cpp
+++ b/Engine/Plugins/PCG/Source/PCGEditor/Private/Widgets/SPCGEditorGraphAttributeListView.cpp
@@ -805,7 +805,7 @@ const FPCGDataCollection* SPCGEditorGraphAttributeListView::GetInspectionData()
 	}
 
 	const FPCGEditorInspectionDataEntry& Entry = PCGEditorPtr.Pin()->GetInspectionDataManager().GetInspectionEntry(WidgetEntryNumber);
-	return Entry.InspectionData;
+	return &Entry.InspectionData;
 }
 
 void SPCGEditorGraphAttributeListView::RefreshAttributeList()
diff --git a/Engine/Plugins/PCG/Source/PCGEditor/Public/Managers/PCGEditorInspectionDataManager.h b/Engine/Plugins/PCG/Source/PCGEditor/Public/Managers/PCGEditorInspectionDataManager.h
--- a/Engine/Plugins/PCG/Source/PCGEditor/Public/Managers/PCGEditorInspectionDataManager.h
+++ b/Engine/Plugins/PCG/Source/PCGEditor/Public/Managers/PCGEditorInspectionDataManager.h
@@ -9,13 +9,13 @@ class UPCGEditorGraphNodeBase;
 
 struct FPCGEditorInspectionDataEntry
 {
-	const FPCGDataCollection* InspectionData = nullptr;
+	FPCGDataCollection InspectionData;
 	TWeakObjectPtr<const UPCGEditorGraphNodeBase> Node;
 	FName PinName = NAME_None;
 	bool bIsOutputPin = true;
 	FPCGStack Stack;
 
-	bool IsValid() const { return InspectionData != nullptr; }
+	PCGEDITOR_API bool IsValid() const;
 	PCGEDITOR_API void Clear();
 };
 
@@ -75,11 +75,16 @@ public:
 	// If the entry is invalid at this index, return the first one (so multiple widget would show the same thing).
 	PCGEDITOR_API const FPCGEditorInspectionDataEntry& GetInspectionEntry(int32 EntryIndex);
 
+	/** Add UObject references for GC */
+	PCGEDITOR_API void AddReferencedObjects(FReferenceCollector& Collector);
+
 	FOnInspectedStackChanged OnInspectedStackChangedDelegate;
 	FOnInspectionEntryChanged OnInspectionEntryChangedDelegate;
 
 private:
 	FPCGEditorInspectionDataEntry ConstructEntry(const FPCGEditorInspectionDataEntrySetupParams& InParams) const;
+
+	bool GetInspectionData_Internal(const FPCGEditorInspectionDataEntrySetupParams& InParams, FPCGDataCollection const*& OutData, FPCGStack& OutStack, UPCGPin const*& OutPin) const;
 	
 	FPCGSoftGraphExecutionSource PCGSourceBeingInspected;
 	// Implementation note: we'll keep the last valid source inspected so we don't un-inspect on spurious selection changes
