[PCG] Added node to get assets from a folder or from a collection. Also allows to extract class/bp class.

diff --git a/Engine/Plugins/PCG/Source/PCG/Private/Elements/IO/PCGGetAssetList.cpp b/Engine/Plugins/PCG/Source/PCG/Private/Elements/IO/PCGGetAssetList.cpp
new file mode 100644
--- /dev/null
+++ b/Engine/Plugins/PCG/Source/PCG/Private/Elements/IO/PCGGetAssetList.cpp
@@ -0,0 +1,136 @@
+// Copyright Epic Games, Inc. All Rights Reserved.
+
+#include "Elements/IO/PCGGetAssetList.h"
+
+#include "PCGParamData.h"
+#include "Utils/PCGLogErrors.h"
+
+#include "Algo/Transform.h"
+#include "AssetRegistry/AssetData.h"
+#include "AssetRegistry/AssetRegistryModule.h"
+#include "Blueprint/BlueprintSupport.h"
+
+#if WITH_EDITOR
+#include "CollectionManagerModule.h"
+#include "CollectionManagerTypes.h"
+#include "ICollectionContainer.h"
+#include "ICollectionManager.h"
+#endif // WITH_EDITOR
+
+#define LOCTEXT_NAMESPACE "PCGGetAssetListElement"
+
+namespace PCGGetAsset
+{
+	const FLazyName AssetReferenceLabel = TEXT("AssetReference");
+	const FLazyName ClassPathLabel = TEXT("ClassPath");
+	const FLazyName BlueprintGeneratedClassPathLabel = TEXT("BlueprintGeneratedClassPath");
+}
+
+TArray<FPCGPinProperties> UPCGGetAssetListSettings::OutputPinProperties() const
+{
+	TArray<FPCGPinProperties> PinProperties;
+	PinProperties.Emplace(PCGPinConstants::DefaultOutputLabel, EPCGDataType::Param);
+
+	return PinProperties;
+}
+
+bool FPCGGetAssetListElement::ExecuteInternal(FPCGContext* Context) const
+{
+	TRACE_CPUPROFILER_EVENT_SCOPE(FPCGGetAssetListElement::Execute);
+
+	check(Context);
+
+	const UPCGGetAssetListSettings* Settings = Context->GetInputSettings<UPCGGetAssetListSettings>();
+	check(Settings);
+
+#if WITH_EDITOR
+	TArray<FAssetData> Assets;
+
+	FAssetRegistryModule* AssetRegistryModule = FModuleManager::GetModulePtr<FAssetRegistryModule>(TEXT("AssetRegistry"));
+	if (!AssetRegistryModule)
+	{
+		PCGLog::LogErrorOnGraph(LOCTEXT("AssetRegistryNotLoaded", "The asset registry module is not loaded."), Context);
+		return true;
+	}
+
+	FCollectionManagerModule* CollectionManagerPtr = FModuleManager::GetModulePtr<FCollectionManagerModule>("CollectionManager");
+	if (!CollectionManagerPtr)
+	{
+		PCGLog::LogErrorOnGraph(LOCTEXT("CollectionManagerModuleNotLoaded", "The collection manager module isn't loaded."), Context);
+		return true;
+	}
+
+	bool bSuccessfullyRetrievedAssets = false;
+
+	if (Settings->bGetAssetsFromFolderPath)
+	{
+		if (!AssetRegistryModule->Get().GetAssetsByPath(FName(*Settings->Directory.Path), Assets, /*bRecursive=*/true))
+		{
+			if (!Settings->bQuiet)
+			{
+				PCGLog::LogWarningOnGraph(LOCTEXT("FolderInvalid", "Unable to retrieve assets from folder."), Context);
+			}
+		}
+	}
+	else
+	{
+		TSharedPtr<ICollectionContainer> CollectionContainer = CollectionManagerPtr->Get().GetProjectCollectionContainer();
+		if (CollectionContainer)
+		{
+			TArray<FSoftObjectPath> AssetPaths;
+			if (!CollectionContainer->GetAssetsInCollection(Settings->Collection, ECollectionShareType::CST_All, AssetPaths, ECollectionRecursionFlags::SelfAndChildren))
+			{
+				// Collection doesn't exist or is empty.
+				if(!Settings->bQuiet)
+				{
+					PCGLog::LogWarningOnGraph(LOCTEXT("CollectionEmptyOrInexistant", "Unable to retrieve collection or it is empty."), Context);
+				}
+			
+				return true;
+			}
+			else
+			{
+				Algo::Transform(AssetPaths, Assets, [AssetRegistryModule](const FSoftObjectPath& AssetPath) { return AssetRegistryModule->Get().GetAssetByObjectPath(AssetPath); });
+			}
+		}
+		else
+		{
+			PCGLog::LogErrorOnGraph(LOCTEXT("UnableToGetCollectionContainer", "Unable to retrieve project collection container."), Context);
+			return true;
+		}
+	}
+
+	UPCGParamData* ParamData = FPCGContext::NewObject_AnyThread<UPCGParamData>(Context);
+	check(ParamData && ParamData->Metadata);
+
+	FPCGMetadataAttribute<FSoftObjectPath>* AssetReferenceAttribute = ParamData->MutableMetadata()->CreateAttribute<FSoftObjectPath>(PCGGetAsset::AssetReferenceLabel, FSoftObjectPath(), /*bAllowsInterpolation=*/false, /*bOverrideParent=*/false);
+	FPCGMetadataAttribute<FSoftClassPath>* ClassPathAttribute = Settings->bGetClassPath ? ParamData->MutableMetadata()->CreateAttribute<FSoftClassPath>(PCGGetAsset::ClassPathLabel, FSoftClassPath(), /*bAllowsInterpolation=*/false, /*bOverrideParent=*/false) : nullptr;
+	FPCGMetadataAttribute<FSoftClassPath>* BPClassPathAttribute = Settings->bGetClassPath ? ParamData->MutableMetadata()->CreateAttribute<FSoftClassPath>(PCGGetAsset::BlueprintGeneratedClassPathLabel, FSoftClassPath(), /*bAllowsInterpolation=*/false, /*bOverrideParent=*/false) : nullptr;
+
+	for(const FAssetData& Asset : Assets)
+	{
+		PCGMetadataEntryKey Entry = ParamData->MutableMetadata()->AddEntry();
+		AssetReferenceAttribute->SetValue(Entry, Asset.GetSoftObjectPath());
+
+		if (ClassPathAttribute)
+		{
+			ClassPathAttribute->SetValue(Entry, FSoftClassPath(Asset.AssetClassPath.ToString()));
+		}
+
+		FString BPClassPath;
+		if (BPClassPathAttribute && Asset.GetTagValue(FBlueprintTags::GeneratedClassPath, BPClassPath))
+		{
+			BPClassPathAttribute->SetValue(Entry, FSoftClassPath(BPClassPath));
+		}
+	}
+
+	FPCGTaggedData& OutputData = Context->OutputData.TaggedData.Emplace_GetRef();
+	OutputData.Data = ParamData;
+#else
+	PCGLog::LogErrorOnGraph(LOCTEXT("NodeWorksOnlyInEdior", "The Get Asset List node works only in editor contexts."), Context);
+#endif // WITH_EDITOR
+
+	return true;
+}
+
+#undef LOCTEXT_NAMESPACE
\ No newline at end of file
diff --git a/Engine/Plugins/PCG/Source/PCG/Public/Elements/IO/PCGGetAssetList.h b/Engine/Plugins/PCG/Source/PCG/Public/Elements/IO/PCGGetAssetList.h
new file mode 100644
--- /dev/null
+++ b/Engine/Plugins/PCG/Source/PCG/Public/Elements/IO/PCGGetAssetList.h
@@ -0,0 +1,60 @@
+// Copyright Epic Games, Inc. All Rights Reserved.
+
+#pragma once
+
+#include "PCGContext.h"
+#include "PCGSettings.h"
+
+#include "PCGGetAssetList.generated.h"
+
+class FPCGGetAssetListElement : public IPCGElement
+{
+public:
+	virtual bool IsCacheable(const UPCGSettings* InSettings) const override { return false; }
+	virtual bool CanExecuteOnlyOnMainThread(FPCGContext* Context) const override { return true; }
+protected:
+	virtual bool ExecuteInternal(FPCGContext* Context) const;
+};
+
+/** Returns the contents of a collection or a folder in an attribute list. Note that this does not track contents and as such is not a cacheable node. */
+UCLASS(MinimalAPI, BlueprintType, ClassGroup = (Procedural))
+class UPCGGetAssetListSettings : public UPCGSettings
+{
+	GENERATED_BODY()
+
+public:
+#if WITH_EDITOR
+	virtual FName GetDefaultNodeName() const override { return FName(TEXT("GetAssetList")); }
+	virtual FText GetDefaultNodeTitle() const override { return NSLOCTEXT("PCGGetAssetListElement", "NodeTitle", "Get Asset List"); }
+	virtual FText GetNodeTooltipText() const override { return NSLOCTEXT("PCGGetAssetListElement", "NodeTooltip", "Returns the list of asset from a collection or a folder."); }
+	virtual EPCGSettingsType GetType() const override { return EPCGSettingsType::InputOutput; }
+#endif // WITH_EDITOR
+
+protected:
+	virtual TArray<FPCGPinProperties> InputPinProperties() const override { return TArray<FPCGPinProperties>(); }
+	virtual TArray<FPCGPinProperties> OutputPinProperties() const override;
+
+	virtual FPCGElementPtr CreateElement() const override { return MakeShared<FPCGGetAssetListElement>(); }
+	// ~End UPCGSettings interface implementation
+
+public:
+	/** Controls whether we will retrieve the files from a specified directory or from a collection. */
+	UPROPERTY(BlueprintReadWrite, EditAnywhere, Category = Settings, meta = (PCG_Overridable))
+	bool bGetAssetsFromFolderPath = true;
+
+	/** Directory path to parse for assets. In 'long package name' format, e.g. '/Game/' and so on. */
+	UPROPERTY(BlueprintReadWrite, EditAnywhere, Category = Settings, meta = (LongPackageName, EditCondition = "bGetAssetsFromFolderPath", EditConditionHides, PCG_Overridable))
+	FDirectoryPath Directory;
+
+	/** Name of the collection to parse for assets. */
+	UPROPERTY(BlueprintReadWrite, EditAnywhere, Category = Settings, meta = (EditCondition = "!bGetAssetsFromFolderPath", EditConditionHides, PCG_Overridable))
+	FName Collection;
+
+	/** Controls whether the class path will also be exported as part of this node in the 'ClassPath' attribute. */
+	UPROPERTY(BlueprintReadWrite, EditAnywhere, Category = Settings, meta = (PCG_Overridable))
+	bool bGetClassPath = false;
+
+	/** Controls whether inexistant or empty collections will log a warning. */
+	UPROPERTY(BlueprintReadWrite, EditAnywhere, Category = Settings, AdvancedDisplay, meta = (PCG_Overridable))
+	bool bQuiet = false;
+};
\ No newline at end of file
