[PCG] Add missing GetOriginalTransform override

#rnx

diff --git a/Engine/Plugins/PCG/Source/PCG/Private/PCGComponent.cpp b/Engine/Plugins/PCG/Source/PCG/Private/PCGComponent.cpp
--- a/Engine/Plugins/PCG/Source/PCG/Private/PCGComponent.cpp
+++ b/Engine/Plugins/PCG/Source/PCG/Private/PCGComponent.cpp
@@ -4442,6 +4442,18 @@ FTransform FPCGComponentExecutionState::GetTransform() const
 	return Component->GetOwner()->GetTransform();
 }
 
+FTransform FPCGComponentExecutionState::GetOriginalTransform() const
+{
+	if (IPCGGraphExecutionSource* OriginalSource = GetOriginalSource())
+	{
+		return OriginalSource->GetExecutionState().GetTransform();
+	}
+	else
+	{
+		return GetTransform();
+	}
+}
+
 UWorld* FPCGComponentExecutionState::GetWorld() const
 {
 	check(Component);
diff --git a/Engine/Plugins/PCG/Source/PCG/Public/PCGComponent.h b/Engine/Plugins/PCG/Source/PCG/Public/PCGComponent.h
--- a/Engine/Plugins/PCG/Source/PCG/Public/PCGComponent.h
+++ b/Engine/Plugins/PCG/Source/PCG/Public/PCGComponent.h
@@ -100,6 +100,7 @@ public:
 	virtual int32 GetSeed() const override;
 	virtual FString GetDebugName() const override;
 	virtual FTransform GetTransform() const override;
+	virtual FTransform GetOriginalTransform() const override;
 	virtual UWorld* GetWorld() const override;
 	virtual bool HasAuthority() const override;
 	virtual FBox GetBounds() const override;
