[PCG] Added lenience to find/create attribute data if domain isn't found.

#rnx

diff --git a/Engine/Plugins/PCG/Source/PCG/Public/Metadata/PCGMetadata.h b/Engine/Plugins/PCG/Source/PCG/Public/Metadata/PCGMetadata.h
--- a/Engine/Plugins/PCG/Source/PCG/Public/Metadata/PCGMetadata.h
+++ b/Engine/Plugins/PCG/Source/PCG/Public/Metadata/PCGMetadata.h
@@ -490,8 +490,7 @@ template<typename T>
 FPCGMetadataAttribute<T>* UPCGMetadata::CreateAttribute(FPCGAttributeIdentifier AttributeName, const T& DefaultValue, bool bAllowsInterpolation, bool bOverrideParent)
 {
 	FPCGMetadataDomain* FoundSubMetadata = FindOrCreateMetadataDomain(AttributeName.MetadataDomain);
-	check(FoundSubMetadata);
-	return FoundSubMetadata->CreateAttribute<T>(AttributeName.Name, DefaultValue, bAllowsInterpolation, bOverrideParent);
+	return FoundSubMetadata ? FoundSubMetadata->CreateAttribute<T>(AttributeName.Name, DefaultValue, bAllowsInterpolation, bOverrideParent) : nullptr;
 }
 
 template<typename T>
@@ -505,8 +504,7 @@ template<typename T>
 FPCGMetadataAttribute<T>* UPCGMetadata::FindOrCreateAttribute(FPCGAttributeIdentifier AttributeName, const T& DefaultValue, bool bAllowsInterpolation, bool bOverrideParent, bool bOverwriteIfTypeMismatch)
 {
 	FPCGMetadataDomain* FoundSubMetadata = FindOrCreateMetadataDomain(AttributeName.MetadataDomain);
-	check(FoundSubMetadata);
-	return FoundSubMetadata->FindOrCreateAttribute<T>(AttributeName.Name, DefaultValue, bAllowsInterpolation, bOverrideParent, bOverwriteIfTypeMismatch);
+	return FoundSubMetadata ? FoundSubMetadata->FindOrCreateAttribute<T>(AttributeName.Name, DefaultValue, bAllowsInterpolation, bOverrideParent, bOverwriteIfTypeMismatch) : nullptr;
 }
 
 template <typename T>
