[PCG] Adjusted collection manager requirement when the get asset list is done by folder. Added enum for future expansion. Cleaned up some edge cases.

#rnx

diff --git a/Engine/Plugins/PCG/Source/PCG/Private/Elements/IO/PCGGetAssetList.cpp b/Engine/Plugins/PCG/Source/PCG/Private/Elements/IO/PCGGetAssetList.cpp
--- a/Engine/Plugins/PCG/Source/PCG/Private/Elements/IO/PCGGetAssetList.cpp
+++ b/Engine/Plugins/PCG/Source/PCG/Private/Elements/IO/PCGGetAssetList.cpp
@@ -6,11 +6,11 @@
 #include "Utils/PCGLogErrors.h"
 
 #include "Algo/Transform.h"
-#include "AssetRegistry/AssetData.h"
-#include "AssetRegistry/AssetRegistryModule.h"
 #include "Blueprint/BlueprintSupport.h"
 
 #if WITH_EDITOR
+#include "AssetRegistry/AssetData.h"
+#include "AssetRegistry/AssetRegistryModule.h"
 #include "CollectionManagerModule.h"
 #include "CollectionManagerTypes.h"
 #include "ICollectionContainer.h"
@@ -53,16 +53,7 @@ bool FPCGGetAssetListElement::ExecuteInternal(FPCGContext* Context) const
 		return true;
 	}
 
-	FCollectionManagerModule* CollectionManagerPtr = FModuleManager::GetModulePtr<FCollectionManagerModule>("CollectionManager");
-	if (!CollectionManagerPtr)
-	{
-		PCGLog::LogErrorOnGraph(LOCTEXT("CollectionManagerModuleNotLoaded", "The collection manager module isn't loaded."), Context);
-		return true;
-	}
-
-	bool bSuccessfullyRetrievedAssets = false;
-
-	if (Settings->bGetAssetsFromFolderPath)
+	if (Settings->AssetListSource == EPCGAssetListSource::Folder)
 	{
 		if (!AssetRegistryModule->Get().GetAssetsByPath(FName(*Settings->Directory.Path), Assets, /*bRecursive=*/true))
 		{
@@ -72,8 +63,15 @@ bool FPCGGetAssetListElement::ExecuteInternal(FPCGContext* Context) const
 			}
 		}
 	}
-	else
+	else if(Settings->AssetListSource == EPCGAssetListSource::Collection)
 	{
+		FCollectionManagerModule* CollectionManagerPtr = FModuleManager::GetModulePtr<FCollectionManagerModule>("CollectionManager");
+		if (!CollectionManagerPtr)
+		{
+			PCGLog::LogErrorOnGraph(LOCTEXT("CollectionManagerModuleNotLoaded", "The collection manager module isn't loaded."), Context);
+			return true;
+		}
+
 		TSharedPtr<ICollectionContainer> CollectionContainer = CollectionManagerPtr->Get().GetProjectCollectionContainer();
 		if (CollectionContainer)
 		{
@@ -99,6 +97,11 @@ bool FPCGGetAssetListElement::ExecuteInternal(FPCGContext* Context) const
 			return true;
 		}
 	}
+	else
+	{
+		PCGLog::LogErrorOnGraph(LOCTEXT("InvalidAssetListSource", "Invalid asset list source."), Context);
+		return true;
+	}
 
 	UPCGParamData* ParamData = FPCGContext::NewObject_AnyThread<UPCGParamData>(Context);
 	check(ParamData && ParamData->Metadata);
diff --git a/Engine/Plugins/PCG/Source/PCG/Public/Elements/IO/PCGGetAssetList.h b/Engine/Plugins/PCG/Source/PCG/Public/Elements/IO/PCGGetAssetList.h
--- a/Engine/Plugins/PCG/Source/PCG/Public/Elements/IO/PCGGetAssetList.h
+++ b/Engine/Plugins/PCG/Source/PCG/Public/Elements/IO/PCGGetAssetList.h
@@ -7,6 +7,13 @@
 
 #include "PCGGetAssetList.generated.h"
 
+UENUM(BlueprintType)
+enum class EPCGAssetListSource : uint8
+{
+	Folder UMETA(Tooltip="Retrieves all the assets from the provided folder."),
+	Collection UMETA(Tooltip="Retrieves all the assets from the provided collection name.")
+};
+
 class FPCGGetAssetListElement : public IPCGElement
 {
 public:
@@ -26,7 +33,7 @@ public:
 #if WITH_EDITOR
 	virtual FName GetDefaultNodeName() const override { return FName(TEXT("GetAssetList")); }
 	virtual FText GetDefaultNodeTitle() const override { return NSLOCTEXT("PCGGetAssetListElement", "NodeTitle", "Get Asset List"); }
-	virtual FText GetNodeTooltipText() const override { return NSLOCTEXT("PCGGetAssetListElement", "NodeTooltip", "Returns the list of asset from a collection or a folder."); }
+	virtual FText GetNodeTooltipText() const override { return NSLOCTEXT("PCGGetAssetListElement", "NodeTooltip", "Returns the list of asset, with options (class, bp generated class, etc.) from a source - collection or folder."); }
 	virtual EPCGSettingsType GetType() const override { return EPCGSettingsType::InputOutput; }
 #endif // WITH_EDITOR
 
@@ -40,14 +47,14 @@ protected:
 public:
 	/** Controls whether we will retrieve the files from a specified directory or from a collection. */
 	UPROPERTY(BlueprintReadWrite, EditAnywhere, Category = Settings, meta = (PCG_Overridable))
-	bool bGetAssetsFromFolderPath = true;
+	EPCGAssetListSource AssetListSource = EPCGAssetListSource::Folder;
 
 	/** Directory path to parse for assets. In 'long package name' format, e.g. '/Game/' and so on. */
-	UPROPERTY(BlueprintReadWrite, EditAnywhere, Category = Settings, meta = (LongPackageName, EditCondition = "bGetAssetsFromFolderPath", EditConditionHides, PCG_Overridable))
+	UPROPERTY(BlueprintReadWrite, EditAnywhere, Category = Settings, meta = (LongPackageName, EditCondition = "AssetListSource == EPCGAssetListSource::Folder", EditConditionHides, PCG_Overridable))
 	FDirectoryPath Directory;
 
 	/** Name of the collection to parse for assets. */
-	UPROPERTY(BlueprintReadWrite, EditAnywhere, Category = Settings, meta = (EditCondition = "!bGetAssetsFromFolderPath", EditConditionHides, PCG_Overridable))
+	UPROPERTY(BlueprintReadWrite, EditAnywhere, Category = Settings, meta = (EditCondition = "AssetListSource == EPCGAssetListSource::Collection", EditConditionHides, PCG_Overridable))
 	FName Collection;
 
 	/** Controls whether the class path will also be exported as part of this node in the 'ClassPath' attribute. */
