[PCG] Bumped to first major version

diff --git a/Engine/Plugins/PCG/PCG.uplugin b/Engine/Plugins/PCG/PCG.uplugin
--- a/Engine/Plugins/PCG/PCG.uplugin
+++ b/Engine/Plugins/PCG/PCG.uplugin
@@ -1,7 +1,7 @@
 {
 	"FileVersion": 3,
 	"Version": 8,
-	"VersionName": "0.2",
+	"VersionName": "1.0",
 	"FriendlyName": "Procedural Content Generation Framework (PCG)",
 	"Description": "Visual scripting framework for procedurally populating worlds with content in editor and/or at run-time.",
 	"Category": "Editor",
@@ -12,7 +12,7 @@
 	"SupportURL": "",
 	"EnabledByDefault": false,
 	"CanContainContent": true,
-	"IsBetaVersion": true,
+	"IsBetaVersion": false,
 	"Installed": false,
 	"Modules": [
 		{
