[PCG] Be a little more lenient on original component being valid

#rnx

diff --git a/Engine/Plugins/PCG/Source/PCG/Private/PCGComponent.cpp b/Engine/Plugins/PCG/Source/PCG/Private/PCGComponent.cpp
--- a/Engine/Plugins/PCG/Source/PCG/Private/PCGComponent.cpp
+++ b/Engine/Plugins/PCG/Source/PCG/Private/PCGComponent.cpp
@@ -4444,7 +4444,14 @@ FTransform FPCGComponentExecutionState::GetTransform() const
 
 FTransform FPCGComponentExecutionState::GetOriginalTransform() const
 {
-	return GetOriginalSource()->GetExecutionState().GetTransform();
+	if (IPCGGraphExecutionSource* OriginalSource = GetOriginalSource())
+	{
+		return OriginalSource->GetExecutionState().GetTransform();
+	}
+	else
+	{
+		return GetTransform();
+	}
 }
 
 UWorld* FPCGComponentExecutionState::GetWorld() const
