[PCG] Exposing editor viewport for subclassing

#rnx

diff --git a/Engine/Plugins/PCG/Source/PCGEditor/Private/Widgets/AssetEditorViewport/PCGEditorViewportToolbarSections.h b/Engine/Plugins/PCG/Source/PCGEditor/Private/Widgets/AssetEditorViewport/PCGEditorViewportToolbarSections.h
deleted file mode 100644
diff --git a/Engine/Plugins/PCG/Source/PCGEditor/Private/Widgets/AssetEditorViewport/SPCGEditorViewport.cpp b/Engine/Plugins/PCG/Source/PCGEditor/Private/Widgets/AssetEditorViewport/SPCGEditorViewport.cpp
--- a/Engine/Plugins/PCG/Source/PCGEditor/Private/Widgets/AssetEditorViewport/SPCGEditorViewport.cpp
+++ b/Engine/Plugins/PCG/Source/PCGEditor/Private/Widgets/AssetEditorViewport/SPCGEditorViewport.cpp
@@ -198,6 +198,8 @@ void SPCGEditorViewport::SetupScene(const TArray<UObject*>& InResources, const F
 
 		EditorViewportClient->Invalidate();
 	}
+
+	OnSetupScene();
 }
 
 void SPCGEditorViewport::ResetScene()
@@ -229,6 +231,8 @@ void SPCGEditorViewport::ResetScene()
 	{
 		OnFocusViewportToSelection();
 	}
+
+	OnResetScene();
 }
 
 TSharedRef<FEditorViewportClient> SPCGEditorViewport::MakeEditorViewportClient()
diff --git a/Engine/Plugins/PCG/Source/PCGEditor/Private/Widgets/AssetEditorViewport/SPCGEditorViewport.h b/Engine/Plugins/PCG/Source/PCGEditor/Private/Widgets/AssetEditorViewport/SPCGEditorViewport.h
deleted file mode 100644
diff --git a/Engine/Plugins/PCG/Source/PCGEditor/Public/PCGEditor.h b/Engine/Plugins/PCG/Source/PCGEditor/Public/PCGEditor.h
--- a/Engine/Plugins/PCG/Source/PCGEditor/Public/PCGEditor.h
+++ b/Engine/Plugins/PCG/Source/PCGEditor/Public/PCGEditor.h
@@ -240,6 +240,9 @@ protected:
 	/** Whether we can toggle enabled state of selected nodes */
 	PCGEDITOR_API virtual bool CanToggleEnabled() const;
 
+	/** Create a new viewport widget. Can be subclassed to have customization around the viewport. */
+	PCGEDITOR_API virtual TSharedRef<SPCGEditorViewport> CreateViewportWidget();
+
 private:
 	/** Register PCG specific toolbar for the editor */
 	void RegisterToolbar() const;
@@ -434,9 +437,6 @@ private:
 	/** Create a new user graph parameters tab widget */
 	TSharedRef<SPCGEditorGraphUserParametersView> CreateGraphParamsWidget();
 
-	/** Create a new viewport widget */
-	TSharedRef<SPCGEditorViewport> CreateViewportWidget();
-
 	/** Called when the component inspected is generated/cleaned */
 	void OnSourceGenerated(IPCGGraphExecutionSource* InSource);
 
diff --git a/Engine/Plugins/PCG/Source/PCGEditor/Public/Widgets/AssetEditorViewport/PCGEditorViewportToolbarSections.h b/Engine/Plugins/PCG/Source/PCGEditor/Public/Widgets/AssetEditorViewport/PCGEditorViewportToolbarSections.h
new file mode 100644
--- /dev/null
+++ b/Engine/Plugins/PCG/Source/PCGEditor/Public/Widgets/AssetEditorViewport/PCGEditorViewportToolbarSections.h
@@ -0,0 +1,10 @@
+// Copyright Epic Games, Inc. All Rights Reserved.
+
+#pragma once
+
+#include "ToolMenuEntry.h"
+
+namespace UE::PCGEditor
+{
+	PCGEDITOR_API FToolMenuEntry CreatePCGViewportSubmenu();
+}
diff --git a/Engine/Plugins/PCG/Source/PCGEditor/Public/Widgets/AssetEditorViewport/SPCGEditorViewport.h b/Engine/Plugins/PCG/Source/PCGEditor/Public/Widgets/AssetEditorViewport/SPCGEditorViewport.h
new file mode 100644
--- /dev/null
+++ b/Engine/Plugins/PCG/Source/PCGEditor/Public/Widgets/AssetEditorViewport/SPCGEditorViewport.h
@@ -0,0 +1,67 @@
+// Copyright Epic Games, Inc. All Rights Reserved.
+
+#pragma once
+
+#include "PCGDataVisualization.h"
+
+#include "SAssetEditorViewport.h"
+#include "SCommonEditorViewportToolbarBase.h"
+
+class FAdvancedPreviewScene;
+class FPCGEditorViewportClient;
+
+class SPCGEditorViewport : public SAssetEditorViewport, public ICommonEditorViewportToolbarInfoProvider
+{
+public:
+	SLATE_BEGIN_ARGS(SPCGEditorViewport) {}
+	SLATE_END_ARGS()
+
+	PCGEDITOR_API SPCGEditorViewport();
+	PCGEDITOR_API virtual ~SPCGEditorViewport();
+
+	PCGEDITOR_API void Construct(const FArguments& InArgs);
+
+	//~ Begin ICommonEditorViewportToolbarInfoProvider Interface
+	PCGEDITOR_API virtual TSharedRef<class SEditorViewport> GetViewportWidget() override;
+	PCGEDITOR_API virtual TSharedPtr<FExtender> GetExtenders() const override;
+	virtual void OnFloatingButtonClicked() override {}
+	//~ End ICommonEditorViewportToolbarInfoProvider Interface
+
+	/** Rebuilds the scene using the provided resources and setup callback. */
+	PCGEDITOR_API void SetupScene(const TArray<UObject*>& InResources, const FPCGSetupSceneFunc& SetupFunc);
+	PCGEDITOR_API void ResetScene();
+	FAdvancedPreviewScene* GetAdvancedPreviewScene() { return AdvancedPreviewScene.Get(); }
+
+	/** Add UObject references for GC. */
+	PCGEDITOR_API void AddReferencedObjects(FReferenceCollector& Collector);
+
+protected:
+	//~ Begin SEditorViewport Interface
+	PCGEDITOR_API virtual TSharedRef<FEditorViewportClient> MakeEditorViewportClient() override;
+	PCGEDITOR_API virtual TSharedPtr<SWidget> BuildViewportToolbar() override;
+	PCGEDITOR_API virtual TSharedPtr<IPreviewProfileController> CreatePreviewProfileController() override;
+	PCGEDITOR_API virtual void BindCommands() override;
+	PCGEDITOR_API virtual void OnFocusViewportToSelection() override;
+	//~ End SEditorViewport Interface
+
+	PCGEDITOR_API void ReleaseManagedResources();
+
+	PCGEDITOR_API void ToggleAutoFocusViewport();
+	PCGEDITOR_API bool IsAutoFocusViewportChecked() const;
+
+	virtual void OnSetupScene() {};
+	virtual void OnResetScene() {};
+
+protected:
+	TSharedPtr<FPCGEditorViewportClient> EditorViewportClient = nullptr;
+	TSharedPtr<FAdvancedPreviewScene> AdvancedPreviewScene = nullptr;
+
+	/** Objects used by the scene. Must be managed for GC. */
+	TArray<TObjectPtr<UObject>> ManagedResources;
+
+	/** Used when focusing the data viewport to the visualization. */
+	FBoxSphereBounds FocusBounds = FBoxSphereBounds(EForceInit::ForceInit);
+
+	/** Used when focusing the data viewport to the visualization for orthographic views. */
+	float FocusOrthoZoom = DEFAULT_ORTHOZOOM;
+};
