Fix engine crash when deleting a ProceduralVegetation asset while its PVE editor is open.

- Adds support for separating the UPCGGraph from the currently edited asset in FPCGEditor
- Correctly set the ProceduralVegetation object as the currently edited object in FPVEditor

diff --git a/Engine/Plugins/Experimental/ProceduralVegetationEditor/Source/ProceduralVegetationEditor/Private/PVEditor.cpp b/Engine/Plugins/Experimental/ProceduralVegetationEditor/Source/ProceduralVegetationEditor/Private/PVEditor.cpp
--- a/Engine/Plugins/Experimental/ProceduralVegetationEditor/Source/ProceduralVegetationEditor/Private/PVEditor.cpp
+++ b/Engine/Plugins/Experimental/ProceduralVegetationEditor/Source/ProceduralVegetationEditor/Private/PVEditor.cpp
@@ -53,9 +53,7 @@ void FPVEditor::Initialize(const EToolkitMode::Type InMode, const TSharedPtr<cla
 		CollectionSpreadSheetWidget = CreateCollectionSpreadSheetWidget();
 	}
 	
-	FPCGEditor::Initialize(InMode, InToolkitHost, ProceduralVegetationGraph);
-
-	AddEditingObject(InProceduralVegetation);
+	FPCGEditor::Initialize(InMode, InToolkitHost, ProceduralVegetationGraph, InProceduralVegetation);
 
 	check(GEngine);
 	
diff --git a/Engine/Plugins/PCG/Source/PCGEditor/Private/PCGEditor.cpp b/Engine/Plugins/PCG/Source/PCGEditor/Private/PCGEditor.cpp
--- a/Engine/Plugins/PCG/Source/PCGEditor/Private/PCGEditor.cpp
+++ b/Engine/Plugins/PCG/Source/PCGEditor/Private/PCGEditor.cpp
@@ -151,8 +151,10 @@ UPCGEditorGraph* FPCGEditor::GetPCGEditorGraph(const UPCGSettings* InSettings)
 	return GetPCGEditorGraph(PCGNode);
 }
 
-void FPCGEditor::Initialize(const EToolkitMode::Type InMode, const TSharedPtr<class IToolkitHost>& InToolkitHost, UPCGGraph* InPCGGraph)
+void FPCGEditor::Initialize(const EToolkitMode::Type InMode, const TSharedPtr<class IToolkitHost>& InToolkitHost, UPCGGraph* InPCGGraph, UObject* InObjectToEdit)
 {
+	InObjectToEdit = InObjectToEdit != nullptr ? InObjectToEdit : InPCGGraph;
+
 	InspectionDataManager.OnInspectedStackChangedDelegate.AddSP(SharedThis(this), &FPCGEditor::UpdateAfterInspectedStackChanged);
 
 	PCGGraphBeingEdited = InPCGGraph;
@@ -210,7 +212,7 @@ void FPCGEditor::Initialize(const EToolkitMode::Type InMode, const TSharedPtr<cl
 
 	const FName PCGGraphEditorAppName = FName(TEXT("PCGEditorApp"));
 
-	InitAssetEditor(InMode, InToolkitHost, PCGGraphEditorAppName, StandaloneDefaultLayout, /*bCreateDefaultStandaloneMenu=*/ true, /*bCreateDefaultToolbar=*/ true, InPCGGraph);
+	InitAssetEditor(InMode, InToolkitHost, PCGGraphEditorAppName, StandaloneDefaultLayout, /*bCreateDefaultStandaloneMenu=*/ true, /*bCreateDefaultToolbar=*/ true, InObjectToEdit);
 
 	PCGGraphBeingEdited->OnGraphChangedDelegate.AddRaw(this, &FPCGEditor::OnGraphChanged);
 	PCGGraphBeingEdited->OnNodeSourceCompiledDelegate.AddRaw(this, &FPCGEditor::OnNodeSourceCompiled);
diff --git a/Engine/Plugins/PCG/Source/PCGEditor/Public/PCGEditor.h b/Engine/Plugins/PCG/Source/PCGEditor/Public/PCGEditor.h
--- a/Engine/Plugins/PCG/Source/PCGEditor/Public/PCGEditor.h
+++ b/Engine/Plugins/PCG/Source/PCGEditor/Public/PCGEditor.h
@@ -81,7 +81,7 @@ class FPCGEditor : public FAssetEditorToolkit, public FGCObject, public FSelfReg
 {
 public:
 	/** Edits the specified PCGGraph */
-	PCGEDITOR_API void Initialize(const EToolkitMode::Type InMode, const TSharedPtr<class IToolkitHost>& InToolkitHost, UPCGGraph* InPCGGraph);
+	PCGEDITOR_API void Initialize(const EToolkitMode::Type InMode, const TSharedPtr<class IToolkitHost>& InToolkitHost, UPCGGraph* InPCGGraph, UObject* InObjectToEdit = nullptr);
 
 	/** Get the PCG editor graph being edited. */
 	PCGEDITOR_API UPCGEditorGraph* GetPCGEditorGraph();
