Pack FPCGCacheEntryKey to reduce memory waste

diff --git a/Engine/Plugins/PCG/Source/PCG/Private/Graph/PCGGraphCache.h b/Engine/Plugins/PCG/Source/PCG/Private/Graph/PCGGraphCache.h
--- a/Engine/Plugins/PCG/Source/PCG/Private/Graph/PCGGraphCache.h
+++ b/Engine/Plugins/PCG/Source/PCG/Private/Graph/PCGGraphCache.h
@@ -31,18 +31,18 @@ struct FPCGCacheEntryKey
 {
 	FPCGCacheEntryKey(const IPCGElement* InElement, const FPCGCrc& InCrc)
 		: Element(InElement)
-		, Crc(InCrc)
+		, CrcHash(InCrc.GetValue())
 	{
 		check(Element);
-		ensure(Crc.IsValid());
+		ensure(InCrc.IsValid());
 
 		Hash = GetTypeHash(Element);
-		Hash = HashCombine(Hash, Crc.GetValue());
+		Hash = HashCombine(Hash, CrcHash);
 	}
 
 	bool operator==(const FPCGCacheEntryKey& Other) const
 	{
-		return Element == Other.Element && Crc == Other.Crc;
+		return Element == Other.Element && CrcHash == Other.CrcHash;
 	}
 
 	friend uint32 GetTypeHash(const FPCGCacheEntryKey& Key)
@@ -54,8 +54,7 @@ struct FPCGCacheEntryKey
 
 private:
 	const IPCGElement* Element = nullptr;
-	FPCGCrc Crc;
-
+	uint32 CrcHash = 0;
 	uint32 Hash = 0;
 };
 
