[PCG] Fix Grammar nodes not working with data domain

* Duplicate Cross Section was losing all attributes when doing a Create Surface, now write to the data domain by default for new nodes, versioning for old nodes.
* Fix SelectGrammar or Primitive Cross Section
* Fix issue in the ApplyMultiAccessor that was not working if the first key was on a data domain and the second key was on a element domain.

#rnx

diff --git a/Engine/Plugins/PCG/Source/PCG/Private/Elements/Grammar/PCGDuplicateCrossSections.cpp b/Engine/Plugins/PCG/Source/PCG/Private/Elements/Grammar/PCGDuplicateCrossSections.cpp
--- a/Engine/Plugins/PCG/Source/PCG/Private/Elements/Grammar/PCGDuplicateCrossSections.cpp
+++ b/Engine/Plugins/PCG/Source/PCG/Private/Elements/Grammar/PCGDuplicateCrossSections.cpp
@@ -10,10 +10,25 @@
 #include "Metadata/Accessors/PCGAttributeAccessorHelpers.h"
 #include "Metadata/Accessors/PCGAttributeAccessorKeys.h"
 
+#include "UObject/UE5ReleaseStreamObjectVersion.h"
+
 #include UE_INLINE_GENERATED_CPP_BY_NAME(PCGDuplicateCrossSections)
 
 #define LOCTEXT_NAMESPACE "PCGDuplicateCrossSectionsElement"
 
+void UPCGDuplicateCrossSectionsSettings::PostLoad()
+{
+	Super::PostLoad();
+
+#if WITH_EDITOR
+	if (GetLinkerCustomVersion(FUE5ReleaseStreamObjectVersion::GUID) < FUE5ReleaseStreamObjectVersion::ExtraOutputAttributesOnDataDomainPCG)
+	{
+		// For all previous nodes, we'll force this option to false. for retro-compatibility
+		bExtraOutputAttributesOnDataDomain = false;
+	}
+#endif // WITH_EDITOR
+}
+
 #if WITH_EDITOR
 FName UPCGDuplicateCrossSectionsSettings::GetDefaultNodeName() const
 {
@@ -174,26 +189,34 @@ bool FPCGDuplicateCrossSectionsElement::ExecuteInternal(FPCGContext* InContext)
 			NewSplineData->Initialize(NewSpline);
 			NewSplineData->InitializeFromData(InputSplineData);
 
-			NewSplineData->Metadata->FindOrCreateAttribute<FName>(Settings->SymbolAttributeName, Symbol, false, false);
+			check(NewSplineData->MutableMetadata());
+
+			FPCGMetadataDomain* MetadataDomain = Settings->bExtraOutputAttributesOnDataDomain
+				? NewSplineData->MutableMetadata()->GetMetadataDomain(PCGMetadataDomainID::Data)
+				: NewSplineData->MutableMetadata()->GetMetadataDomain(PCGMetadataDomainID::Elements);
+
+			check(MetadataDomain);
+
+			MetadataDomain->FindOrCreateAttribute<FName>(Settings->SymbolAttributeName, Symbol, false, false);
 
 			if (!bMatchAndSetAttributes && Settings->bOutputDebugColorAttribute)
 			{
-				NewSplineData->Metadata->FindOrCreateAttribute<FVector4>(Settings->DebugColorAttributeName, FVector4(CurrentBlock.DebugColor, 1.0), false, false);
+				MetadataDomain->FindOrCreateAttribute<FVector4>(Settings->DebugColorAttributeName, FVector4(CurrentBlock.DebugColor, 1.0), false, false);
 			}
 
 			if (Settings->bOutputSizeAttribute)
 			{
-				NewSplineData->Metadata->FindOrCreateAttribute<FVector>(Settings->SizeAttributeName, Size, false, false);
+				MetadataDomain->FindOrCreateAttribute<FVector>(Settings->SizeAttributeName, Size, false, false);
 			}
 
 			if (!bMatchAndSetAttributes && Settings->bOutputScalableAttribute)
 			{
-				NewSplineData->Metadata->FindOrCreateAttribute<bool>(Settings->ScalableAttributeName, CurrentBlock.bScalable, false, false);
+				MetadataDomain->FindOrCreateAttribute<bool>(Settings->ScalableAttributeName, CurrentBlock.bScalable, false, false);
 			}
 
 			if (Settings->bOutputSplineIndexAttribute)
 			{
-				NewSplineData->Metadata->FindOrCreateAttribute<int32>(Settings->SplineIndexAttributeName, SplineIndex++, false, false);
+				MetadataDomain->FindOrCreateAttribute<int32>(Settings->SplineIndexAttributeName, SplineIndex++, false, false);
 			}
 
 			Outputs.Add_GetRef(Input).Data = NewSplineData;
@@ -203,7 +226,7 @@ bool FPCGDuplicateCrossSectionsElement::ExecuteInternal(FPCGContext* InContext)
 
 	if (bMatchAndSetAttributes)
 	{
-		MatchAndSetAttributes(Inputs, Outputs, ModuleInfoParamData, Settings);
+		MatchAndSetAttributes(Inputs, Outputs, ModuleInfoParamData, Settings, Settings->bExtraOutputAttributesOnDataDomain ? PCGMetadataDomainID::Data : PCGMetadataDomainID::Elements);
 	}
 
 	return true;
diff --git a/Engine/Plugins/PCG/Source/PCG/Private/Elements/Grammar/PCGSelectGrammar.cpp b/Engine/Plugins/PCG/Source/PCG/Private/Elements/Grammar/PCGSelectGrammar.cpp
--- a/Engine/Plugins/PCG/Source/PCG/Private/Elements/Grammar/PCGSelectGrammar.cpp
+++ b/Engine/Plugins/PCG/Source/PCG/Private/Elements/Grammar/PCGSelectGrammar.cpp
@@ -374,8 +374,9 @@ bool FPCGSelectGrammarElement::ExecuteInternal(FPCGContext* InContext) const
 		{
 			FPCGAttributePropertyInputSelector ComparedValueSelector = Settings->ComparedValueAttribute.CopyAndFixLast(InputPointData);
 			TUniquePtr<const IPCGAttributeAccessor> ComparedValuesAccessor = PCGAttributeAccessorHelpers::CreateConstAccessor(InputPointData, ComparedValueSelector);
+			TUniquePtr<const IPCGAttributeAccessorKeys> ComparedValuesKeys = PCGAttributeAccessorHelpers::CreateConstKeys(InputPointData, ComparedValueSelector);
 
-			if (!ComparedValuesAccessor)
+			if (!ComparedValuesAccessor || !ComparedValuesKeys)
 			{
 				PCGLog::Metadata::LogFailToCreateAccessorError(ComparedValueSelector, InContext);
 				continue;
@@ -383,7 +384,8 @@ bool FPCGSelectGrammarElement::ExecuteInternal(FPCGContext* InContext) const
 
 			// TODO: For now, cast everything to double for numerical comparison, but this could easily expand in the future
 			TArray<const IPCGAttributeAccessor*> Accessors = {InputKeyAccessor.Get(), ComparedValuesAccessor.Get()};
-			PCGMetadataElementCommon::ApplyOnMultiAccessorsRange<FName, double>(*InputKeys, Accessors, ProcessWithCriteria);
+			TArray<const IPCGAttributeAccessorKeys*> Keys = {InputKeys.Get(), ComparedValuesKeys.Get()};
+			PCGMetadataElementCommon::ApplyOnMultiAccessorsRange<FName, double>(Keys, Accessors, ProcessWithCriteria);
 		}
 	}
 
diff --git a/Engine/Plugins/PCG/Source/PCG/Private/Elements/Grammar/PCGSubdivideSegment.cpp b/Engine/Plugins/PCG/Source/PCG/Private/Elements/Grammar/PCGSubdivideSegment.cpp
--- a/Engine/Plugins/PCG/Source/PCG/Private/Elements/Grammar/PCGSubdivideSegment.cpp
+++ b/Engine/Plugins/PCG/Source/PCG/Private/Elements/Grammar/PCGSubdivideSegment.cpp
@@ -256,7 +256,9 @@ bool FPCGSegmentSubdivisionElement::ExecuteInternal(FPCGContext* InContext) cons
 		TUniquePtr<const IPCGAttributeAccessor> GrammarAccessor;
 		TUniquePtr<const IPCGAttributeAccessor> FlipAxisAccessor;
 		TUniquePtr<const IPCGAttributeAccessor> SeedAccessor;
-		TUniquePtr<const IPCGAttributeAccessorKeys> Keys = PCGAttributeAccessorHelpers::CreateConstKeys(InputPointData, FPCGAttributePropertyOutputSelector());
+		TUniquePtr<const IPCGAttributeAccessorKeys> GrammarKeys;
+		TUniquePtr<const IPCGAttributeAccessorKeys> FlipAxisKeys;
+		TUniquePtr<const IPCGAttributeAccessorKeys> SeedKeys;
 
 
 		const bool bShouldUseAccessors = Settings->GrammarSelection.bGrammarAsAttribute || Settings->bFlipAxisAsAttribute || Settings->bUseSeedAttribute;
@@ -265,14 +267,15 @@ bool FPCGSegmentSubdivisionElement::ExecuteInternal(FPCGContext* InContext) cons
 		if (bShouldUseAccessors)
 		{
 			// Create accessors for grammar, flip axis and seed as soon as either is needed. For the ones that aren't needed, we will create a constant accessor to uniformize the code below, even though there is a small performance cost.
-			auto CreateAccessor = [InputPointData, InContext]<typename T>(TUniquePtr<const IPCGAttributeAccessor>&OutAccessor, bool bValueAsAttribute, const FPCGAttributePropertyInputSelector & InputSelector, const T& ConstantValue) -> bool
+			auto CreateAccessorAndKeys = [InputPointData, InContext]<typename T>(TUniquePtr<const IPCGAttributeAccessor>&OutAccessor, TUniquePtr<const IPCGAttributeAccessorKeys>& OutKeys, bool bValueAsAttribute, const FPCGAttributePropertyInputSelector & InputSelector, const T& ConstantValue) -> bool
 			{
 				if (bValueAsAttribute)
 				{
 					const FPCGAttributePropertyInputSelector Selector = InputSelector.CopyAndFixLast(InputPointData);
 					OutAccessor = PCGAttributeAccessorHelpers::CreateConstAccessor(InputPointData, Selector);
+					OutKeys = PCGAttributeAccessorHelpers::CreateConstKeys(InputPointData, Selector);
 
-					if (!OutAccessor)
+					if (!OutAccessor || !OutKeys)
 					{
 						PCGLog::Metadata::LogFailToCreateAccessorError(Selector, InContext);
 						return false;
@@ -281,14 +284,16 @@ bool FPCGSegmentSubdivisionElement::ExecuteInternal(FPCGContext* InContext) cons
 				else
 				{
 					OutAccessor = MakeUnique<FPCGConstantValueAccessor<T>>(ConstantValue);
+					// Dummy single key
+					OutKeys = MakeUnique<FPCGAttributeAccessorKeysSingleObjectPtr<void>>();
 				}
 
 				return true;
 			};
 
-			if (!CreateAccessor(GrammarAccessor, Settings->GrammarSelection.bGrammarAsAttribute, Settings->GrammarSelection.GrammarAttribute, Settings->GrammarSelection.GrammarString)
-				|| !CreateAccessor(FlipAxisAccessor, Settings->bFlipAxisAsAttribute, Settings->FlipAxisAttribute, Settings->bShouldFlipAxis)
-				|| !CreateAccessor(SeedAccessor, Settings->bUseSeedAttribute, Settings->SeedAttribute, DefaultAdditionalSeed))
+			if (!CreateAccessorAndKeys(GrammarAccessor, GrammarKeys, Settings->GrammarSelection.bGrammarAsAttribute, Settings->GrammarSelection.GrammarAttribute, Settings->GrammarSelection.GrammarString)
+				|| !CreateAccessorAndKeys(FlipAxisAccessor, FlipAxisKeys, Settings->bFlipAxisAsAttribute, Settings->FlipAxisAttribute, Settings->bShouldFlipAxis)
+				|| !CreateAccessorAndKeys(SeedAccessor, SeedKeys, Settings->bUseSeedAttribute, Settings->SeedAttribute, DefaultAdditionalSeed))
 			{
 				continue;
 			}
@@ -338,7 +343,7 @@ bool FPCGSegmentSubdivisionElement::ExecuteInternal(FPCGContext* InContext) cons
 				PCGSubdivideSegmentHelpers::Process(Parameters, InGrammar, bFlipAxis, AdditionalSeed, Index);
 			};
 
-			PCGMetadataElementCommon::ApplyOnMultiAccessors<FString, bool, int32>(*Keys, { GrammarAccessor.Get(), FlipAxisAccessor.Get(), SeedAccessor.Get()}, Process);
+			PCGMetadataElementCommon::ApplyOnMultiAccessors<FString, bool, int32>({GrammarKeys.Get(), FlipAxisKeys.Get(), SeedKeys.Get()}, { GrammarAccessor.Get(), FlipAxisAccessor.Get(), SeedAccessor.Get()}, Process);
 		}
 		else
 		{
diff --git a/Engine/Plugins/PCG/Source/PCG/Private/Elements/Grammar/PCGSubdivisionBase.cpp b/Engine/Plugins/PCG/Source/PCG/Private/Elements/Grammar/PCGSubdivisionBase.cpp
--- a/Engine/Plugins/PCG/Source/PCG/Private/Elements/Grammar/PCGSubdivisionBase.cpp
+++ b/Engine/Plugins/PCG/Source/PCG/Private/Elements/Grammar/PCGSubdivisionBase.cpp
@@ -317,7 +317,7 @@ TMap<FString, PCGGrammar::FTokenizedGrammar> FPCGSubdivisionBaseElement::GetToke
 	return Result;
 }
 
-bool FPCGSubdivisionBaseElement::MatchAndSetAttributes(const TArray<FPCGTaggedData>& InputData, TArray<FPCGTaggedData>& OutputData, const UPCGParamData* InModuleInfoParamData, const UPCGSubdivisionBaseSettings* InSettings) const
+bool FPCGSubdivisionBaseElement::MatchAndSetAttributes(const TArray<FPCGTaggedData>& InputData, TArray<FPCGTaggedData>& OutputData, const UPCGParamData* InModuleInfoParamData, const UPCGSubdivisionBaseSettings* InSettings, const FPCGMetadataDomainID& InTargetDomain) const
 {
 	check(InModuleInfoParamData && InModuleInfoParamData->Metadata);
 
@@ -372,10 +372,20 @@ bool FPCGSubdivisionBaseElement::MatchAndSetAttributes(const TArray<FPCGTaggedDa
 		check(OutData && OutData->MutableMetadata());
 
 		UPCGMetadata* OutputMetadata = OutData->MutableMetadata();
+		check(OutputMetadata);
+
+		FPCGMetadataDomain* OutputMetadataDomain = OutputMetadata->GetMetadataDomain(InTargetDomain);
+
+		if (!OutputMetadataDomain)
+		{
+			continue;
+		}
+
+		const bool bSupportMultipleValues = OutputMetadataDomain->SupportsMultiEntries();
 
 		// Look for the Symbol Attribute in the output metadata to query its value.
-		const FPCGMetadataAttribute<FName>* OutSymbolAttribute = OutputMetadata->GetConstTypedAttribute<FName>(InSettings->SymbolAttributeName);
-		if (!ensure(OutSymbolAttribute))
+		const FPCGMetadataAttribute<FName>* OutSymbolAttribute = OutputMetadataDomain->GetConstTypedAttribute<FName>(InSettings->SymbolAttributeName);
+		if (!OutSymbolAttribute)
 		{
 			continue;
 		}
@@ -398,20 +408,45 @@ bool FPCGSubdivisionBaseElement::MatchAndSetAttributes(const TArray<FPCGTaggedDa
 				}
 
 				// If the attribute already exists, override it
-				if (OutputMetadata->HasAttribute(AttributeName))
+				if (OutputMetadataDomain->HasAttribute(AttributeName))
 				{
-					OutputMetadata->DeleteAttribute(AttributeName);
+					OutputMetadataDomain->DeleteAttribute(AttributeName);
 				}
 
 				InAttributes.Add(InputMetadata->GetConstAttribute(AttributeName));
-				OutAttributes.Add(OutputMetadata->CopyAttribute(InAttributes.Last(), AttributeName, /*bKeepParent=*/false, /*bCopyEntries=*/false, /*bCopyValues=*/true));
+				OutAttributes.Add(OutputMetadataDomain->CopyAttribute(InAttributes.Last(), AttributeName, /*bKeepParent=*/false, /*bCopyEntries=*/false, /*bCopyValues=*/bSupportMultipleValues));
 			}
 
-			// @todo_pcg: potential optimization GetValuesFromItemKeys / SetValuesFromValueKeys
-			TPCGValueRange<int64> MetadataEntryRange = OutPointData->GetMetadataEntryValueRange(/*bAllocate=*/true);
-			for (int32 Index = 0; Index < OutPointData->GetNumPoints(); ++Index)
+			// @todo_pcg: potential optimization GetValuesFromItemKeys / SetValuesFromValueKeys. Would need support for TPCGValueRange on those.
+			if (bSupportMultipleValues)
 			{
-				const FName Symbol = OutSymbolAttribute->GetValueFromItemKey(MetadataEntryRange[Index]);
+				TPCGValueRange<int64> MetadataEntryRange = OutPointData->GetMetadataEntryValueRange(/*bAllocate=*/true);
+				for (int32 Index = 0; Index < OutPointData->GetNumPoints(); ++Index)
+				{
+					const FName Symbol = OutSymbolAttribute->GetValueFromItemKey(MetadataEntryRange[Index]);
+
+					if (!SymbolToEntryKeyMapping.Contains(Symbol))
+					{
+						continue;
+					}
+
+					const PCGMetadataEntryKey InputEntryKey = SymbolToEntryKeyMapping[Symbol];
+
+					if (InputEntryKey == PCGInvalidEntryKey)
+					{
+						continue;
+					}
+
+					for (int32 i = 0; i < InAttributes.Num(); ++i)
+					{
+						PCGMetadataValueKey InputValueKey = InAttributes[i]->GetValueKey(InputEntryKey);
+						OutAttributes[i]->SetValueFromValueKey(MetadataEntryRange[Index], InputValueKey);
+					}
+				}
+			}
+			else
+			{
+				const FName Symbol = OutSymbolAttribute->GetValueFromItemKey(PCGFirstEntryKey);
 
 				if (!SymbolToEntryKeyMapping.Contains(Symbol))
 				{
@@ -424,11 +459,15 @@ bool FPCGSubdivisionBaseElement::MatchAndSetAttributes(const TArray<FPCGTaggedDa
 				{
 					continue;
 				}
+				
+				if (OutputMetadataDomain->GetItemCountForChild() == 0)
+				{
+					OutputMetadataDomain->AddEntry();
+				}
 
 				for (int32 i = 0; i < InAttributes.Num(); ++i)
 				{
-					PCGMetadataValueKey InputValueKey = InAttributes[i]->GetValueKey(InputEntryKey);
-					OutAttributes[i]->SetValueFromValueKey(MetadataEntryRange[Index], InputValueKey);
+					OutAttributes[i]->SetValue(PCGFirstEntryKey, InAttributes[i], InputEntryKey);
 				}
 			}
 		}
@@ -439,18 +478,18 @@ bool FPCGSubdivisionBaseElement::MatchAndSetAttributes(const TArray<FPCGTaggedDa
 			const FName Symbol = OutSymbolAttribute->GetValue(PCGDefaultValueKey);
 			const PCGMetadataEntryKey InputEntryKey = SymbolToEntryKeyMapping[Symbol];
 
-			auto AddAttribute = [InputMetadata, OutputMetadata, InputEntryKey]<typename T>(T Dummy, const FName AttributeName)
+			auto AddAttribute = [InputMetadata, OutputMetadataDomain, InputEntryKey]<typename T>(T Dummy, const FName AttributeName)
 			{
 				const FPCGMetadataAttribute<T>* InAttribute = InputMetadata->GetConstTypedAttribute<T>(AttributeName);
 				check(InAttribute);
 
 				// Overwrite any existing attribute
-				if (OutputMetadata->HasAttribute(AttributeName))
+				if (OutputMetadataDomain->HasAttribute(AttributeName))
 				{
-					OutputMetadata->DeleteAttribute(AttributeName);
+					OutputMetadataDomain->DeleteAttribute(AttributeName);
 				}
 
-				OutputMetadata->CreateAttribute<T>(AttributeName, InAttribute->GetValueFromItemKey(InputEntryKey), InAttribute->AllowsInterpolation(), /*bOverrideParent=*/true);
+				OutputMetadataDomain->CreateAttribute<T>(AttributeName, InAttribute->GetValueFromItemKey(InputEntryKey), InAttribute->AllowsInterpolation(), /*bOverrideParent=*/true);
 			};
 
 			for (int32 i = 0; i < AttributeNames.Num(); ++i)
@@ -462,7 +501,7 @@ bool FPCGSubdivisionBaseElement::MatchAndSetAttributes(const TArray<FPCGTaggedDa
 				}
 
 				// Skip any attribute that already exists in the output data
-				if (OutputMetadata->HasAttribute(AttributeNames[i]))
+				if (OutputMetadataDomain->HasAttribute(AttributeNames[i]))
 				{
 					continue;
 				}
diff --git a/Engine/Plugins/PCG/Source/PCG/Private/PCGSettings.cpp b/Engine/Plugins/PCG/Source/PCG/Private/PCGSettings.cpp
--- a/Engine/Plugins/PCG/Source/PCG/Private/PCGSettings.cpp
+++ b/Engine/Plugins/PCG/Source/PCG/Private/PCGSettings.cpp
@@ -22,6 +22,7 @@
 #include "UObject/FortniteMainBranchObjectVersion.h"
 #include "UObject/FortniteReleaseBranchCustomObjectVersion.h"
 #include "UObject/ObjectSaveContext.h"
+#include "UObject/UE5ReleaseStreamObjectVersion.h"
 
 #include UE_INLINE_GENERATED_CPP_BY_NAME(PCGSettings)
 
@@ -350,6 +351,7 @@ void UPCGSettings::Serialize(FArchive& Ar)
 	Ar.UsingCustomVersion(FPCGCustomVersion::GUID);
 	Ar.UsingCustomVersion(FFortniteMainBranchObjectVersion::GUID);
 	Ar.UsingCustomVersion(FFortniteReleaseBranchCustomObjectVersion::GUID);
+	Ar.UsingCustomVersion(FUE5ReleaseStreamObjectVersion::GUID);
 
 #if WITH_EDITOR
 	if (Ar.IsLoading())
diff --git a/Engine/Plugins/PCG/Source/PCG/Public/Elements/Grammar/PCGDuplicateCrossSections.h b/Engine/Plugins/PCG/Source/PCG/Public/Elements/Grammar/PCGDuplicateCrossSections.h
--- a/Engine/Plugins/PCG/Source/PCG/Public/Elements/Grammar/PCGDuplicateCrossSections.h
+++ b/Engine/Plugins/PCG/Source/PCG/Public/Elements/Grammar/PCGDuplicateCrossSections.h
@@ -14,6 +14,10 @@ class UPCGDuplicateCrossSectionsSettings : public UPCGSubdivisionBaseSettings
 	GENERATED_BODY()
 
 public:
+	//~Begin UObject interface
+	void PostLoad() override;
+	//~End UObject interface
+	
 	//~Begin UPCGSettings interface
 #if WITH_EDITOR
 	virtual FName GetDefaultNodeName() const override;
@@ -38,6 +42,10 @@ public:
 	UPROPERTY(BlueprintReadWrite, EditAnywhere, Category = Settings, meta = (EditCondition = "bExtrudeVectorAsAttribute", EditConditionHides, DisplayAfter = bExtrudeVectorAsAttribute, PCG_Overridable))
 	FPCGAttributePropertyInputSelector ExtrudeVectorAttribute;
 
+	/** Choose if the Extra Output Attributes are added to the Data domain or the Elements domain. If ForwardAttributesFromModulesInfo is true, module info attributes will also be pushed to the data domain. */
+	UPROPERTY(BlueprintReadWrite, EditAnywhere, Category = "Settings|Extra Output Attributes", meta = (PCG_Overridable))
+	bool bExtraOutputAttributesOnDataDomain = true;
+
 	UPROPERTY(BlueprintReadWrite, EditAnywhere, Category = "Settings|Extra Output Attributes", meta = (PCG_Overridable, InlineEditConditionToggle))
 	bool bOutputSplineIndexAttribute = true;
 
diff --git a/Engine/Plugins/PCG/Source/PCG/Public/Elements/Grammar/PCGSubdivisionBase.h b/Engine/Plugins/PCG/Source/PCG/Public/Elements/Grammar/PCGSubdivisionBase.h
--- a/Engine/Plugins/PCG/Source/PCG/Public/Elements/Grammar/PCGSubdivisionBase.h
+++ b/Engine/Plugins/PCG/Source/PCG/Public/Elements/Grammar/PCGSubdivisionBase.h
@@ -459,5 +459,5 @@ protected:
 	FModuleInfoMap GetModulesInfoMap(FPCGContext* InContext, const UPCGSubdivisionBaseSettings* InSettings, const UPCGParamData*& OutModuleInfoParamData) const;
 	PCGGrammar::FTokenizedGrammar GetTokenizedGrammar(FPCGContext* InContext, const UPCGData* InputData, const UPCGSubdivisionBaseSettings* InSettings, const FModuleInfoMap& InModulesInfo, double& OutMinSize) const;
 	TMap<FString, PCGGrammar::FTokenizedGrammar> GetTokenizedGrammarForPoints(FPCGContext* InContext, const UPCGBasePointData* InputData, const UPCGSubdivisionBaseSettings* InSettings, const FModuleInfoMap& InModulesInfo, double& OutMinSize) const;
-	bool MatchAndSetAttributes(const TArray<FPCGTaggedData>& InputData, TArray<FPCGTaggedData>& OutputData, const UPCGParamData* InModuleInfoParamData, const UPCGSubdivisionBaseSettings* InSettings) const;
+	bool MatchAndSetAttributes(const TArray<FPCGTaggedData>& InputData, TArray<FPCGTaggedData>& OutputData, const UPCGParamData* InModuleInfoParamData, const UPCGSubdivisionBaseSettings* InSettings, const FPCGMetadataDomainID& InTargetDomain = PCGMetadataDomainID::Default) const;
 };
diff --git a/Engine/Plugins/PCG/Source/PCG/Public/Elements/Metadata/PCGMetadataElementCommon.h b/Engine/Plugins/PCG/Source/PCG/Public/Elements/Metadata/PCGMetadataElementCommon.h
--- a/Engine/Plugins/PCG/Source/PCG/Public/Elements/Metadata/PCGMetadataElementCommon.h
+++ b/Engine/Plugins/PCG/Source/PCG/Public/Elements/Metadata/PCGMetadataElementCommon.h
@@ -180,7 +180,7 @@ namespace PCGMetadataElementCommon
 
 	/**
 	* Iterate over the full range of the keys (if Count is negative, otherwise, as many times as Count), calling the callback with values ranges get from the multiple accessors.
-	* There can be 1 Keys for all accessors or 1 per accessor, but in that case if Count is set to -1, it will use the first key for the number of elements to process.
+	* There can be 1 Keys for all accessors or 1 per accessor, but in that case if Count is set to -1, it will use the first key in the list with more than 1 element for the number of elements to process.
 	* ChunkSize influence the max number of values to get in one go with GetRange.
 	* Callback should have a signature: void(const TArrayView<T>&... Args, int32 Start, int32 Range). Can also return a bool to know if we need to continue.
 	* Return false if it processes nothing.
@@ -198,14 +198,25 @@ namespace PCGMetadataElementCommon
 		}
 
 		// Also make sure that the keys are matching in num (or just have a single entry)
-		if (!ensure(Algo::AllOf(MultiKeys, [FirstKeys=MultiKeys[0]](const IPCGAttributeAccessorKeys* Keys){ return Keys->GetNum() == FirstKeys->GetNum() || Keys->GetNum() == 1; })))
+		int32 NumElements = 1;
+		for (const IPCGAttributeAccessorKeys* Keys : MultiKeys)
 		{
-			return false;
+			const int32 KeysNum = Keys->GetNum();
+			if (NumElements == 1)
+			{
+				NumElements = KeysNum;
+			}
+
+			if (NumElements != KeysNum && KeysNum != 1)
+			{
+				PCGLog::LogErrorOnGraph(NSLOCTEXT("PCGApplyOnMultiAccessorsRange", "InvalidKeyCardinality", "Try to use ApplyOnMultiAccessorsRange with keys nums that doesn't match."));
+				return false;
+			}
 		}
 
-		const int32 NumberOfEntries = Count < 0 ? MultiKeys[0]->GetNum() : Count;
+		const int32 NumberOfEntries = Count < 0 ? NumElements : Count;
 
-		if (NumberOfEntries == 0)
+		if (NumElements == 0 || NumberOfEntries == 0)
 		{
 			return false;
 		}
diff --git a/Engine/Plugins/PCGInterops/PCGGeometryScriptInterop/Source/PCGGeometryScriptInterop/Private/Elements/PCGPrimitiveCrossSection.cpp b/Engine/Plugins/PCGInterops/PCGGeometryScriptInterop/Source/PCGGeometryScriptInterop/Private/Elements/PCGPrimitiveCrossSection.cpp
--- a/Engine/Plugins/PCGInterops/PCGGeometryScriptInterop/Source/PCGGeometryScriptInterop/Private/Elements/PCGPrimitiveCrossSection.cpp
+++ b/Engine/Plugins/PCGInterops/PCGGeometryScriptInterop/Source/PCGGeometryScriptInterop/Private/Elements/PCGPrimitiveCrossSection.cpp
@@ -421,6 +421,9 @@ bool FPCGPrimitiveCrossSectionElement::ExecuteInternal(FPCGContext* Context) con
 			}
 		}
 
+		bool bMetadataDomainWarned = false;
+		bool bAttributeFailedWarned = false;
+
 		TArray<FSplinePoint> SplinePoints;
 		// Create the splines from the tiers
 		for (const FCrossSection& CrossSection : CrossSections)
@@ -443,8 +446,29 @@ bool FPCGPrimitiveCrossSectionElement::ExecuteInternal(FPCGContext* Context) con
 			OutSplineData->Initialize(SplinePoints, true, FTransform::Identity);
 			OutputData.Data = OutSplineData;
 
-			check(OutSplineData->Metadata);
-			OutSplineData->Metadata->CreateAttribute(Settings->ExtrusionVectorAttribute.GetName(), SliceDirection * CrossSection.Height, /*bAllowsInterpolation=*/false, /*bOverridesParent=*/true);
+			UPCGMetadata* Metadata = OutSplineData->MutableMetadata();
+			check(Metadata);
+
+			const FPCGMetadataDomainID MetadataDomainID = OutSplineData->GetMetadataDomainIDFromSelector(Settings->ExtrusionVectorAttribute);
+			FPCGMetadataDomain* MetadataDomain = Metadata->GetMetadataDomain(MetadataDomainID);
+
+			const FName AttributeName = Settings->ExtrusionVectorAttribute.GetName();
+
+			if (MetadataDomain)
+			{
+				const FPCGMetadataAttributeBase* Attribute = MetadataDomain->CreateAttribute<FVector>(AttributeName, SliceDirection * CrossSection.Height, /*bAllowsInterpolation=*/false, /*bOverridesParent=*/true);
+				
+				if (Attribute == nullptr && !bAttributeFailedWarned)
+				{
+					bAttributeFailedWarned = true;
+					PCGLog::Metadata::LogFailToCreateAttributeError<FVector>(AttributeName, Context);
+				}
+			}
+			else if (!bMetadataDomainWarned)
+			{
+				bMetadataDomainWarned = true;
+				PCGLog::Metadata::LogInvalidMetadataDomain(Settings->ExtrusionVectorAttribute, Context);
+			}
 		}
 	}
 
diff --git a/Engine/Source/Runtime/Core/Public/UObject/UE5ReleaseStreamObjectVersions.inl b/Engine/Source/Runtime/Core/Public/UObject/UE5ReleaseStreamObjectVersions.inl
--- a/Engine/Source/Runtime/Core/Public/UObject/UE5ReleaseStreamObjectVersions.inl
+++ b/Engine/Source/Runtime/Core/Public/UObject/UE5ReleaseStreamObjectVersions.inl
@@ -180,5 +180,8 @@ UE5_RELEASE_VERSION(RecalculateClothAssetSerializedBounds, 56)
 // Composite plugin now uses its own derived scene capture components
 UE5_RELEASE_VERSION(CompositePluginDerivedSceneCaptures, 57)
 
+// Add option to output attributes on the PCG Duplicate Cross Section node on the Data domain
+UE5_RELEASE_VERSION(ExtraOutputAttributesOnDataDomainPCG, 58)
+
 // -----<new versions can be added above this line>-------------------------------------------------
 #undef UE5_RELEASE_VERSION
\ No newline at end of file
