[PCG] Bughawk reports

#rnx

diff --git a/Engine/Plugins/PCG/Source/PCG/Private/Compute/DataInterfaces/PCGDataLabelResolverDataInterface.cpp b/Engine/Plugins/PCG/Source/PCG/Private/Compute/DataInterfaces/PCGDataLabelResolverDataInterface.cpp
--- a/Engine/Plugins/PCG/Source/PCG/Private/Compute/DataInterfaces/PCGDataLabelResolverDataInterface.cpp
+++ b/Engine/Plugins/PCG/Source/PCG/Private/Compute/DataInterfaces/PCGDataLabelResolverDataInterface.cpp
@@ -62,7 +62,7 @@ void UPCGDataLabelResolverDataInterface::GetHLSL(FString& OutHLSL, FString const
 		"	}\n"
 		"\n"
 		"	// Data index could be -1 if the label was not resolved, so fall-back to data index 0 if necessary.\n"
-		"	return max((uint){DataInterfaceName}_InDataIdToDataIndexMap[InDataId], 0);\n"
+		"	return (uint)max({DataInterfaceName}_InDataIdToDataIndexMap[InDataId], 0);\n"
 		"}\n"
 		"\n"),
 		TemplateArgs);
diff --git a/Engine/Plugins/PCG/Source/PCG/Private/Elements/Metadata/PCGMetadataPartitionKernel.cpp b/Engine/Plugins/PCG/Source/PCG/Private/Elements/Metadata/PCGMetadataPartitionKernel.cpp
--- a/Engine/Plugins/PCG/Source/PCG/Private/Elements/Metadata/PCGMetadataPartitionKernel.cpp
+++ b/Engine/Plugins/PCG/Source/PCG/Private/Elements/Metadata/PCGMetadataPartitionKernel.cpp
@@ -154,7 +154,7 @@ TSharedPtr<const FPCGDataCollectionDesc> UPCGMetadataPartitionKernel::ComputeOut
 		}
 		else
 		{
-			UE_LOG(LogPCG, Warning, TEXT("No analysis data received by static mesh spawner kernel, worst case instance allocations will be made."));
+			UE_LOG(LogPCG, Warning, TEXT("UPCGMetadataPartitionKernel: No analysis data received."));
 
 			if (InBinding->InputDataCollection.TaggedData[AnalysisDataIndex].Data->IsA<UPCGProxyForGPUData>())
 			{
