[PCG] Disable landscape VT sampling by default

#rnx

diff --git a/Engine/Plugins/PCG/Source/PCG/Private/Data/Registry/PCGGetDataFunctionRegistry.cpp b/Engine/Plugins/PCG/Source/PCG/Private/Data/Registry/PCGGetDataFunctionRegistry.cpp
--- a/Engine/Plugins/PCG/Source/PCG/Private/Data/Registry/PCGGetDataFunctionRegistry.cpp
+++ b/Engine/Plugins/PCG/Source/PCG/Private/Data/Registry/PCGGetDataFunctionRegistry.cpp
@@ -167,6 +167,7 @@ int FPCGGetDataFunctionRegistry::DefaultDataFromActor(FPCGContext* InContext, co
 		FPCGLandscapeDataProps LandscapeDataProps;
 		LandscapeDataProps.bGetHeightOnly = false;
 		LandscapeDataProps.bGetLayerWeights = (!PCGGraph || PCGGraph->bLandscapeUsesMetadata);
+		LandscapeDataProps.bSampleVirtualTextures = false;
 
 		Data->Initialize({ LandscapeActor }, PCGHelpers::GetGridBounds(Actor, Component), LandscapeDataProps);
 
diff --git a/Engine/Plugins/PCG/Source/PCG/Private/Elements/PCGTypedGetter.cpp b/Engine/Plugins/PCG/Source/PCG/Private/Elements/PCGTypedGetter.cpp
--- a/Engine/Plugins/PCG/Source/PCG/Private/Elements/PCGTypedGetter.cpp
+++ b/Engine/Plugins/PCG/Source/PCG/Private/Elements/PCGTypedGetter.cpp
@@ -32,6 +32,8 @@ UPCGGetLandscapeSettings::UPCGGetLandscapeSettings()
 		ActorSelector.bMustOverlapSelf = true;
 		ActorSelector.ActorSelection = EPCGActorSelection::ByClass;
 		ActorSelector.ActorSelectionClass = ALandscapeProxy::StaticClass();
+
+		SamplingProperties.bSampleVirtualTextures = false;
 	}
 }
 
diff --git a/Engine/Plugins/PCG/Source/PCG/Private/PCGComponent.cpp b/Engine/Plugins/PCG/Source/PCG/Private/PCGComponent.cpp
--- a/Engine/Plugins/PCG/Source/PCG/Private/PCGComponent.cpp
+++ b/Engine/Plugins/PCG/Source/PCG/Private/PCGComponent.cpp
@@ -3462,6 +3462,7 @@ UPCGData* UPCGComponent::CreateLandscapePCGData(bool bHeightOnly) const
 	FPCGLandscapeDataProps LandscapeDataProps;
 	LandscapeDataProps.bGetHeightOnly = bHeightOnly;
 	LandscapeDataProps.bGetLayerWeights = (PCGGraph && PCGGraph->bLandscapeUsesMetadata);
+	LandscapeDataProps.bSampleVirtualTextures = false;
 
 	LandscapeData->Initialize(Landscapes, LandscapeBounds, LandscapeDataProps);
 
