[PCG] Add SanitizeNodeTitle as editor functions to update a name of a node when it changes

#rnx

diff --git a/Engine/Plugins/PCG/Source/PCG/Private/PCGGraph.cpp b/Engine/Plugins/PCG/Source/PCG/Private/PCGGraph.cpp
--- a/Engine/Plugins/PCG/Source/PCG/Private/PCGGraph.cpp
+++ b/Engine/Plugins/PCG/Source/PCG/Private/PCGGraph.cpp
@@ -928,6 +928,9 @@ UPCGNode* UPCGGraph::AddNode(UPCGSettingsInterface* InSettingsInterface)
 			// Flags added because default flags favor tick/interactive, not load-time renaming.
 			Node->Rename(*NodeName.ToString(), nullptr, REN_DontCreateRedirectors);
 		}
+
+		// Give the chance to the graph/settings to update the node title.
+		Node->SetNodeTitle(NAME_None);
 #endif
 
 		Nodes.Add(Node);
diff --git a/Engine/Plugins/PCG/Source/PCG/Private/PCGNode.cpp b/Engine/Plugins/PCG/Source/PCG/Private/PCGNode.cpp
--- a/Engine/Plugins/PCG/Source/PCG/Private/PCGNode.cpp
+++ b/Engine/Plugins/PCG/Source/PCG/Private/PCGNode.cpp
@@ -440,6 +440,39 @@ FText UPCGNode::GetNodeTooltipText() const
 		return FText::GetEmpty();
 	}
 }
+
+bool UPCGNode::SetNodeTitle(FName NewTitle, bool bApplySanitization)
+{	
+	if (NewTitle == NodeTitle)
+	{
+		return false;
+	}
+
+	if (bApplySanitization)
+	{
+		// Ask the graph and the settings to do sanitization if needed.
+		if (const UPCGGraph* Graph = GetGraph())
+		{
+			Graph->SanitizeNodeTitle(NewTitle, this);
+		}
+
+		if (const UPCGSettings* Settings = GetSettings())
+		{
+			Settings->SanitizeNodeTitle(NewTitle);
+		}
+	}
+
+	// Re-check if after the sanitization the name is the same or not.
+	if (NewTitle == NodeTitle)
+	{
+		return false;
+	}
+
+	Modify();
+	NodeTitle = NewTitle;
+
+	return true;
+}
 #endif
 
 bool UPCGNode::IsInstance() const
diff --git a/Engine/Plugins/PCG/Source/PCG/Public/PCGGraph.h b/Engine/Plugins/PCG/Source/PCG/Public/PCGGraph.h
--- a/Engine/Plugins/PCG/Source/PCG/Public/PCGGraph.h
+++ b/Engine/Plugins/PCG/Source/PCG/Public/PCGGraph.h
@@ -462,6 +462,9 @@ public:
 	UE_API bool Recompile();
 
 	UE_API void OnPCGQualityLevelChanged();
+
+	/** Overridable function for child classes to have a graph-wide node title sanitization when the title of a node changes. Returns true if the title have changed. */
+	virtual bool SanitizeNodeTitle(FName& InOutTitle, const UPCGNode* InNode) const { return false; }
 #endif
 
 #if WITH_EDITOR
diff --git a/Engine/Plugins/PCG/Source/PCG/Public/PCGNode.h b/Engine/Plugins/PCG/Source/PCG/Public/PCGNode.h
--- a/Engine/Plugins/PCG/Source/PCG/Public/PCGNode.h
+++ b/Engine/Plugins/PCG/Source/PCG/Public/PCGNode.h
@@ -91,6 +91,9 @@ public:
 #if WITH_EDITOR
 	/** Tooltip that describes node functionality and other information. */
 	UE_API FText GetNodeTooltipText() const;
+	
+	/** Set the node title. Can be sanitized by the graph/settings if asked. Returns true if the title changed. */
+	UE_API bool SetNodeTitle(FName NewTitle, bool bApplySanitization = true);
 #endif
 
 	/** Returns all the input pin properties */
diff --git a/Engine/Plugins/PCG/Source/PCG/Public/PCGSettings.h b/Engine/Plugins/PCG/Source/PCG/Public/PCGSettings.h
--- a/Engine/Plugins/PCG/Source/PCG/Public/PCGSettings.h
+++ b/Engine/Plugins/PCG/Source/PCG/Public/PCGSettings.h
@@ -351,6 +351,9 @@ public:
 	/** Returns whether the execution dependency pin can be configured to be used for culling. */
 	UFUNCTION()
 	bool CanUseExecutionDependencyPinForCulling() const;
+
+	/** Overridable function for child classes to have a node title sanitization when the title of a node changes. Return true if the name has changed. */
+	virtual bool SanitizeNodeTitle(FName& InOutName) const { return false; }
 #endif // WITH_EDITOR
 
 	/** Derived classes can implement this to expose additional information or context, such as an asset in use by the node. */
diff --git a/Engine/Plugins/PCG/Source/PCGEditor/Private/Nodes/PCGEditorGraphNode.cpp b/Engine/Plugins/PCG/Source/PCGEditor/Private/Nodes/PCGEditorGraphNode.cpp
--- a/Engine/Plugins/PCG/Source/PCGEditor/Private/Nodes/PCGEditorGraphNode.cpp
+++ b/Engine/Plugins/PCG/Source/PCGEditor/Private/Nodes/PCGEditorGraphNode.cpp
@@ -103,9 +103,10 @@ void UPCGEditorGraphNode::OnRenameNode(const FString& NewName)
 	const FName Name = FName(NewName);
 	if (PCGNode->NodeTitle != Name)
 	{
-		Modify();
-		PCGNode->Modify();
-		PCGNode->NodeTitle = Name;
+		if (PCGNode->SetNodeTitle(Name))
+		{
+			Modify();
+		}
 	}
 }
 
diff --git a/Engine/Plugins/PCG/Source/PCGEditor/Private/Nodes/PCGEditorGraphNodeGetUserParameter.cpp b/Engine/Plugins/PCG/Source/PCGEditor/Private/Nodes/PCGEditorGraphNodeGetUserParameter.cpp
--- a/Engine/Plugins/PCG/Source/PCGEditor/Private/Nodes/PCGEditorGraphNodeGetUserParameter.cpp
+++ b/Engine/Plugins/PCG/Source/PCGEditor/Private/Nodes/PCGEditorGraphNodeGetUserParameter.cpp
@@ -63,7 +63,7 @@ void UPCGEditorGraphGetUserParameter::OnRenameNode(const FString& NewName)
 		}
 
 		const FString NodeTitle = GetNodeTitle(ENodeTitleType::FullTitle).ToString();
-		PCGNode->NodeTitle = FName(NodeTitle);
+		PCGNode->SetNodeTitle(FName(NodeTitle));
 		Super::OnRenameNode(NodeTitle);
 	}
 }
diff --git a/Engine/Plugins/PCG/Source/PCGEditor/Private/PCGEditorGraphSchemaActions.cpp b/Engine/Plugins/PCG/Source/PCGEditor/Private/PCGEditorGraphSchemaActions.cpp
--- a/Engine/Plugins/PCG/Source/PCGEditor/Private/PCGEditorGraphSchemaActions.cpp
+++ b/Engine/Plugins/PCG/Source/PCGEditor/Private/PCGEditorGraphSchemaActions.cpp
@@ -175,7 +175,7 @@ void FPCGEditorGraphSchemaAction_NewGetParameterElement::PostCreation(UPCGNode*
 	Settings->PropertyGuid = PropertyDesc.ID;
 	Settings->PropertyName = PropertyDesc.Name;
 
-	NewNode->NodeTitle = Settings->PropertyName;
+	NewNode->SetNodeTitle(Settings->PropertyName, /*bApplySanitization=*/false);
 
 	// We need to set the settings to update the pins.
 	NewNode->SetSettingsInterface(Settings);
@@ -813,7 +813,7 @@ UEdGraphNode* FPCGEditorGraphSchemaAction_NewNamedRerouteUsage::PerformAction(cl
 	NewPCGNode->UpdateAfterSettingsChangeDuringCreation();
 
 	CastChecked<UPCGNamedRerouteUsageSettings>(DefaultNodeSettings)->Declaration = Declaration;
-	NewPCGNode->NodeTitle = DeclarationNode->NodeTitle;
+	NewPCGNode->SetNodeTitle(DeclarationNode->NodeTitle, /*bApplySanitization=*/false);
 	NewPCGNode->NodeTitleColor = DeclarationNode->NodeTitleColor;
 
 	// Create edge from the declaration to the new node
