[PCG] Fix Data Layer + Level Instance crash when resolving DataLayerInstances
- Make sure to get the data layer instances from the level and not the owning world.
- For persistent world actors this will keep the same behavior but for sub-level actors the proper DataLayerManager will get used to resolve the DataLayerInstances

diff --git a/Engine/Plugins/PCG/Source/PCG/Private/Helpers/PCGDataLayerHelpers.cpp b/Engine/Plugins/PCG/Source/PCG/Private/Helpers/PCGDataLayerHelpers.cpp
--- a/Engine/Plugins/PCG/Source/PCG/Private/Helpers/PCGDataLayerHelpers.cpp
+++ b/Engine/Plugins/PCG/Source/PCG/Private/Helpers/PCGDataLayerHelpers.cpp
@@ -61,7 +61,7 @@ namespace PCGDataLayerHelpers
 			return DataLayerInstances;
 		}
 
-		TArray<const UDataLayerInstance*> GetDataLayerInstancesFromDataLayerReferences(FPCGContext* Context, const FPCGDataLayerSettings& DataLayerSettings, UWorld* World)
+		TArray<const UDataLayerInstance*> GetDataLayerInstancesFromDataLayerReferences(FPCGContext* Context, const FPCGDataLayerSettings& DataLayerSettings, ULevel* Level)
 		{
 			const TArray<FPCGTaggedData> DataLayersInputs = Context->InputData.GetInputsByPin(PCGDataLayerHelpers::Constants::DataLayerReferenceAttribute);
 
@@ -89,7 +89,7 @@ namespace PCGDataLayerHelpers
 			TArray<TSoftObjectPtr<UDataLayerAsset>> DataLayerAssetPtrs;
 			Algo::Transform(DataLayerSoftObjectPaths, DataLayerAssetPtrs, [](const FSoftObjectPath& DataLayerPath) { return TSoftObjectPtr<UDataLayerAsset>(DataLayerPath); });
 
-			return GetDataLayerInstancesFromAssets(Context, DataLayerAssetPtrs, UDataLayerManager::GetDataLayerManager(World));
+			return GetDataLayerInstancesFromAssets(Context, DataLayerAssetPtrs, UDataLayerManager::GetDataLayerManager(Level));
 		}
 
 		void FilterDataLayerInstances(FPCGContext* Context, const FPCGDataLayerSettings& DataLayerSettings, TArray<const UDataLayerInstance*>& DataLayerInstances)
@@ -111,10 +111,10 @@ namespace PCGDataLayerHelpers
 			DataLayerInstances.SetNum(RemoveIndex);
 		}
 
-		void AddDataLayerInstances(FPCGContext* Context, const FPCGDataLayerSettings& DataLayerSettings, UWorld* World, TArray<const UDataLayerInstance*>& DataLayerInstances)
+		void AddDataLayerInstances(FPCGContext* Context, const FPCGDataLayerSettings& DataLayerSettings, ULevel* Level, TArray<const UDataLayerInstance*>& DataLayerInstances)
 		{
 			TArray<TSoftObjectPtr<UDataLayerAsset>> AddDataLayerAssets = GetDataLayerAssetsFromInput(Context, PCGDataLayerHelpers::Constants::AddDataLayersAttribute, DataLayerSettings.AddDataLayers);
-			TArray<const UDataLayerInstance*> AddDataLayerInstances = GetDataLayerInstancesFromAssets(Context, AddDataLayerAssets, UDataLayerManager::GetDataLayerManager(World));
+			TArray<const UDataLayerInstance*> AddDataLayerInstances = GetDataLayerInstancesFromAssets(Context, AddDataLayerAssets, UDataLayerManager::GetDataLayerManager(Level));
 
 			for (const UDataLayerInstance* AddDataLayerInstance : AddDataLayerInstances)
 			{
@@ -210,19 +210,19 @@ namespace PCGDataLayerHelpers
 		{
 		case EPCGDataLayerSource::Self:
 		{
-			DataLayerInstances = DefaultDataLayerSource->GetDataLayerInstances();
+			DataLayerInstances = DefaultDataLayerSource->GetDataLayerInstancesForLevel();
 			break;
 		}
 		case EPCGDataLayerSource::DataLayerReferences:
 		{
-			DataLayerInstances = Private::GetDataLayerInstancesFromDataLayerReferences(Context, DataLayerSettings, DefaultDataLayerSource->GetWorld());
+			DataLayerInstances = Private::GetDataLayerInstancesFromDataLayerReferences(Context, DataLayerSettings, DefaultDataLayerSource->GetLevel());
 			break;
 		}
 		}
 
 		Private::FilterDataLayerInstances(Context, DataLayerSettings, DataLayerInstances);
 
-		Private::AddDataLayerInstances(Context, DataLayerSettings, DefaultDataLayerSource->GetWorld(), DataLayerInstances);
+		Private::AddDataLayerInstances(Context, DataLayerSettings, DefaultDataLayerSource->GetLevel(), DataLayerInstances);
 
 		// Sort for determinism
 		DataLayerInstances.Sort([](const UDataLayerInstance& A, const UDataLayerInstance& B)
