[PCG] GetExecutionContext: Add IsGameWorld query

diff --git a/Engine/Plugins/PCG/Source/PCG/Private/Elements/PCGGetExecutionContext.cpp b/Engine/Plugins/PCG/Source/PCG/Private/Elements/PCGGetExecutionContext.cpp
--- a/Engine/Plugins/PCG/Source/PCG/Private/Elements/PCGGetExecutionContext.cpp
+++ b/Engine/Plugins/PCG/Source/PCG/Private/Elements/PCGGetExecutionContext.cpp
@@ -119,6 +119,10 @@ bool FPCGGetExecutionContextElement::ExecuteInternal(FPCGContext* Context) const
 			Value = PCGHelpers::IsInWorldPartitionBuilderCommandletContext();
 #endif
 		}
+		else if (Mode == EPCGGetExecutionContextMode::IsGameWorld)
+		{
+			Value = SupportingWorld && SupportingWorld->IsGameWorld();
+		}
 
 		FPCGMetadataAttribute<bool>* Attribute = ParamData->Metadata->CreateAttribute<bool>(PCGGetExecutionContextConstants::AttributeName, Value, /*bAllowInterpolation=*/false, /*bOverrideParent=*/false);
 	}
diff --git a/Engine/Plugins/PCG/Source/PCG/Public/Elements/PCGGetExecutionContext.h b/Engine/Plugins/PCG/Source/PCG/Public/Elements/PCGGetExecutionContext.h
--- a/Engine/Plugins/PCG/Source/PCG/Public/Elements/PCGGetExecutionContext.h
+++ b/Engine/Plugins/PCG/Source/PCG/Public/Elements/PCGGetExecutionContext.h
@@ -10,17 +10,18 @@
 UENUM(BlueprintType)
 enum class EPCGGetExecutionContextMode : uint8
 {
-	IsEditor UMETA(Tooltip="Returns whether the execution happens in the editor (not in PIE)."),
-	IsRuntime UMETA(Tooltip="Returns whether the execution happens in a runtime environment (PIE or cooked build) or not."),
+	IsEditor UMETA(Tooltip="Returns whether the execution happens in the editor (not in PIE and not runtime generation (including editor viewport generation)."),
+	IsRuntime UMETA(Tooltip="Returns whether the execution happens in a runtime environment (PIE or cooked build or any runtime generation (including editor viewport generation)."),
 	IsOriginal UMETA(Tooltip="Returns whether the executing context is the original context (e.g. not partitioned)."),
 	IsLocal UMETA(Tooltip="Returns whether the executing context is partitioned (in opposition to original)."),
 	IsPartitioned UMETA(Tooltip = "Returns whether the executing context is the original context and that it is partitioned."),
-	IsRuntimeGeneration UMETA(Tooltip="Returns whether the executing context is doing runtime gneration."),
+	IsRuntimeGeneration UMETA(Tooltip="Returns whether the executing context is doing runtime generation."),
 	IsDedicatedServer UMETA(Tooltip="Returns whether the executing context is running on a dedicated server."),
 	IsListenServer UMETA(Tooltip="Returns whether the executing context is running on a listen server."),
 	HasAuthority UMETA(Tooltip="Returns whether the execution context is running with network authority."),
 	IsBuilder UMETA(Tooltip="Returns whether the execution source is generated by the PCG World Partition Builder. Always false in non-editor."),
-	TrackingPriority UMETA(Tooltip = "Gives tracking priority of the graph instance (lower values are unaffected by changes triggered by higher values). Always 0 in non-editor.")
+	TrackingPriority UMETA(Tooltip = "Gives tracking priority of the graph instance (lower values are unaffected by changes triggered by higher values). Always 0 in non-editor."),
+	IsGameWorld UMETA(Tooltip="Returns whether the execution happens in a game world.")
 };
 
 /** Returns some context-specific information in the form of an attribute set. */
