[PCG] GPU: Add texel size functions to texture data interface.

#rnx

diff --git a/Engine/Plugins/PCG/Shaders/Private/PCGTextureDataInterface.ush b/Engine/Plugins/PCG/Shaders/Private/PCGTextureDataInterface.ush
--- a/Engine/Plugins/PCG/Shaders/Private/PCGTextureDataInterface.ush
+++ b/Engine/Plugins/PCG/Shaders/Private/PCGTextureDataInterface.ush
@@ -44,8 +44,8 @@ float4 {DataInterfaceName}_TextureBounds[PCG_MAX_NUM_SRV_BINDINGS];
  *   - uint32 - SliceIndex (aka TextureIndex if bound to a TextureArray, 0 otherwise.)
  *   - uint32 - Dimension
  *   - uint32 - Sampler
- *   - uint32 - SizeX
- *   - uint32 - SizeY
+ *   - uint32 - ResolutionX
+ *   - uint32 - ResolutionY
  * - TextureInfo1
  * - TextureInfo2
  * - ...
@@ -74,6 +74,18 @@ void GetTextureInfoInternal_{DataInterfaceName}(uint InDataIndex, out uint OutBi
 	}
 }
 
+void GetTextureResolutionInternal_{DataInterfaceName}(uint InDataIndex, out uint2 OutResolution)
+{
+	if (InDataIndex >= 0u && InDataIndex < {DataInterfaceName}_NumTextureInfos)
+	{
+		OutResolution = {DataInterfaceName}_TextureInfos[InDataIndex * 2 + 1].xy;
+	}
+	else
+	{
+		OutResolution = (uint2)0;
+	}
+}
+
 // #################### GENERAL HELPERS ##########################
 
 uint GetNumData_{DataInterfaceName}()
@@ -136,6 +148,24 @@ float2 GetTexCoords_{DataInterfaceName}(uint InDataIndex, float3 InWorldPos)
 	return (InWorldPos.xy - Min) / Extents;
 }
 
+float2 GetTexelSize_{DataInterfaceName}(uint InDataIndex)
+{
+	uint2 Resolution;
+	GetTextureResolutionInternal_{DataInterfaceName}(InDataIndex, Resolution);
+
+	return 1.0f / float2(Resolution);
+}
+
+float2 GetTexelSizeWorld_{DataInterfaceName}(uint InDataIndex)
+{
+	const float2 Min = {DataInterfaceName}_TextureBounds[InDataIndex].xy;
+	const float2 Max = {DataInterfaceName}_TextureBounds[InDataIndex].zw;
+
+	uint2 Resolution;
+	GetTextureResolutionInternal_{DataInterfaceName}(InDataIndex, Resolution);
+	return (Max - Min) / float2(Resolution);
+}
+
 #define PCG_CONDITIONAL_TEXTURE_SAMPLE(Index) \
 	if (BindingIndex == Index) \
 	{ \
diff --git a/Engine/Plugins/PCG/Source/PCG/Private/Compute/DataInterfaces/PCGTextureDataInterface.cpp b/Engine/Plugins/PCG/Source/PCG/Private/Compute/DataInterfaces/PCGTextureDataInterface.cpp
--- a/Engine/Plugins/PCG/Source/PCG/Private/Compute/DataInterfaces/PCGTextureDataInterface.cpp
+++ b/Engine/Plugins/PCG/Source/PCG/Private/Compute/DataInterfaces/PCGTextureDataInterface.cpp
@@ -37,7 +37,7 @@ namespace PCGTextureDataInterfaceConstants
 
 namespace PCGTextureDataInterfaceHelpers
 {
-	void GetSupportedShaderFunctions(TArray<FShaderFunctionDefinition>& OutFunctions)
+	void GetSharedFunctions(TArray<FShaderFunctionDefinition>& OutFunctions)
 	{
 		OutFunctions.AddDefaulted_GetRef()
 			.SetName(TEXT("GetNumData"))
@@ -54,6 +54,22 @@ namespace PCGTextureDataInterfaceHelpers
 			.AddParam(EShaderFundamentalType::Uint) // InThreadIndex
 			.AddParam(EShaderFundamentalType::Uint, 0, 0, EShaderParamModifier::Out) // OutDataIndex
 			.AddParam(EShaderFundamentalType::Uint, 2, 0, EShaderParamModifier::Out); // OutElementIndex
+
+		OutFunctions.AddDefaulted_GetRef()
+			.SetName(TEXT("GetTexCoords"))
+			.AddReturnType(EShaderFundamentalType::Float, 2)
+			.AddParam(EShaderFundamentalType::Uint) // InDataIndex
+			.AddParam(EShaderFundamentalType::Float, 3); // InWorldPos
+
+		OutFunctions.AddDefaulted_GetRef()
+			.SetName(TEXT("GetTexelSize"))
+			.AddReturnType(EShaderFundamentalType::Float, 2)
+			.AddParam(EShaderFundamentalType::Uint); // InDataIndex
+
+		OutFunctions.AddDefaulted_GetRef()
+			.SetName(TEXT("GetTexelSizeWorld"))
+			.AddReturnType(EShaderFundamentalType::Float, 2)
+			.AddParam(EShaderFundamentalType::Uint); // InDataIndex
 	}
 }
 
@@ -111,13 +127,7 @@ void UPCGTextureDataInterface::GetSupportedInputs(TArray<FShaderFunctionDefiniti
 {
 	// TODO: Support texture3d, etc. UNiagaraDataInterfaceRenderTarget2D and related are useful references.
 
-	PCGTextureDataInterfaceHelpers::GetSupportedShaderFunctions(OutFunctions);
-
-	OutFunctions.AddDefaulted_GetRef()
-		.SetName(TEXT("GetTexCoords"))
-		.AddReturnType(EShaderFundamentalType::Float, 2)
-		.AddParam(EShaderFundamentalType::Uint) // InDataIndex
-		.AddParam(EShaderFundamentalType::Float, 3); // InWorldPos
+	PCGTextureDataInterfaceHelpers::GetSharedFunctions(OutFunctions);
 
 	OutFunctions.AddDefaulted_GetRef()
 		.SetName(TEXT("Sample"))
@@ -167,7 +177,7 @@ void UPCGTextureDataInterface::GetSupportedInputs(TArray<FShaderFunctionDefiniti
 
 void UPCGTextureDataInterface::GetSupportedOutputs(TArray<FShaderFunctionDefinition>& OutFunctions) const
 {
-	PCGTextureDataInterfaceHelpers::GetSupportedShaderFunctions(OutFunctions);
+	PCGTextureDataInterfaceHelpers::GetSharedFunctions(OutFunctions);
 
 	OutFunctions.AddDefaulted_GetRef()
 		.SetName(TEXT("Store"))
@@ -665,8 +675,8 @@ void FPCGTextureDataProviderProxy::PackTextureInfos(FRDGBuilder& GraphBuilder)
 	//   - uint32 - SliceIndex (aka TextureIndex if bound to a TextureArray, 0 otherwise.)
 	//   - uint32 - Dimension
 	//   - uint32 - Sampler
-	//   - uint32 - SizeX
-	//   - uint32 - SizeY
+	//   - uint32 - ResolutionX
+	//   - uint32 - ResolutionY
 	// - TextureInfo1
 	// - TextureInfo2
 	// - ...
diff --git a/Engine/Plugins/PCG/Source/PCG/Private/Compute/Elements/PCGCustomHLSL.cpp b/Engine/Plugins/PCG/Source/PCG/Private/Compute/Elements/PCGCustomHLSL.cpp
--- a/Engine/Plugins/PCG/Source/PCG/Private/Compute/Elements/PCGCustomHLSL.cpp
+++ b/Engine/Plugins/PCG/Source/PCG/Private/Compute/Elements/PCGCustomHLSL.cpp
@@ -540,7 +540,9 @@ void UPCGCustomHLSLSettings::UpdateInputDeclarations()
 			"float4 {0}_Sample(uint DataIndex, float2 TexCoords);\n"
 			"// Computes sample coordinates of the WorldPos relative to the texture data's bounds.\n"
 			"float4 {0}_SampleWorldPos(uint DataIndex, float3 WorldPos);\n"
-			"float4 {0}_Load(uint DataIndex, uint2 ElementIndex);\n"),
+			"float4 {0}_Load(uint DataIndex, uint2 ElementIndex);\n"
+			"float2 {0}_GetTexelSize(uint DataIndex);\n"
+			"float2 {0}_GetTexelSizeWorld(uint DataIndex);\n"),
 			{ bMultiPin ? PCGHLSLElement::PinDeclTemplateStr : TextureDataPins[0] });
 
 		InputDeclarations += TEXT("\n");
@@ -729,6 +731,9 @@ void UPCGCustomHLSLSettings::UpdateOutputDeclarations()
 		OutputDeclarations += FString::Format(TEXT(
 			"uint {0}_GetNumData();\n"
 			"uint2 {0}_GetNumElements(uint DataIndex);\n"
+			"float2 {0}_GetTexCoords(uint DataIndex, float3 WorldPos);\n"
+			"float2 {0}_GetTexelSize(uint DataIndex);\n"
+			"float2 {0}_GetTexelSizeWorld(uint DataIndex);\n"
 			"void {0}_Store(uint DataIndex, uint2 ElementIndex, float4 Value);\n"),
 			{ bMultiPin ? PCGHLSLElement::PinDeclTemplateStr : TextureDataPins[0] });
 
