[PCG] Remove Graph Description Json for 5.7, as it is no longer necessary

#rnx

diff --git a/Engine/Plugins/PCG/Source/PCG/Private/PCGGraph.cpp b/Engine/Plugins/PCG/Source/PCG/Private/PCGGraph.cpp
--- a/Engine/Plugins/PCG/Source/PCG/Private/PCGGraph.cpp
+++ b/Engine/Plugins/PCG/Source/PCG/Private/PCGGraph.cpp
@@ -29,11 +29,7 @@
 #include "Dialogs/Dialogs.h"
 #include "EdGraph/EdGraphPin.h"
 #include "GameFramework/Actor.h"
-#include "Serialization/JsonSerializer.h"
-#include "Serialization/JsonWriter.h"
 #include "UObject/ObjectSaveContext.h"
-#include "Helpers/IO/PCGJsonHelpers.h"
-#include "Helpers/IO/PCGIOHelpers.h"
 #else
 #include "UObject/Package.h"
 #endif
@@ -2248,100 +2244,6 @@ EPropertyBagAlterationResult UPCGGraph::AddUserParameters(const TArray<FProperty
 }
 
 #if WITH_EDITOR
-FString UPCGGraph::GetGraphDescriptionJson() const
-{
-	using namespace PCG::IO::Json;
-	TSharedPtr<FJsonObject> RootJsonObject = MakeShared<FJsonObject>();
-
-	// Set up the Json header to include custom versions
-	TSharedPtr<FJsonObject> JsonHeaderObject = MakeShared<FJsonObject>();
-	Helpers::AppendHeader(JsonHeaderObject, PCG::IO::Constants::Attribute::FCustomExportVersion::LatestVersion);
-	RootJsonObject->SetObjectField(LOCTEXT("JsonHeaderField", "Header").ToString(), MoveTemp(JsonHeaderObject));
-
-	// 0. Graph information
-	// - Name
-	RootJsonObject->SetStringField("Name", GetFName().ToString());
-	// - Category (might identify family of graph
-	RootJsonObject->SetStringField("Category", Category.ToString());
-	// - Description
-	RootJsonObject->SetStringField("Description", Description.ToString());
-
-	auto CreatePinDescriptionJsonObject = [](const FPCGPinProperties& Pin)
-	{
-		TSharedPtr<FJsonObject> JsonObject = MakeShared<FJsonObject>();
-		JsonObject->SetStringField(TEXT("Type"), Pin.AllowedTypes.ToString());
-		JsonObject->SetBoolField(TEXT("Multiple"), Pin.bAllowMultipleData);
-		return JsonObject;
-	};
-
-	// 1. Inputs
-	TSharedPtr<FJsonObject> InputsJsonObject = MakeShared<FJsonObject>();
-	TArray<TSharedPtr<FJsonObject>> InputsJsonObjects;
-
-	TArray<FPCGPinProperties> InputPins = GetInputNode()->InputPinProperties();
-	for (const FPCGPinProperties& InputPin : InputPins)
-	{
-		InputsJsonObject->SetObjectField(InputPin.Label.ToString(), CreatePinDescriptionJsonObject(InputPin));
-	}
-
-	RootJsonObject->SetObjectField(TEXT("Inputs"), MoveTemp(InputsJsonObject));
-
-	// 2. Output
-	TSharedPtr<FJsonObject> OutputsJsonObject = MakeShared<FJsonObject>();
-
-	TArray<FPCGPinProperties> OutputPins = GetOutputNode()->OutputPinProperties();
-	for (const FPCGPinProperties& OutputPin : OutputPins)
-	{
-		OutputsJsonObject->SetObjectField(OutputPin.Label.ToString(), CreatePinDescriptionJsonObject(OutputPin));
-	}
-
-	RootJsonObject->SetObjectField(TEXT("Outputs"), MoveTemp(OutputsJsonObject));
-
-	// 3. Parameters
-	TSharedPtr<FJsonObject> ParametersJsonObject = MakeShared<FJsonObject>();
-
-	const FInstancedPropertyBag* Parameters = GetUserParametersStruct();
-	if (const UPropertyBag* PropertyBag = Parameters ? Parameters->GetPropertyBagStruct() : nullptr)
-	{
-		TConstArrayView<FPropertyBagPropertyDesc> PropertyDescs = PropertyBag->GetPropertyDescs();
-		for (const FPropertyBagPropertyDesc& PropertyDesc : PropertyDescs)
-		{
-			TSharedPtr<FJsonObject> JsonObject = MakeShared<FJsonObject>();
-			// - Type
-			JsonObject->SetStringField(TEXT("Type"), StaticEnum<EPropertyBagPropertyType>()->GetDisplayValueAsText(PropertyDesc.ValueType).ToString());
-
-			// - Description
-			if (PropertyDesc.CachedProperty->HasMetaData(TEXT("ToolTip")))
-			{
-				JsonObject->SetStringField(TEXT("Description"), PropertyDesc.CachedProperty->GetMetaData(TEXT("ToolTip")));
-			}
-
-			// - Value
-			auto SetJson = [&JsonObject]<typename T>(const T& Value) -> bool
-			{
-				Helpers::SetValue<T>(JsonObject, TEXT("value"), Value);
-				return true;
-			};
-
-			PCGPropertyHelpers::GetPropertyValueWithCallback(Parameters->GetValue().GetMemory(), PropertyDesc.CachedProperty, SetJson);
-			ParametersJsonObject->SetObjectField(PropertyDesc.Name.ToString(), MoveTemp(JsonObject));
-		}
-	}
-
-	RootJsonObject->SetObjectField(TEXT("Parameters"), ParametersJsonObject);
-
-	FString JsonString;
-	const TSharedRef<TJsonWriter<>> JsonWriter = TJsonWriterFactory<>::Create(&JsonString);
-
-	if (!FJsonSerializer::Serialize(RootJsonObject.ToSharedRef(), JsonWriter))
-	{
-		UE_LOG(LogPCG, Error, TEXT("Serialization of graph '%s' Json failed."), *GetFName().ToString());
-		return {};
-	}
-
-	return JsonString;
-}
-
 void UPCGGraph::SetHiddenFlagInputNode(bool bHidden)
 {
 	if (InputNode && InputNode->bHidden != bHidden)
diff --git a/Engine/Plugins/PCG/Source/PCG/Public/PCGGraph.h b/Engine/Plugins/PCG/Source/PCG/Public/PCGGraph.h
--- a/Engine/Plugins/PCG/Source/PCG/Public/PCGGraph.h
+++ b/Engine/Plugins/PCG/Source/PCG/Public/PCGGraph.h
@@ -478,12 +478,6 @@ public:
 	TObjectPtr<UPCGGraphCompilationData> GetCookedCompilationData() { return CookedCompilationData; }
 	const TObjectPtr<UPCGGraphCompilationData> GetCookedCompilationData() const { return CookedCompilationData; }
 
-#if WITH_EDITOR
-	/** Outputs a Json representation of the graph inputs, outputs, and parameters. See the PCGExportSelectedAttributes node for more examples. */
-	UFUNCTION(BlueprintCallable, Category = "Graph|Advanced|Json")
-	PCG_API FString GetGraphDescriptionJson() const;
-#endif
-
 protected:
 #if WITH_EDITOR
 	UFUNCTION()
