[PCG] Fix AActor::Rename crash: Handle UWorld::ReInitWorld (Editor only non-gameworld)
- Runtime Partition Actors will not get deleted when UWorld::ReInitWorld gets called so we need to find them when initializing the new PCGSubsystem

diff --git a/Engine/Plugins/PCG/Source/PCG/Private/RuntimeGen/PCGRuntimeGenScheduler.cpp b/Engine/Plugins/PCG/Source/PCG/Private/RuntimeGen/PCGRuntimeGenScheduler.cpp
--- a/Engine/Plugins/PCG/Source/PCG/Private/RuntimeGen/PCGRuntimeGenScheduler.cpp
+++ b/Engine/Plugins/PCG/Source/PCG/Private/RuntimeGen/PCGRuntimeGenScheduler.cpp
@@ -201,6 +201,39 @@ FPCGRuntimeGenScheduler::FPCGRuntimeGenScheduler(UWorld* InWorld, FPCGActorAndCo
 	FramesUntilGeneration = PCGRuntimeGenSchedulerHelpers::CVarFramesBeforeFirstGenerate.GetValueOnGameThread();
 
 	FLevelStreamingDelegates::OnLevelStreamingStateChanged.AddRaw(this, &FPCGRuntimeGenScheduler::OnLevelStreamingStateChanged);
+
+#if WITH_EDITOR
+	// Handle UWorld::ReInitWorld
+	// - Actors will get unregistered and re-registered in a new subsystem, thus creating a new RGS
+	// - We need to find them and properly re-register them (Pooled and in-use runtime PAs)
+	TArray<APCGPartitionActor*> ExistingPoolPartitionActors;
+	int32 NumPartitionActors = 0;
+
+	UPCGActorHelpers::ForEachActorInLevel<APCGPartitionActor>(World->PersistentLevel, [&NumPartitionActors, &ExistingPoolPartitionActors](AActor* InActor)
+	{
+		if (APCGPartitionActor* InPartitionActor = Cast<APCGPartitionActor>(InActor); InPartitionActor && InPartitionActor->HasAnyFlags(RF_Transient) && InPartitionActor->IsRuntimeGenerated())
+		{
+			NumPartitionActors++;
+
+			if (InPartitionActor->GetPCGGridSize() > 0)
+			{
+				InPartitionActor->RegisterPCG();
+			}
+			else
+			{
+				ExistingPoolPartitionActors.Add(InPartitionActor);
+			}
+		}
+		return true;
+	});
+
+	if (NumPartitionActors > 0)
+	{
+		check(PartitionActorPool.IsEmpty());
+		PartitionActorPool.Append(ExistingPoolPartitionActors);
+		PartitionActorPoolSize = NumPartitionActors;
+	}
+#endif
 }
 
 FPCGRuntimeGenScheduler::~FPCGRuntimeGenScheduler()
