[PCG] Fix LoadAsPreview issue where Partitioned components editing mode wasn't updated soon enough causing some Partition Actors to not serialize their generation / cleanup properly

diff --git a/Engine/Plugins/PCG/Source/PCG/Private/Grid/PCGPartitionActor.cpp b/Engine/Plugins/PCG/Source/PCG/Private/Grid/PCGPartitionActor.cpp
--- a/Engine/Plugins/PCG/Source/PCG/Private/Grid/PCGPartitionActor.cpp
+++ b/Engine/Plugins/PCG/Source/PCG/Private/Grid/PCGPartitionActor.cpp
@@ -437,8 +437,10 @@ void APCGPartitionActor::CleanupDeadGraphInstances(bool bRemoveNonNullOnly)
 	CleanupDeadGraphInstancesInternal();
 }
 
-void APCGPartitionActor::CleanupDeadGraphInstancesInternal()
+bool APCGPartitionActor::CleanupDeadGraphInstancesInternal(bool bCleanupTest, TArray<UPCGComponent*>* OutRemovedLocalComponents)
 {
+	bool bModified = false;
+
 	// First find if we have any local dead instance (= nullptr) hooked to an original component.
 	TSet<TObjectKey<UPCGComponent>> DeadOriginalInstances;
 	for (const auto& OriginalToLocalItem : OriginalToLocal)
@@ -451,21 +453,26 @@ void APCGPartitionActor::CleanupDeadGraphInstancesInternal()
 
 	if (!DeadOriginalInstances.IsEmpty())
 	{
-		Modify();
-
-		for (const TObjectKey<UPCGComponent>& DeadInstance : DeadOriginalInstances)
+		bModified = true;
+		
+		if (!bCleanupTest)
 		{
-			OriginalToLocal.Remove(DeadInstance);
-		}
+			Modify();
 
-		LocalToOriginal.Remove(nullptr);
+			for (const TObjectKey<UPCGComponent>& DeadInstance : DeadOriginalInstances)
+			{
+				OriginalToLocal.Remove(DeadInstance);
+			}
+
+			LocalToOriginal.Remove(nullptr);
+		}
 	}
 
 	// And do the same with dead original ones.
 	TSet<TObjectPtr<UPCGComponent>> DeadLocalInstances;
 	for (const auto& LocalToOriginalItem : LocalToOriginal)
 	{
-		if(!LocalToOriginalItem.Value.IsValid())
+		if(IsValid(LocalToOriginalItem.Key) && !LocalToOriginalItem.Value.IsValid())
 		{
 			DeadLocalInstances.Add(LocalToOriginalItem.Key);
 		}
@@ -473,6 +480,18 @@ void APCGPartitionActor::CleanupDeadGraphInstancesInternal()
 
 	if (!DeadLocalInstances.IsEmpty())
 	{
+		bModified = true;
+
+		if (bCleanupTest)
+		{
+			if (OutRemovedLocalComponents)
+			{
+				OutRemovedLocalComponents->Append(DeadLocalInstances.Array());
+			}
+
+			return true;
+		}
+
 		Modify();
 
 		for (const TObjectPtr<UPCGComponent>& DeadInstance : DeadLocalInstances)
@@ -489,6 +508,8 @@ void APCGPartitionActor::CleanupDeadGraphInstancesInternal()
 		// Remove all dead entries
 		OriginalToLocal.Remove(nullptr);
 	}
+
+	return bModified;
 }
 
 void APCGPartitionActor::AddGraphInstance(UPCGComponent* OriginalComponent)
diff --git a/Engine/Plugins/PCG/Source/PCG/Public/Grid/PCGPartitionActor.h b/Engine/Plugins/PCG/Source/PCG/Public/Grid/PCGPartitionActor.h
--- a/Engine/Plugins/PCG/Source/PCG/Public/Grid/PCGPartitionActor.h
+++ b/Engine/Plugins/PCG/Source/PCG/Public/Grid/PCGPartitionActor.h
@@ -133,7 +133,14 @@ public:
 	FGuid PCGGuid;
 
 private:
-	UE_API void CleanupDeadGraphInstancesInternal();
+	/** 
+	 * Cleanup graph instances if Original or Local components are invalid. *
+	 * 
+	 * @param bCleanupTest					If true, we only test that there are some components to cleanup.
+	 * @param OutRemovedLocalComponents		Returns the local components to cleanup if non-null and bCleanupTest is true.
+	 * @return True if there are some components to cleanup (when bCleanupTest is true) or if actual cleanup happened (when bCleanupTest is false)
+	 */
+	UE_API bool CleanupDeadGraphInstancesInternal(bool bCleanupTest = false, TArray<UPCGComponent*>* OutRemovedLocalComponents = nullptr);
 
 	UE_API UPCGSubsystem* GetSubsystem() const;
 
diff --git a/Engine/Plugins/PCG/Source/PCGEditor/Private/WorldPartitionBuilder/PCGWorldPartitionBuilder.cpp b/Engine/Plugins/PCG/Source/PCGEditor/Private/WorldPartitionBuilder/PCGWorldPartitionBuilder.cpp
--- a/Engine/Plugins/PCG/Source/PCGEditor/Private/WorldPartitionBuilder/PCGWorldPartitionBuilder.cpp
+++ b/Engine/Plugins/PCG/Source/PCGEditor/Private/WorldPartitionBuilder/PCGWorldPartitionBuilder.cpp
@@ -110,6 +110,8 @@ struct FPCGWorldPartitionBuilder
 	static void LogModifiedAndNewActorsFromDirtyPackages(const TArray<UPackage*>& InDirtyPackages);
 
 	static void SetChainedDispatchToLocalComponents(UWorld* InWorld, bool bInChainedDispatch);
+
+	static void UpdateEditingMode(UPCGComponent* InComponent);
 private:
 	static void CollectSubGridComponentsToGenerateInternal(UPCGSubsystem* InPCGSubsystem, int32 InGridSizeIndex, const PCGHiGenGrid::FSizeArray& InGridSizes, UPCGComponent* InOriginalComponent, APCGPartitionActor* InPartitionActor, TArray<UPCGComponent*>& OutAdditionalComponentsToGenerate, TMap<UPCGComponent*, UPCGComponent*>& OutDependencies);
 
@@ -1282,16 +1284,8 @@ FPCGTaskId FPCGWorldPartitionBuilder::GenerateComponent(UPCGComponent* InCompone
 	ensure(InComponent->bActivated);
 	ensure(!InComponent->IsManagedByRuntimeGenSystem());
 
-	if (InComponent->GetSerializedEditingMode() == EPCGEditorDirtyMode::LoadAsPreview)
-	{
-		UE_LOG(LogPCGWorldPartitionBuilder, Display, TEXT("Setting PCG editing mode to Load As Preview on actor '%s' label '%s' graph '%s'."),
-			*InComponent->GetOwner()->GetName(),
-			*InComponent->GetOwner()->GetActorNameOrLabel(),
-			*InComponent->GetGraph()->GetName());
-
-		InComponent->SetEditingMode(EPCGEditorDirtyMode::LoadAsPreview, InComponent->GetSerializedEditingMode());
-		InComponent->ChangeTransientState(EPCGEditorDirtyMode::LoadAsPreview);
-	}
+	// Non-partitioned components get their editing mode updated here to make sure generation is persisted
+	FPCGWorldPartitionBuilder::UpdateEditingMode(InComponent);
 
 	// Force generate as components that are already generated may decline the request.
 	FPCGTaskId GenerateTask = InvalidPCGTaskId;
@@ -1648,6 +1642,20 @@ void FPCGWorldPartitionBuilder::SetChainedDispatchToLocalComponents(UWorld* InWo
 	PCGSubsystem->SetChainedDispatchToLocalComponents(bInChainedDispatch);
 }
 
+void FPCGWorldPartitionBuilder::UpdateEditingMode(UPCGComponent* InComponent)
+{
+	if (InComponent->GetSerializedEditingMode() == EPCGEditorDirtyMode::LoadAsPreview && InComponent->GetEditingMode() != EPCGEditorDirtyMode::LoadAsPreview)
+	{
+		UE_LOG(LogPCGWorldPartitionBuilder, Display, TEXT("Setting PCG editing mode to Load As Preview on actor '%s' label '%s' graph '%s'."),
+			*InComponent->GetOwner()->GetName(),
+			*InComponent->GetOwner()->GetActorNameOrLabel(),
+			*InComponent->GetGraph()->GetName());
+
+		InComponent->SetEditingMode(EPCGEditorDirtyMode::LoadAsPreview, EPCGEditorDirtyMode::LoadAsPreview);
+		InComponent->ChangeTransientState(EPCGEditorDirtyMode::LoadAsPreview);
+	}
+}
+
 TArray<UPackage*> FPCGWorldPartitionBuilder::GetDirtyPackages(const UWorld* InWorld, const TSet<FString>& AllowedAssetPackages)
 {
 	TArray<UPackage*> DirtyPackages;
@@ -1890,6 +1898,9 @@ bool FPCGWorldPartitionBuilder::CreatePartitionedActors(UWorld* InWorld, const F
 	{
 		if (UPCGComponent* PCGComponent = ComponentPtr.Get())
 		{
+			// Partitioned components need their editing mode updated before we start creating/updating PAs
+			FPCGWorldPartitionBuilder::UpdateEditingMode(PCGComponent);
+
 			check(!InOutPartitionedComponents.Contains(PCGComponent));
 			InOutPartitionedComponents.Add(PCGComponent, FWorldPartitionReference(WorldPartition, PCGComponent->GetOwner()->GetActorGuid()));
 
@@ -1949,14 +1960,32 @@ bool FPCGWorldPartitionBuilder::UpdatePartitionedActors(UWorld* InWorld, const T
 		if (PartitionActor->GetLevel() == InWorld->PersistentLevel)
 		{
 			// Here we cleanup dead graphs since we tried to load the Original components before, if they can't resolve it means they don't exist anymore.
-			PartitionActor->CleanupDeadGraphInstancesInternal();
+			TArray<UPCGComponent*> DeadLocalComponents;
+
+			// Test first so we can update editing mode before doing the clean up
+			if (PartitionActor->CleanupDeadGraphInstancesInternal(/*bCleanupTest=*/true, &DeadLocalComponents))
+			{
+				for (UPCGComponent* DeadLocalComponent : DeadLocalComponents)
+				{
+					// Update editing mode before cleaning up to make sure cleanup gets persisted
+					FPCGWorldPartitionBuilder::UpdateEditingMode(DeadLocalComponent);
+				}
+
+				PartitionActor->CleanupDeadGraphInstancesInternal(/*bCleanupTest=*/false);
+			}
 
 			// Remove Graph Instances from Partition Actors that are no longer part of the component mappings (bounds change)
 			for (const TKeyValuePair<UPCGComponent*, TSet<TObjectPtr<APCGPartitionActor>>>& Mapping : UpdatedMappings)
 			{
 				if (!Mapping.Value.Contains(PartitionActor))
 				{
-					PartitionActor->RemoveGraphInstance(Mapping.Key);
+					if (UPCGComponent* LocalComponent = PartitionActor->GetLocalComponent(Mapping.Key))
+					{
+						// Update editing mode before removing graph instance to make sure cleanup gets persisted
+						FPCGWorldPartitionBuilder::UpdateEditingMode(LocalComponent);
+
+						PartitionActor->RemoveGraphInstance(Mapping.Key);
+					}
 				}
 			}
 		}
