[PCG] Fix ensure when GPU pins have the same label

#rnx

diff --git a/Engine/Plugins/PCG/Source/PCG/Private/Graph/PCGGraphCompilerGPU.cpp b/Engine/Plugins/PCG/Source/PCG/Private/Graph/PCGGraphCompilerGPU.cpp
--- a/Engine/Plugins/PCG/Source/PCG/Private/Graph/PCGGraphCompilerGPU.cpp
+++ b/Engine/Plugins/PCG/Source/PCG/Private/Graph/PCGGraphCompilerGPU.cpp
@@ -1251,8 +1251,6 @@ void FPCGGraphCompilerGPU::CompileComputeGraph(
 					if (!DataLabels.IsEmpty())
 					{
 						FPCGPinDataLabels& PinDataLabels = InOutComputeGraph->StaticDataLabelsTable.FindOrAdd(Kernel->GetKernelIndex());
-						ensure(!PinDataLabels.PinToDataLabels.Contains(PinProps.Label));
-
 						FPCGDataLabels& Labels = PinDataLabels.PinToDataLabels.Add(PinProps.Label);
 						Labels.Labels = MoveTemp(DataLabels);
 					}
