[PCG] Add grid streaming mode to scheduling policy
* The new default behavior will be to depend on all grids (except unbounded) for world streaming
* Add support for unbounded grid as a world streaming dependency
* Fix crash in PCGActorHelpers::GetCellCoord()

diff --git a/Engine/Plugins/PCG/Source/PCG/Private/Helpers/PCGActorHelpers.cpp b/Engine/Plugins/PCG/Source/PCG/Private/Helpers/PCGActorHelpers.cpp
--- a/Engine/Plugins/PCG/Source/PCG/Private/Helpers/PCGActorHelpers.cpp
+++ b/Engine/Plugins/PCG/Source/PCG/Private/Helpers/PCGActorHelpers.cpp
@@ -814,7 +814,7 @@ UPackage* UPCGActorHelpers::CreatePreviewPackage(ULevel* InLevel, const FString&
 }
 #endif // WITH_EDITOR
 
-FIntVector UPCGActorHelpers::GetCellCoord(FVector InPosition, int InGridSize, bool bUse2DGrid)
+FIntVector UPCGActorHelpers::GetCellCoord(FVector InPosition, uint32 InGridSize, bool bUse2DGrid)
 {
 	check(InGridSize > 0);
 
diff --git a/Engine/Plugins/PCG/Source/PCG/Private/RuntimeGen/PCGRuntimeGenScheduler.cpp b/Engine/Plugins/PCG/Source/PCG/Private/RuntimeGen/PCGRuntimeGenScheduler.cpp
--- a/Engine/Plugins/PCG/Source/PCG/Private/RuntimeGen/PCGRuntimeGenScheduler.cpp
+++ b/Engine/Plugins/PCG/Source/PCG/Private/RuntimeGen/PCGRuntimeGenScheduler.cpp
@@ -468,6 +468,8 @@ void FPCGRuntimeGenScheduler::TickQueueComponentsForGeneration(
 
 		TRACE_CPUPROFILER_EVENT_SCOPE(IsWorldStreamingComplete);
 
+		// @todo_pcg: We should be querying a box instead of a radius. As is, this will miss corners, but maybe this is preferable to
+		// dilating the streaming query radius to circumscribe the volume, which could induce significantly more generation latency.
 		StreamingQuerySources[0].Radius = InGridSize / 2.0f;
 		StreamingQuerySources[0].Location = InLocation;
 
@@ -502,9 +504,8 @@ void FPCGRuntimeGenScheduler::TickQueueComponentsForGeneration(
 
 		// todo_pcg: For each execution domain (for now only GenAtRuntime/dynamic), assuming we run Preview through this scheduler, which it seems like we will.
 		// todo_pcg: All this stuff can be hoisted - whether valid, has graph, active, managed by runtime gen..
-		if (OriginalComponent->IsManagedByRuntimeGenSystem())
+		if (OriginalComponent->IsManagedByRuntimeGenSystem() && ensure(Policy))
 		{
-			ensure(Policy);
 			bool bHasUnbounded = false;
 			PCGHiGenGrid::FSizeArray GridSizes;
 			ensure(PCGHelpers::GetGenerationGridSizes(OriginalComponent->GetGraph(), Inputs.PCGWorldActor, GridSizes, bHasUnbounded));
@@ -623,7 +624,18 @@ void FPCGRuntimeGenScheduler::TickQueueComponentsForGeneration(
 
 				if (bHasUnbounded)
 				{
-					// TODO check that unbounded is loaded, if dependent on world state. Or temporarily don't allow querying world state in unbounded, for runtime gen.
+					// Poll for streaming to complete on the unbounded grid.
+					if (bCheckStreaming && Policy->DoesGridDependOnWorldStreaming(PCGHiGenGrid::UnboundedGridSize()))
+					{
+						FVector StreamingLocation, StreamingExtents;
+						OriginalComponentBounds.GetCenterAndExtents(StreamingLocation, StreamingExtents);
+
+						if (!IsWorldStreamingComplete(StreamingLocation, FMath::Max(StreamingExtents.X, StreamingExtents.Y)))
+						{
+							UE_LOG(LogPCG, VeryVerbose, TEXT("Partitioned original component '%s' rejected as world is not fully loaded."), *OriginalComponent->GetOwner()->GetActorNameOrLabel());
+							continue;
+						}
+					}
 					
 					// Ignore components that have already been generated or marked for generation. Unbounded grid size means not-partitioned.
 					FGridGenerationKey Key(PCGHiGenGrid::UnboundedGridSize(), FIntVector(0), OriginalComponent);
@@ -769,9 +781,23 @@ void FPCGRuntimeGenScheduler::TickQueueComponentsForGeneration(
 		const UPCGSchedulingPolicyBase* Policy = OriginalComponent->GetRuntimeGenSchedulingPolicy();
 
 		// TODO: For each execution domain (for now only GenAtRuntime/dynamic), assuming we run Preview through this scheduler, which it seems like we will.
-		if (OriginalComponent->IsManagedByRuntimeGenSystem())
+		if (OriginalComponent->IsManagedByRuntimeGenSystem() && ensure(Policy))
 		{
-			ensure(Policy);
+			const FBox OriginalComponentBounds = OriginalComponent->GetGridBounds();
+
+			// Poll for streaming to complete.
+			if (bCheckStreaming && Policy->DoesGridDependOnWorldStreaming(PCGHiGenGrid::UnboundedGridSize()))
+			{
+				FVector StreamingLocation, StreamingExtents;
+				OriginalComponentBounds.GetCenterAndExtents(StreamingLocation, StreamingExtents);
+
+				if (!IsWorldStreamingComplete(StreamingLocation, FMath::Max(StreamingExtents.X, StreamingExtents.Y)))
+				{
+					UE_LOG(LogPCG, VeryVerbose, TEXT("Non-partitioned original component '%s' rejected as world is not fully loaded."), *OriginalComponent->GetOwner()->GetActorNameOrLabel());
+					continue;
+				}
+			}
+
 			// Unbounded will grab the base GenerationRadius used for non-partitioned and unbounded.
 			const double MaxGenerationRadius = OriginalComponent->GetGenerationRadiusFromGrid(EPCGHiGenGrid::Unbounded);
 
@@ -784,7 +810,6 @@ void FPCGRuntimeGenScheduler::TickQueueComponentsForGeneration(
 				}
 
 				const FVector GenSourcePosition = GenSourcePositionOptional.GetValue();
-				const FBox OriginalComponentBounds = OriginalComponent->GetGridBounds();
 
 				FVector ModifiedGenSourcePosition = GenSourcePosition;
 				if (OriginalComponent->Use2DGrid())
diff --git a/Engine/Plugins/PCG/Source/PCG/Public/Helpers/PCGActorHelpers.h b/Engine/Plugins/PCG/Source/PCG/Public/Helpers/PCGActorHelpers.h
--- a/Engine/Plugins/PCG/Source/PCG/Public/Helpers/PCGActorHelpers.h
+++ b/Engine/Plugins/PCG/Source/PCG/Public/Helpers/PCGActorHelpers.h
@@ -231,7 +231,7 @@ PRAGMA_ENABLE_DEPRECATION_WARNINGS
 	/**
 	 * Return the grid cell coordinates on the PCG partition grid given a position and the grid size.
 	 */
-	static UE_API FIntVector GetCellCoord(FVector InPosition, int InGridSize, bool bUse2DGrid);
+	static UE_API FIntVector GetCellCoord(FVector InPosition, uint32 InGridSize, bool bUse2DGrid);
 
 	/**
 	 * Return the center of the PCG partition grid cell given a position and the grid size.
diff --git a/Engine/Plugins/PCG/Source/PCG/Public/RuntimeGen/SchedulingPolicies/PCGSchedulingPolicyBase.h b/Engine/Plugins/PCG/Source/PCG/Public/RuntimeGen/SchedulingPolicies/PCGSchedulingPolicyBase.h
--- a/Engine/Plugins/PCG/Source/PCG/Public/RuntimeGen/SchedulingPolicies/PCGSchedulingPolicyBase.h
+++ b/Engine/Plugins/PCG/Source/PCG/Public/RuntimeGen/SchedulingPolicies/PCGSchedulingPolicyBase.h
@@ -10,6 +10,16 @@
 
 class IPCGGenSourceBase;
 
+/** Determines grid dependencies for world streaming. */
+UENUM()
+enum class EPCGGridStreamingDependencyMode
+{
+	AllGridsExceptUnbounded = 0,
+	AllGrids,
+	SpecificGrids,
+	NoGrids,
+};
+
 /** 
  * Scheduling Policies provide custom logic to efficiently schedule work for the Runtime Generation Scheduling system.
  * A higher priority value means the work will be scheduled sooner, and larger grid sizes will always have a higher
@@ -46,16 +56,33 @@ public:
 
 	bool DoesGridDependOnWorldStreaming(uint32 InGridSize) const
 	{
+		if (GridStreamingDependencyMode == EPCGGridStreamingDependencyMode::NoGrids)
+		{
+			return false;
+		}
+
 		const EPCGHiGenGrid Grid = PCGHiGenGrid::GridSizeToGrid(InGridSize);
-		return Grid != EPCGHiGenGrid::Uninitialized && GridsDependentOnWorldStreaming.Contains(Grid);
+
+		if (Grid == EPCGHiGenGrid::Uninitialized)
+		{
+			return false;
+		}
+
+		return GridStreamingDependencyMode == EPCGGridStreamingDependencyMode::AllGrids
+			|| (GridStreamingDependencyMode == EPCGGridStreamingDependencyMode::AllGridsExceptUnbounded && Grid != EPCGHiGenGrid::Unbounded)
+			|| (GridStreamingDependencyMode == EPCGGridStreamingDependencyMode::SpecificGrids && GridsDependentOnWorldStreaming.Contains(Grid));
 	}
 
 public:
 	/**
-	* Grids that depend on world streaming. Cells on these grids will only generate if the world within the generation volume is fully streamed. Enable this
-	* if the graph depends on actors in the world such as a landscape or a set of streamable actors.
-	*/
+	 * Grids with a streaming dependency will only generate if the world within the generation volume is fully streamed. Enable this
+	 * if the graph depends on actors in the world such as a landscape or a set of streamable actors.
+	 */
 	UPROPERTY(EditAnywhere, Category = "RuntimeGeneration|Scheduling Policy Parameters", meta = (EditCondition = "bShouldDisplayProperties", EditConditionHides, HideEditConditionToggle))
+	EPCGGridStreamingDependencyMode GridStreamingDependencyMode = EPCGGridStreamingDependencyMode::AllGridsExceptUnbounded;
+
+	/** Grids that depend on world streaming. */
+	UPROPERTY(EditAnywhere, Category = "RuntimeGeneration|Scheduling Policy Parameters", meta = (EditCondition = "bShouldDisplayProperties && GridStreamingDependencyMode == EPCGGridStreamingDependencyMode::SpecificGrids", EditConditionHides, HideEditConditionToggle))
 	TArray<EPCGHiGenGrid> GridsDependentOnWorldStreaming;
 
 #if WITH_EDITORONLY_DATA
