[PCG] QoL improvement for data type info naming

* Will look only for DataTypeInfo by default in the class name
* Can specify a display name in metadata

#rnx

diff --git a/Engine/Plugins/PCG/Source/PCG/Private/Data/Registry/PCGDataType.cpp b/Engine/Plugins/PCG/Source/PCG/Private/Data/Registry/PCGDataType.cpp
--- a/Engine/Plugins/PCG/Source/PCG/Private/Data/Registry/PCGDataType.cpp
+++ b/Engine/Plugins/PCG/Source/PCG/Private/Data/Registry/PCGDataType.cpp
@@ -51,13 +51,28 @@ FString FPCGDataTypeBaseId::ToString() const
 	{
 		return TEXT("Any");
 	}
+#if WITH_EDITOR
+	else if (Struct->HasMetaData(PCGObjectMetadata::DataTypeDisplayName))
+	{
+		return Struct->GetMetaData(PCGObjectMetadata::DataTypeDisplayName);
+	}
+#endif // WITH_EDITOR
 	else
 	{
 		// Strip `PCGDataTypeInfo` of the name
-		static const FString Prefix = FPCGDataTypeInfo::StaticStruct()->GetName();
+		static const FString Prefix = TEXT("DataTypeInfo");
+		static const int32 PrefixSize = Prefix.Len();
+
 		FString Result = Struct->GetName();
-		Result.RemoveFromStart(Prefix);
-		return Result;
+		const int32 PrefixIndex = Result.Find(Prefix);
+		if (PrefixIndex != INDEX_NONE)
+		{
+			return FString(FStringView(Result).RightChop(PrefixIndex + PrefixSize));
+		}
+		else
+		{
+			return Result;
+		}
 	}
 }
 
diff --git a/Engine/Plugins/PCG/Source/PCG/Public/Data/Registry/PCGDataType.h b/Engine/Plugins/PCG/Source/PCG/Public/Data/Registry/PCGDataType.h
--- a/Engine/Plugins/PCG/Source/PCG/Public/Data/Registry/PCGDataType.h
+++ b/Engine/Plugins/PCG/Source/PCG/Public/Data/Registry/PCGDataType.h
@@ -57,6 +57,7 @@ class UPCGSettings;
 /**
  * Store the information about a type.
  * Used by FPCGDataTypeBaseId as the UScriptStruct (the meta-struct) and will have a "default object" initialized in the registry for the virtual calls.
+ * You can also add a "PCG_DataTypeDisplayName" meta tag to the UStruct to have a custom display name. Otherwise by default, it will keep what is at the right of `DataTypeInfo`.
  */
 USTRUCT()
 struct FPCGDataTypeInfo
diff --git a/Engine/Plugins/PCG/Source/PCG/Public/PCGCommon.h b/Engine/Plugins/PCG/Source/PCG/Public/PCGCommon.h
--- a/Engine/Plugins/PCG/Source/PCG/Public/PCGCommon.h
+++ b/Engine/Plugins/PCG/Source/PCG/Public/PCGCommon.h
@@ -250,6 +250,7 @@ namespace PCGObjectMetadata
 	const FLazyName PropertyReadOnly = TEXT("PCG_PropertyReadOnly");
 	const FLazyName DataTypeIdentifierSupportsComposition = TEXT("PCG_SupportsComposition");
 	const FLazyName DataTypeIdentifierFilter = TEXT("PCG_TypeFilter");
+	const FLazyName DataTypeDisplayName = TEXT("PCG_DataTypeDisplayName");
 
 	// Metadata usable in UPROPERTY for customizing the behavior when displaying the property in a property panel or graph node
 	enum
@@ -292,6 +293,13 @@ namespace PCGObjectMetadata
 		PCG_TypeFilter,
 	};
 
+	// Metadata usable in USTRUCT
+	enum
+	{
+		/// [StructMetadata] For FPCGDataTypeInfo, can specify a display name.
+		PCG_DataTypeDisplayName
+	};
+
 	// Metadata usable in UENUM for customizing the entry in the attribute property selector
 	enum
 	{
