[PCG] Make sure PVE types and nodes are not bleeding in PCG.

* Very hacky and temporary solution until we have a design for that

#rnx

diff --git a/Engine/Plugins/Experimental/ProceduralVegetationEditor/Source/ProceduralVegetation/Public/DataTypes/PVFoliageMeshData.h b/Engine/Plugins/Experimental/ProceduralVegetationEditor/Source/ProceduralVegetation/Public/DataTypes/PVFoliageMeshData.h
--- a/Engine/Plugins/Experimental/ProceduralVegetationEditor/Source/ProceduralVegetation/Public/DataTypes/PVFoliageMeshData.h
+++ b/Engine/Plugins/Experimental/ProceduralVegetationEditor/Source/ProceduralVegetation/Public/DataTypes/PVFoliageMeshData.h
@@ -12,6 +12,11 @@ struct FPVDataTypeInfoFoliageMesh : public FPCGDataTypeInfo
 	GENERATED_BODY()
 
 	PCG_DECLARE_TYPE_INFO(PROCEDURALVEGETATION_API);
+
+#if WITH_EDITOR
+	// If the type should be hidden to the user.
+	virtual bool Hidden() const override { return true; };
+#endif // WITH_EDITOR
 };
 
 UCLASS(BlueprintType, ClassGroup = (Procedural))
diff --git a/Engine/Plugins/Experimental/ProceduralVegetationEditor/Source/ProceduralVegetation/Public/DataTypes/PVGrowthData.h b/Engine/Plugins/Experimental/ProceduralVegetationEditor/Source/ProceduralVegetation/Public/DataTypes/PVGrowthData.h
--- a/Engine/Plugins/Experimental/ProceduralVegetationEditor/Source/ProceduralVegetation/Public/DataTypes/PVGrowthData.h
+++ b/Engine/Plugins/Experimental/ProceduralVegetationEditor/Source/ProceduralVegetation/Public/DataTypes/PVGrowthData.h
@@ -12,6 +12,11 @@ struct FPVDataTypeInfoGrowth : public FPCGDataTypeInfo
 	GENERATED_BODY()
 
 	PCG_DECLARE_TYPE_INFO(PROCEDURALVEGETATION_API);
+
+#if WITH_EDITOR
+	// If the type should be hidden to the user.
+	virtual bool Hidden() const override { return true; };
+#endif // WITH_EDITOR
 };
 
 UCLASS(BlueprintType, ClassGroup = (Procedural))
diff --git a/Engine/Plugins/Experimental/ProceduralVegetationEditor/Source/ProceduralVegetation/Public/DataTypes/PVMeshData.h b/Engine/Plugins/Experimental/ProceduralVegetationEditor/Source/ProceduralVegetation/Public/DataTypes/PVMeshData.h
--- a/Engine/Plugins/Experimental/ProceduralVegetationEditor/Source/ProceduralVegetation/Public/DataTypes/PVMeshData.h
+++ b/Engine/Plugins/Experimental/ProceduralVegetationEditor/Source/ProceduralVegetation/Public/DataTypes/PVMeshData.h
@@ -12,6 +12,11 @@ struct FPVDataTypeInfoMesh : public FPCGDataTypeInfo
 	GENERATED_BODY()
 
 	PCG_DECLARE_TYPE_INFO(PROCEDURALVEGETATION_API);
+
+#if WITH_EDITOR
+	// If the type should be hidden to the user.
+	virtual bool Hidden() const override { return true; };
+#endif // WITH_EDITOR
 };
 
 UCLASS(BlueprintType, ClassGroup = (Procedural))
diff --git a/Engine/Plugins/PCG/Source/PCGEditor/Private/PCGEditorGraphSchema.cpp b/Engine/Plugins/PCG/Source/PCGEditor/Private/PCGEditorGraphSchema.cpp
--- a/Engine/Plugins/PCG/Source/PCGEditor/Private/PCGEditorGraphSchema.cpp
+++ b/Engine/Plugins/PCG/Source/PCGEditor/Private/PCGEditorGraphSchema.cpp
@@ -362,11 +362,34 @@ void UPCGEditorGraphSchema::BreakSinglePinLink(UEdGraphPin* SourcePin, UEdGraphP
 
 void UPCGEditorGraphSchema::GetNativeElementActions(FGraphActionMenuBuilder& ActionMenuBuilder, const FPCGActionsFilter& ActionFilter) const
 {
+	// TO BE REMOVED IN 5.8
+	// @todo_pcg: alogut
+	// We'll design a extensible way to specify which settings is compatible to which editor in 5.8.
+	// In the meantime, hide all Procedural Vegetation node in other editors than the PVEditor.
+	const UPackage* PVPackage = FindPackage(nullptr, TEXT("/Script/ProceduralVegetation"));
+	const UPackage* PVEPackage = FindPackage(nullptr, TEXT("/Script/ProceduralVegetationEditor"));
+	const bool bIsPVEditor = GetClass() && GetClass()->GetOuterUPackage() == PVEPackage;
+
+	// Returns true if it is a PV node in a non-PV editor.
+	auto ExtraPVEFilter = [PVPackage, bIsPVEditor](TSubclassOf<UPCGSettings> InSettingsClass)
+	{
+		return !bIsPVEditor && PVPackage && InSettingsClass && InSettingsClass->GetOuterUPackage() == PVPackage;
+	};
+	
 	TArray<UClass*> SettingsClasses;
 	for (TObjectIterator<UClass> It; It; ++It)
 	{
 		TSubclassOf<UPCGSettings> Class(*It);
 
+		// TO BE REMOVED IN 5.8
+		// @todo_pcg: alogut
+		// We'll design a extensible way to specify which settings is compatible to which editor in 5.8.
+		// In the meantime, hide all Procedural Vegetation node in other editors than the PVEditor.
+		if (ExtraPVEFilter(Class))
+		{
+			continue;
+		}
+
 		if (*Class &&
 			!Class->HasAnyClassFlags(CLASS_Abstract | CLASS_Deprecated | CLASS_Hidden) &&
 			(!ActionFilter.Customization.FiltersSettings() || !ActionFilter.Customization.FilterSettings(Class)))
