[PCG] Temporary solution for allowing the viewport to update without the attribute list view opened

* Attribute list view can be forced to tick even if it is not opened.
* To be set by child editor
* To be removed when we seperate the viewport to the attribute list view.

#rnx

diff --git a/Engine/Plugins/PCG/Source/PCGEditor/Private/Widgets/SPCGEditorGraphAttributeListView.cpp b/Engine/Plugins/PCG/Source/PCGEditor/Private/Widgets/SPCGEditorGraphAttributeListView.cpp
--- a/Engine/Plugins/PCG/Source/PCGEditor/Private/Widgets/SPCGEditorGraphAttributeListView.cpp
+++ b/Engine/Plugins/PCG/Source/PCGEditor/Private/Widgets/SPCGEditorGraphAttributeListView.cpp
@@ -636,6 +636,11 @@ void SPCGEditorGraphAttributeListView::Tick(const FGeometry& AllottedGeometry, c
 {
 	SCompoundWidget::Tick(AllottedGeometry, InCurrentTime, InDeltaTime);
 
+	OnTick();
+}
+
+void SPCGEditorGraphAttributeListView::OnTick()
+{
 	if (bNeedsRefresh)
 	{
 		bNeedsRefresh = false;
@@ -707,6 +712,14 @@ void SPCGEditorGraphAttributeListView::RequestRefresh()
 	bNeedsRefresh = true;
 	// When we ask for refresh, we should also release the hold on the visualized data
 	DataPtr = nullptr;
+
+	PRAGMA_DISABLE_DEPRECATION_WARNINGS
+	if (!IsOpen() && PCGEditorPtr.IsValid() && PCGEditorPtr.Pin()->bForceRefreshAttributeEvenIfClosed)
+	PRAGMA_ENABLE_DEPRECATION_WARNINGS
+	{
+		RequestViewportRefresh();
+		OnTick();
+	}
 }
 
 void SPCGEditorGraphAttributeListView::RequestViewportRefresh()
@@ -810,6 +823,12 @@ const FPCGDataCollection* SPCGEditorGraphAttributeListView::GetInspectionData()
 
 void SPCGEditorGraphAttributeListView::RefreshAttributeList()
 {
+	// Don't refresh the list if we are not opened
+	if (!IsOpen())
+	{
+		return;
+	}
+
 	CacheColumnWidthVisibility();
 
 	HiddenAttributes = ListViewHeader->GetHiddenColumnIds();
@@ -1969,6 +1988,12 @@ bool SPCGEditorGraphAttributeListView::IsViewportOpen() const
 	return PCGEditorPtr.IsValid() && PCGEditorPtr.Pin()->IsPanelCurrentlyOpen(ViewportEditorPanel);
 }
 
+bool SPCGEditorGraphAttributeListView::IsOpen() const
+{
+	const int AttributeEditorPanel = static_cast<int>(EPCGEditorPanel::Attributes1) + WidgetEntryNumber;
+	return PCGEditorPtr.IsValid() && PCGEditorPtr.Pin()->IsPanelCurrentlyOpen(static_cast<EPCGEditorPanel>(AttributeEditorPanel));
+}
+
 
 void SPCGEditorGraphAttributeListView::OnToggleShowDefaultValue()
 {
diff --git a/Engine/Plugins/PCG/Source/PCGEditor/Private/Widgets/SPCGEditorGraphAttributeListView.h b/Engine/Plugins/PCG/Source/PCGEditor/Private/Widgets/SPCGEditorGraphAttributeListView.h
--- a/Engine/Plugins/PCG/Source/PCGEditor/Private/Widgets/SPCGEditorGraphAttributeListView.h
+++ b/Engine/Plugins/PCG/Source/PCGEditor/Private/Widgets/SPCGEditorGraphAttributeListView.h
@@ -171,6 +171,7 @@ public:
 	/** Add UObject references for GC */
 	void AddReferencedObjects(FReferenceCollector& Collector);
 private:
+	void OnTick();
 	TSharedRef<SHeaderRow> CreateHeaderRowWidget() const;
 
 	void OnInspectedStackChanged(const FPCGStack& InPCGStack);
@@ -263,6 +264,7 @@ private:
 	bool CanFocusOnSelection() const;
 
 	bool IsViewportOpen() const;
+	bool IsOpen() const;
 
 	/** Returns the current PCG component */
 	TWeakInterfacePtr<IPCGGraphExecutionSource> GetExecutionSource() const;
diff --git a/Engine/Plugins/PCG/Source/PCGEditor/Public/PCGEditor.h b/Engine/Plugins/PCG/Source/PCGEditor/Public/PCGEditor.h
--- a/Engine/Plugins/PCG/Source/PCGEditor/Public/PCGEditor.h
+++ b/Engine/Plugins/PCG/Source/PCGEditor/Public/PCGEditor.h
@@ -243,6 +243,14 @@ protected:
 	/** Create a new viewport widget. Can be subclassed to have customization around the viewport. */
 	PCGEDITOR_API virtual TSharedRef<SPCGEditorViewport> CreateViewportWidget();
 
+	// TO BE REMOVED IN 5.8
+	// @todo_pcg
+	// Temporary boolean to force the Attribute List view to Update even if it is closed. To be set in the Editor Ctor.
+	UE_DEPRECATED(all, "Do not use.")
+	bool bForceRefreshAttributeEvenIfClosed = false;
+
+	friend SPCGEditorGraphAttributeListView;
+
 private:
 	/** Register PCG specific toolbar for the editor */
 	void RegisterToolbar() const;
