[PCG] Enable pcg.FullDataCrcInCommandlets by default
- Needed to generate consistent Crcs to support reverting unchanged actors with builder commandlet

diff --git a/Engine/Plugins/PCG/Source/PCG/Private/PCGData.cpp b/Engine/Plugins/PCG/Source/PCG/Private/PCGData.cpp
--- a/Engine/Plugins/PCG/Source/PCG/Private/PCGData.cpp
+++ b/Engine/Plugins/PCG/Source/PCG/Private/PCGData.cpp
@@ -22,7 +22,7 @@
 
 static TAutoConsoleVariable<bool> CVarPCGFullDataCrcInCommandlets(
 	TEXT("pcg.FullDataCrcInCommandlets"),
-	false,
+	true,
 	TEXT("Always compute full data Crcs in Commandlets (to improve reusability)."));
 
 UPCGData::UPCGData(const FObjectInitializer& ObjectInitializer)
