Fixed duplication of texture asset potentially propagating RF_Standalone/RF_Public if source texture is a texture asset. This prevents the transient UTexture2D inside UPCGTextureData not ever being GC'd.

diff --git a/Engine/Plugins/PCG/Source/PCG/Private/Data/PCGTextureData.cpp b/Engine/Plugins/PCG/Source/PCG/Private/Data/PCGTextureData.cpp
--- a/Engine/Plugins/PCG/Source/PCG/Private/Data/PCGTextureData.cpp
+++ b/Engine/Plugins/PCG/Source/PCG/Private/Data/PCGTextureData.cpp
@@ -563,6 +563,7 @@ void UPCGTextureData::InitializeInternal(UTexture* InTexture, uint32 InTextureIn
 		{
 			// Duplicate texture and change access flags (editor only). This duplicate texture will be used by the normal logic below.
 			FObjectDuplicationParameters DuplicationParams(Texture2D, /*Outer=*/this);
+			DuplicationParams.FlagMask = GetFlags() & ~(RF_Standalone | RF_Public);
 			DuplicateTexture = CastChecked<UTexture2D>(StaticDuplicateObjectEx(DuplicationParams));
 		}
 
