[PCG] Prevent ChangeTransientState(EPCGEditorDirtyMode::Preview) being called on Generated Resources that have been marked transient on load
- With a Blueprint PCG Volume, this was causing generated ISMs from being marked RF_Transient on load, if followed by a save of the volume those ISMs were lost even if the Generated Resources were not

diff --git a/Engine/Plugins/PCG/Source/PCG/Private/PCGComponent.cpp b/Engine/Plugins/PCG/Source/PCG/Private/PCGComponent.cpp
--- a/Engine/Plugins/PCG/Source/PCG/Private/PCGComponent.cpp
+++ b/Engine/Plugins/PCG/Source/PCG/Private/PCGComponent.cpp
@@ -4134,6 +4134,13 @@ void UPCGComponent::ChangeTransientState(EPCGEditorDirtyMode NewEditingMode)
 		{
 			if (GeneratedResource)
 			{
+				// Avoid changing transient state when switching to preview mode on marked as transient resources.
+				// When serializing a component we need its marked as transient resources to remain unchanged so they get serialized properly.
+				if (NewEditingMode == EPCGEditorDirtyMode::Preview && GeneratedResource->IsMarkedTransientOnLoad())
+				{
+					continue;
+				}
+
 				GeneratedResource->ChangeTransientState(NewEditingMode);
 				bShouldMarkDirty = true;
 			}
diff --git a/Engine/Plugins/PCG/Source/PCG/Private/PCGManagedResource.cpp b/Engine/Plugins/PCG/Source/PCG/Private/PCGManagedResource.cpp
--- a/Engine/Plugins/PCG/Source/PCG/Private/PCGManagedResource.cpp
+++ b/Engine/Plugins/PCG/Source/PCG/Private/PCGManagedResource.cpp
@@ -64,6 +64,9 @@ bool UPCGManagedResource::DebugForcePurgeAllResourcesOnGenerate()
 #if WITH_EDITOR
 void UPCGManagedResource::ChangeTransientState(EPCGEditorDirtyMode NewEditingMode)
 {
+	// We expect that our resource isn't bMarkedTransientOnLoad when setting the editor mode to preview because this would mean potential data loss
+	ensure(!bMarkedTransientOnLoad || NewEditingMode != EPCGEditorDirtyMode::Preview);
+
 	// Any change in the transient state resets the transient state that was set on load, regardless of the bNowTransient flag
 	bMarkedTransientOnLoad = false;
 
