[PCG] BugHawk fix - verify mips are non-empty and don't unlock unless lock was successful.

#rnx
#FailedValidations Preflight

diff --git a/Engine/Plugins/PCG/Source/PCG/Private/Data/PCGTextureData.cpp b/Engine/Plugins/PCG/Source/PCG/Private/Data/PCGTextureData.cpp
--- a/Engine/Plugins/PCG/Source/PCG/Private/Data/PCGTextureData.cpp
+++ b/Engine/Plugins/PCG/Source/PCG/Private/Data/PCGTextureData.cpp
@@ -1015,7 +1015,7 @@ TOptional<bool> UPCGTextureData::InitializeGPUTextureFromCPU()
 
 	bool bBulkDataAccessed = false;
 
-	if (const uint8_t* BulkData = reinterpret_cast<const uint8_t*>(PlatformData->Mips[0].BulkData.LockReadOnly()))
+	if (const uint8_t* BulkData = reinterpret_cast<const uint8_t*>(PlatformData->Mips.IsEmpty() ? nullptr : PlatformData->Mips[0].BulkData.LockReadOnly()))
 	{
 		bBulkDataAccessed = true;
 
@@ -1029,14 +1029,14 @@ TOptional<bool> UPCGTextureData::InitializeGPUTextureFromCPU()
 		{
 			ColorData[D] = FormattedImageData[D].ReinterpretAsLinear();
 		}
+
+		PlatformData->Mips[0].BulkData.Unlock();
 	}
 	else
 	{
 		UE_LOG(LogPCG, Error, TEXT("PCGTextureData unable to get bulk data from '%s'."), *Texture->GetFName().ToString());
 	}
 
-	PlatformData->Mips[0].BulkData.Unlock();
-
 	return bBulkDataAccessed;
 }
 #endif
