[PCG] Fix Preview behavior for GeneratedGraphOutput (to match ManagedResources)
- Also Load As Preview behavior should only affect generated components

diff --git a/Engine/Plugins/PCG/Source/PCG/Private/PCGComponent.cpp b/Engine/Plugins/PCG/Source/PCG/Private/PCGComponent.cpp
--- a/Engine/Plugins/PCG/Source/PCG/Private/PCGComponent.cpp
+++ b/Engine/Plugins/PCG/Source/PCG/Private/PCGComponent.cpp
@@ -1832,18 +1832,38 @@ void UPCGComponent::CleanupUnusedManagedResources()
 	PCGGeneratedResourcesLogging::LogCleanupUnusedManagedResourcesFinished(this, GeneratedResources);
 }
 
-void UPCGComponent::ClearGraphGeneratedOutput()
+void UPCGComponent::ClearGraphGeneratedOutput(bool bClearLoadedPreviewData)
 {
 	GeneratedGraphOutput.ClearUsage(EPCGDataUsage::ComponentOutputData);
 
+#if WITH_EDITOR
+	// Do not re-outer loaded preview data if bClearLoadedPreviewData is false
+	auto IsLoadedPreviewData = [this, bClearLoadedPreviewData](const UPCGData* InData)
+	{
+		return !bClearLoadedPreviewData && LoadedPreviewGeneratedGraphOutput.TaggedData.ContainsByPredicate([InData](const FPCGTaggedData & TaggedData) { return TaggedData.Data == InData; });
+	};
+#endif
+
 	for (FPCGTaggedData& GeneratedData : GeneratedGraphOutput.TaggedData)
 	{
-		if (GeneratedData.Data && GeneratedData.Data->GetOuter() == this)
+		if ( GeneratedData.Data 
+			&& GeneratedData.Data->GetOuter() == this 
+#if WITH_EDITOR
+			&& !IsLoadedPreviewData(GeneratedData.Data.Get())
+#endif
+			)
 		{
 			const_cast<UPCGData*>(GeneratedData.Data.Get())->Rename(nullptr, GetTransientPackage(), REN_DoNotDirty | REN_DontCreateRedirectors | REN_NonTransactional);
 		}
 	}
 
+#if WITH_EDITOR
+	if (bClearLoadedPreviewData)
+	{
+		LoadedPreviewGeneratedGraphOutput.Reset();
+	}
+#endif
+
 	GeneratedGraphOutput.Reset();
 }
 
@@ -1979,7 +1999,7 @@ void UPCGComponent::OnUnregister()
 void UPCGComponent::MarkSubObjectsAsGarbage()
 {
 	// First move out the generated output
-	ClearGraphGeneratedOutput();
+	ClearGraphGeneratedOutput(/*bClearLoadedPreviewData=*/true);
 	ClearPerPinGeneratedOutput();
 
 	auto ModifyAndMarkForGarbage = [](UObject* SubObject)
@@ -2071,11 +2091,15 @@ void UPCGComponent::Serialize(FArchive& Ar)
 
 #if WITH_EDITOR
 	TArray<TObjectPtr<UPCGManagedResource>> GeneratedResourcesCopy;
+	FPCGDataCollection GeneratedGraphOutputCopy;
 
 	if (!Ar.IsLoading() && CurrentEditingMode == EPCGEditorDirtyMode::Preview)
 	{
 		GeneratedResourcesCopy = GeneratedResources;
 		GeneratedResources = LoadedPreviewResources;
+
+		GeneratedGraphOutputCopy = GeneratedGraphOutput;
+		GeneratedGraphOutput = LoadedPreviewGeneratedGraphOutput;
 	}
 
 	// When duplicating for PIE, we need to Update the RuntimeGridDescriptorHash before duplication for unsaved changes
@@ -2115,6 +2139,7 @@ void UPCGComponent::Serialize(FArchive& Ar)
 	if (!Ar.IsLoading() && CurrentEditingMode == EPCGEditorDirtyMode::Preview)
 	{
 		GeneratedResources = GeneratedResourcesCopy;
+		GeneratedGraphOutput = GeneratedGraphOutputCopy;
 	}
 #endif // WITH_EDITOR
 }
@@ -2204,10 +2229,11 @@ void UPCGComponent::PostLoad()
 
 		bGenerated = false;
 	}
-	else if (CurrentEditingMode == EPCGEditorDirtyMode::LoadAsPreview && !PCGHelpers::IsRuntimeOrPIE())
+	else if (CurrentEditingMode == EPCGEditorDirtyMode::LoadAsPreview && !PCGHelpers::IsRuntimeOrPIE() && bGenerated)
 	{
 		CurrentEditingMode = EPCGEditorDirtyMode::Preview;
 		MarkResourcesAsTransientOnLoad();
+		LoadedPreviewGeneratedGraphOutput = GeneratedGraphOutput;
 		bDirtyGenerated = PCGSystemSwitches::CVarDirtyLoadAsPreviewOnLoad.GetValueOnAnyThread();
 	}
 #endif
@@ -4151,6 +4177,9 @@ void UPCGComponent::ChangeTransientState(EPCGEditorDirtyMode NewEditingMode)
 		if (NewEditingMode != EPCGEditorDirtyMode::Preview)
 		{
 			bShouldMarkDirty |= DeletePreviewResources();
+
+			bShouldMarkDirty |= (LoadedPreviewGeneratedGraphOutput != GeneratedGraphOutput);
+			LoadedPreviewGeneratedGraphOutput.Reset();
 		}
 	}
 
@@ -4362,6 +4391,16 @@ void FPCGComponentInstanceData::ApplyToComponent(UActorComponent* Component, con
 				PCGComponent->LoadedPreviewResources = DuplicateLoadedPreviewResources;
 			}
 
+			PCGComponent->LoadedPreviewGeneratedGraphOutput = ConstructionSourceComponent->LoadedPreviewGeneratedGraphOutput;
+
+			for (FPCGTaggedData& LoadedPreviewTaggedData : PCGComponent->LoadedPreviewGeneratedGraphOutput.TaggedData)
+			{
+				if (LoadedPreviewTaggedData.Data)
+				{
+					LoadedPreviewTaggedData.Data = CastChecked<UPCGData>(StaticDuplicateObject(LoadedPreviewTaggedData.Data.Get(), PCGComponent));
+				}
+			}
+
 			PCGComponent->bWasGeneratedThisSession = ConstructionSourceComponent->bWasGeneratedThisSession;
 
 			PCGComponent->ExecutionInspection.InspectionCounter = ConstructionSourceComponent->ExecutionInspection.InspectionCounter;
diff --git a/Engine/Plugins/PCG/Source/PCG/Public/PCGComponent.h b/Engine/Plugins/PCG/Source/PCG/Public/PCGComponent.h
--- a/Engine/Plugins/PCG/Source/PCG/Public/PCGComponent.h
+++ b/Engine/Plugins/PCG/Source/PCG/Public/PCGComponent.h
@@ -688,7 +688,7 @@ private:
 	UE_API void PostCleanupGraph(bool bRemoveComponents);
 	UE_API void OnProcessGraphAborted(bool bQuiet = false, bool bCleanupUnusedResources = true);
 	UE_API void CleanupUnusedManagedResources();
-	UE_API void ClearGraphGeneratedOutput();
+	UE_API void ClearGraphGeneratedOutput(bool bClearLoadedPreviewData = false);
 	UE_API bool MoveResourcesToNewActor(AActor* InNewActor, bool bCreateChild);
 
 	/** Called as part of the RerunConstructionScript, just takes resources as-is, assumes current state is empty + the resources have been retargeted if needed */
@@ -763,6 +763,9 @@ private:
 	UPROPERTY(Transient)
 	TArray<TObjectPtr<UPCGManagedResource>> LoadedPreviewResources;
 
+	UPROPERTY(Transient)
+	FPCGDataCollection LoadedPreviewGeneratedGraphOutput;
+
 	UPROPERTY(Transient, VisibleInstanceOnly, Category = Debug)
 	bool bGenerationInProgress = false;
 #endif
