[PCG] Fix issue with valid conversion on pin connection was not producing an edge

* Conversions check was done on dynamic output type, while the conversion itself was done on the allowed types of the pin which can be different

#rnx

diff --git a/Engine/Plugins/PCG/Source/PCGEditor/Private/PCGEditorGraphSchema.cpp b/Engine/Plugins/PCG/Source/PCGEditor/Private/PCGEditorGraphSchema.cpp
--- a/Engine/Plugins/PCG/Source/PCGEditor/Private/PCGEditorGraphSchema.cpp
+++ b/Engine/Plugins/PCG/Source/PCGEditor/Private/PCGEditorGraphSchema.cpp
@@ -254,16 +254,19 @@ bool UPCGEditorGraphSchema::TryCreateConnectionInternal(UEdGraphPin* InA, UEdGra
 	{
 		const FPCGDataTypeRegistry& Registry = FPCGModule::GetConstDataTypeRegistry();
 
-		auto AddConversionNodes = [&Registry, PCGPinA, PCGPinB, PCGGraph, &bShouldReconstruct](const FPCGDataTypeBaseId& Id, bool bIsInput) -> bool
+		const FPCGDataTypeIdentifier PinATypes = PCGPinA->GetCurrentTypesID();
+		const FPCGDataTypeIdentifier PinBTypes = PCGPinB->GetCurrentTypesID();
+
+		auto AddConversionNodes = [&Registry, PCGPinA, PCGPinB, &PinATypes, &PinBTypes, PCGGraph, &bShouldReconstruct](const FPCGDataTypeBaseId& Id, bool bIsInput) -> bool
 		{
 			const FPCGDataTypeInfo* TypeInfo = Registry.GetTypeInfo(Id);
 
-			if (bIsInput && TypeInfo && TypeInfo->AddConversionNodesTo(PCGPinA->Properties.AllowedTypes, PCGPinB->Properties.AllowedTypes, PCGGraph, PCGPinA, PCGPinB).IsSet())
+			if (bIsInput && TypeInfo && TypeInfo->AddConversionNodesTo(PinATypes, PinBTypes, PCGGraph, PCGPinA, PCGPinB).IsSet())
 			{
 				bShouldReconstruct = true;
 				return true;
 			}
-			else if (!bIsInput && TypeInfo && TypeInfo->AddConversionNodesFrom(PCGPinA->Properties.AllowedTypes, PCGPinB->Properties.AllowedTypes, PCGGraph, PCGPinA, PCGPinB).IsSet())
+			else if (!bIsInput && TypeInfo && TypeInfo->AddConversionNodesFrom(PinATypes, PinBTypes, PCGGraph, PCGPinA, PCGPinB).IsSet())
 			{
 				bShouldReconstruct = true;
 				return true;
@@ -273,7 +276,7 @@ bool UPCGEditorGraphSchema::TryCreateConnectionInternal(UEdGraphPin* InA, UEdGra
 		};
 
 		// Try for both pins
-		for (const FPCGDataTypeBaseId& Id : PCGPinA->Properties.AllowedTypes.GetIds())
+		for (const FPCGDataTypeBaseId& Id : PinATypes.GetIds())
 		{
 			if (AddConversionNodes(Id, /*bIsInput=*/ true))
 			{
@@ -283,7 +286,7 @@ bool UPCGEditorGraphSchema::TryCreateConnectionInternal(UEdGraphPin* InA, UEdGra
 
 		if (!bShouldReconstruct)
 		{
-			for (const FPCGDataTypeBaseId& Id : PCGPinB->Properties.AllowedTypes.GetIds())
+			for (const FPCGDataTypeBaseId& Id : PinBTypes.GetIds())
 			{
 				if (AddConversionNodes(Id, /*bIsInput=*/ false))
 				{
