[PCG] Rename new BP API

#rnx
#virtualized

diff --git a/Engine/Plugins/PCG/Source/PCG/Private/Data/PCGBasePointData.cpp b/Engine/Plugins/PCG/Source/PCG/Private/Data/PCGBasePointData.cpp
--- a/Engine/Plugins/PCG/Source/PCG/Private/Data/PCGBasePointData.cpp
+++ b/Engine/Plugins/PCG/Source/PCG/Private/Data/PCGBasePointData.cpp
@@ -1522,42 +1522,42 @@ TArray<int64> UPCGBasePointData::GetMetadataEntryValuesFromRange(const FPCGPoint
 	return PCGBasePointData::GetInputRangeValues<EPCGPointNativeProperties::MetadataEntry>(InputRange);
 }
 
-void UPCGBasePointData::SetTransformValuesFromRange(FPCGPointOutputRange& OutputRange, const TArray<FTransform>& Values)
+void UPCGBasePointData::SetTransformValuesOnRange(FPCGPointOutputRange& OutputRange, const TArray<FTransform>& Values)
 {
 	PCGBasePointData::SetOutputRangeValues<EPCGPointNativeProperties::Transform>(OutputRange, Values);
 }
 
-void UPCGBasePointData::SetDensityValuesFromRange(UPARAM(ref) FPCGPointOutputRange& OutputRange, const TArray<float>& Values)
+void UPCGBasePointData::SetDensityValuesOnRange(UPARAM(ref) FPCGPointOutputRange& OutputRange, const TArray<float>& Values)
 {
 	PCGBasePointData::SetOutputRangeValues<EPCGPointNativeProperties::Density>(OutputRange, Values);
 }
 
-void UPCGBasePointData::SetBoundsMinValuesFromRange(UPARAM(ref) FPCGPointOutputRange& OutputRange, const TArray<FVector>& Values)
+void UPCGBasePointData::SetBoundsMinValuesOnRange(UPARAM(ref) FPCGPointOutputRange& OutputRange, const TArray<FVector>& Values)
 {
 	PCGBasePointData::SetOutputRangeValues<EPCGPointNativeProperties::BoundsMin>(OutputRange, Values);
 }
 
-void UPCGBasePointData::SetBoundsMaxValuesFromRange(UPARAM(ref) FPCGPointOutputRange& OutputRange, const TArray<FVector>& Values)
+void UPCGBasePointData::SetBoundsMaxValuesOnRange(UPARAM(ref) FPCGPointOutputRange& OutputRange, const TArray<FVector>& Values)
 {
 	PCGBasePointData::SetOutputRangeValues<EPCGPointNativeProperties::BoundsMax>(OutputRange, Values);
 }
 
-void UPCGBasePointData::SetColorValuesFromRange(UPARAM(ref) FPCGPointOutputRange& OutputRange, const TArray<FVector4>& Values)
+void UPCGBasePointData::SetColorValuesOnRange(UPARAM(ref) FPCGPointOutputRange& OutputRange, const TArray<FVector4>& Values)
 {
 	PCGBasePointData::SetOutputRangeValues<EPCGPointNativeProperties::Color>(OutputRange, Values);
 }
 
-void UPCGBasePointData::SetSteepnessValuesFromRange(UPARAM(ref) FPCGPointOutputRange& OutputRange, const TArray<float>& Values)
+void UPCGBasePointData::SetSteepnessValuesOnRange(UPARAM(ref) FPCGPointOutputRange& OutputRange, const TArray<float>& Values)
 {
 	PCGBasePointData::SetOutputRangeValues<EPCGPointNativeProperties::Steepness>(OutputRange, Values);
 }
 
-void UPCGBasePointData::SetSeedValuesFromRange(FPCGPointOutputRange& OutputRange, const TArray<int32>& Values)
+void UPCGBasePointData::SetSeedValuesOnRange(FPCGPointOutputRange& OutputRange, const TArray<int32>& Values)
 {
 	PCGBasePointData::SetOutputRangeValues<EPCGPointNativeProperties::Seed>(OutputRange, Values);
 }
 
-void UPCGBasePointData::SetMetadataEntryValuesFromRange(UPARAM(ref) FPCGPointOutputRange& OutputRange, const TArray<int64>& Values)
+void UPCGBasePointData::SetMetadataEntryValuesOnRange(UPARAM(ref) FPCGPointOutputRange& OutputRange, const TArray<int64>& Values)
 {
 	PCGBasePointData::SetOutputRangeValues<EPCGPointNativeProperties::MetadataEntry>(OutputRange, Values);
 }
@@ -1580,7 +1580,7 @@ FPCGPoint UPCGBasePointData::GetPointFromRange(const FPCGPointInputRange& InputR
 	return PointRanges.GetPoint(Index);
 }
 
-void UPCGBasePointData::SetPointFromRange(UPARAM(ref) FPCGPointOutputRange& OutputRange, int32 Index, const FPCGPoint& Point)
+void UPCGBasePointData::SetPointOnRange(UPARAM(ref) FPCGPointOutputRange& OutputRange, int32 Index, const FPCGPoint& Point)
 {
 	if (!OutputRange.PointData)
 	{
diff --git a/Engine/Plugins/PCG/Source/PCG/Public/Data/PCGBasePointData.h b/Engine/Plugins/PCG/Source/PCG/Public/Data/PCGBasePointData.h
--- a/Engine/Plugins/PCG/Source/PCG/Public/Data/PCGBasePointData.h
+++ b/Engine/Plugins/PCG/Source/PCG/Public/Data/PCGBasePointData.h
@@ -362,34 +362,34 @@ public:
 	static TArray<int64> GetMetadataEntryValuesFromRange(const FPCGPointInputRange& InputRange);
 
 	UFUNCTION(BlueprintCallable, Category = SpatialData)
-	static void SetTransformValuesFromRange(UPARAM(ref) FPCGPointOutputRange& OutputRange, const TArray<FTransform>& Values);
+	static void SetTransformValuesOnRange(UPARAM(ref) FPCGPointOutputRange& OutputRange, const TArray<FTransform>& Values);
 		
 	UFUNCTION(BlueprintCallable, Category = SpatialData)
-	static void SetDensityValuesFromRange(UPARAM(ref) FPCGPointOutputRange& OutputRange, const TArray<float>& Values);
+	static void SetDensityValuesOnRange(UPARAM(ref) FPCGPointOutputRange& OutputRange, const TArray<float>& Values);
 
 	UFUNCTION(BlueprintCallable, Category = SpatialData)
-	static void SetBoundsMinValuesFromRange(UPARAM(ref) FPCGPointOutputRange& OutputRange, const TArray<FVector>& Values);
+	static void SetBoundsMinValuesOnRange(UPARAM(ref) FPCGPointOutputRange& OutputRange, const TArray<FVector>& Values);
 
 	UFUNCTION(BlueprintCallable, Category = SpatialData)
-	static void SetBoundsMaxValuesFromRange(UPARAM(ref) FPCGPointOutputRange& OutputRange, const TArray<FVector>& Values);
+	static void SetBoundsMaxValuesOnRange(UPARAM(ref) FPCGPointOutputRange& OutputRange, const TArray<FVector>& Values);
 
 	UFUNCTION(BlueprintCallable, Category = SpatialData)
-	static void SetColorValuesFromRange(UPARAM(ref) FPCGPointOutputRange& OutputRange, const TArray<FVector4>& Values);
+	static void SetColorValuesOnRange(UPARAM(ref) FPCGPointOutputRange& OutputRange, const TArray<FVector4>& Values);
 
 	UFUNCTION(BlueprintCallable, Category = SpatialData)
-	static void SetSteepnessValuesFromRange(UPARAM(ref) FPCGPointOutputRange& OutputRange, const TArray<float>& Values);
+	static void SetSteepnessValuesOnRange(UPARAM(ref) FPCGPointOutputRange& OutputRange, const TArray<float>& Values);
 
 	UFUNCTION(BlueprintCallable, Category = SpatialData)
-	static void SetSeedValuesFromRange(UPARAM(ref) FPCGPointOutputRange& OutputRange, const TArray<int32>& Values);
+	static void SetSeedValuesOnRange(UPARAM(ref) FPCGPointOutputRange& OutputRange, const TArray<int32>& Values);
 
 	UFUNCTION(BlueprintCallable, Category = SpatialData)
-	static void SetMetadataEntryValuesFromRange(UPARAM(ref) FPCGPointOutputRange& OutputRange, const TArray<int64>& Values);
+	static void SetMetadataEntryValuesOnRange(UPARAM(ref) FPCGPointOutputRange& OutputRange, const TArray<int64>& Values);
 
 	UFUNCTION(BlueprintCallable, Category = SpatialData)
 	static FPCGPoint GetPointFromRange(const FPCGPointInputRange& InputRange, int32 Index);
 
 	UFUNCTION(BlueprintCallable, Category = SpatialData)
-	static void SetPointFromRange(UPARAM(ref) FPCGPointOutputRange& OutputRange, int32 Index, const FPCGPoint& Point);
+	static void SetPointOnRange(UPARAM(ref) FPCGPointOutputRange& OutputRange, int32 Index, const FPCGPoint& Point);
 
 protected:
 
