[PCG] Add 'WithContext' versions of BlueprintCallables that take a FPCGBlueprintContextHandle instead of a FPCGContext reference
- also update SetPointOnRange log (bughawk)

diff --git a/Engine/Plugins/PCG/Source/PCG/Private/Data/PCGBasePointData.cpp b/Engine/Plugins/PCG/Source/PCG/Private/Data/PCGBasePointData.cpp
--- a/Engine/Plugins/PCG/Source/PCG/Private/Data/PCGBasePointData.cpp
+++ b/Engine/Plugins/PCG/Source/PCG/Private/Data/PCGBasePointData.cpp
@@ -1584,13 +1584,13 @@ void UPCGBasePointData::SetPointOnRange(UPARAM(ref) FPCGPointOutputRange& Output
 {
 	if (!OutputRange.PointData)
 	{
-		UE_LOG(LogPCG, Error, TEXT("SetPointFromRange called with invalid point data."));
+		UE_LOG(LogPCG, Error, TEXT("SetPointOnRange called with invalid point data."));
 		return;
 	}
 
 	if (Index < 0 || Index >= OutputRange.PointData->GetNumPoints())
 	{
-		UE_LOG(LogPCG, Error, TEXT("SetPointFromRange called with invalid index '%d' when valid range is [0, %d]."), Index, OutputRange.PointData->GetNumPoints()-1);
+		UE_LOG(LogPCG, Error, TEXT("SetPointOnRange called with invalid index '%d' when valid range is [0, %d]."), Index, OutputRange.PointData->GetNumPoints()-1);
 		return;
 	}
 
diff --git a/Engine/Plugins/PCG/Source/PCG/Private/Data/PCGSpatialData.cpp b/Engine/Plugins/PCG/Source/PCG/Private/Data/PCGSpatialData.cpp
--- a/Engine/Plugins/PCG/Source/PCG/Private/Data/PCGSpatialData.cpp
+++ b/Engine/Plugins/PCG/Source/PCG/Private/Data/PCGSpatialData.cpp
@@ -97,6 +97,20 @@ const UPCGBasePointData* UPCGSpatialData::ToBasePointData(FPCGContext* Context,
 	}
 }
 
+const UPCGBasePointData* UPCGSpatialData::ToBasePointDataWithContext(const FPCGBlueprintContextHandle& ContextHandle) const
+{
+	if (TSharedPtr<FPCGContextHandle> NativeContextHandle = ContextHandle.Handle.Pin())
+	{
+		if (FPCGContext* Context = NativeContextHandle->GetContext())
+		{
+			return ToBasePointData(Context);
+		}
+	}
+
+	UE_LOG(LogPCG, Error, TEXT("UPCGSpatialData::ToBasePointDataWithContext:: Invalid context handle"));
+	return nullptr;
+}
+
 const UPCGPointArrayData* UPCGSpatialData::ToPointArrayData(FPCGContext* Context, const FBox& InBounds) const
 {
 	// @todo_pcg: this is sub-optimal but until we implement this in other spatial types (including the cached types) it works.
diff --git a/Engine/Plugins/PCG/Source/PCG/Private/Helpers/PCGBlueprintHelpers.cpp b/Engine/Plugins/PCG/Source/PCG/Private/Helpers/PCGBlueprintHelpers.cpp
--- a/Engine/Plugins/PCG/Source/PCG/Private/Helpers/PCGBlueprintHelpers.cpp
+++ b/Engine/Plugins/PCG/Source/PCG/Private/Helpers/PCGBlueprintHelpers.cpp
@@ -94,23 +94,79 @@ const UPCGSettings* UPCGBlueprintHelpers::GetSettings(FPCGContext& Context)
 	return Context.GetInputSettings<UPCGSettings>();
 }
 
+const UPCGSettings* UPCGBlueprintHelpers::GetSettingsWithContext(const FPCGBlueprintContextHandle& ContextHandle)
+{
+	if (TSharedPtr<FPCGContextHandle> NativeContextHandle = ContextHandle.Handle.Pin())
+	{
+		if (FPCGContext* Context = NativeContextHandle->GetContext())
+		{
+			return GetSettings(*Context);
+		}
+	}
+
+	UE_LOG(LogPCG, Error, TEXT("UPCGBlueprintHelpers::GetSettingsWithContext : Invalid context handle"));
+	return nullptr;
+}
+
 UPCGData* UPCGBlueprintHelpers::GetActorData(FPCGContext& Context)
 {
 	UPCGComponent* PCGComponent = GetComponent(Context);
 	return PCGComponent ? PCGComponent->GetActorPCGData() : nullptr;
 }
 
+UPCGData* UPCGBlueprintHelpers::GetActorDataWithContext(const FPCGBlueprintContextHandle& ContextHandle)
+{
+	if (TSharedPtr<FPCGContextHandle> NativeContextHandle = ContextHandle.Handle.Pin())
+	{
+		if (FPCGContext* Context = NativeContextHandle->GetContext())
+		{
+			return GetActorData(*Context);
+		}
+	}
+
+	UE_LOG(LogPCG, Error, TEXT("UPCGBlueprintHelpers::GetActorDataWithContext : Invalid context handle"));
+	return nullptr;
+}
+
 UPCGData* UPCGBlueprintHelpers::GetInputData(FPCGContext& Context)
 {
 	UPCGComponent* PCGComponent = GetComponent(Context);
 	return PCGComponent ? PCGComponent->GetInputPCGData() : nullptr;
 }
 
+UPCGData* UPCGBlueprintHelpers::GetInputDataWithContext(const FPCGBlueprintContextHandle& ContextHandle)
+{
+	if (TSharedPtr<FPCGContextHandle> NativeContextHandle = ContextHandle.Handle.Pin())
+	{
+		if (FPCGContext* Context = NativeContextHandle->GetContext())
+		{
+			return GetInputData(*Context);
+		}
+	}
+
+	UE_LOG(LogPCG, Error, TEXT("UPCGBlueprintHelpers::GetInputDataWithContext : Invalid context handle"));
+	return nullptr;
+}
+
 UPCGComponent* UPCGBlueprintHelpers::GetComponent(FPCGContext& Context)
 {
 	return Cast<UPCGComponent>(Context.ExecutionSource.Get());
 }
 
+UPCGComponent* UPCGBlueprintHelpers::GetComponentWithContext(const FPCGBlueprintContextHandle& ContextHandle)
+{
+	if (TSharedPtr<FPCGContextHandle> NativeContextHandle = ContextHandle.Handle.Pin())
+	{
+		if (FPCGContext* Context = NativeContextHandle->GetContext())
+		{
+			return GetComponent(*Context);
+		}
+	}
+
+	UE_LOG(LogPCG, Error, TEXT("UPCGBlueprintHelpers::GetComponentWithContext : Invalid context handle"));
+	return nullptr;
+}
+
 UPCGComponent* UPCGBlueprintHelpers::GetOriginalComponent(FPCGContext& Context)
 {
 	UPCGComponent* SourceComponent = Cast<UPCGComponent>(Context.ExecutionSource.Get());
@@ -127,11 +183,39 @@ UPCGComponent* UPCGBlueprintHelpers::GetOriginalComponent(FPCGContext& Context)
 	}
 }
 
+UPCGComponent* UPCGBlueprintHelpers::GetOriginalComponentWithContext(const FPCGBlueprintContextHandle& ContextHandle)
+{
+	if (TSharedPtr<FPCGContextHandle> NativeContextHandle = ContextHandle.Handle.Pin())
+	{
+		if (FPCGContext* Context = NativeContextHandle->GetContext())
+		{
+			return GetOriginalComponent(*Context);
+		}
+	}
+
+	UE_LOG(LogPCG, Error, TEXT("UPCGBlueprintHelpers::GetOriginalComponentWithContext : Invalid context handle"));
+	return nullptr;
+}
+
 AActor* UPCGBlueprintHelpers::GetTargetActor(FPCGContext& Context, UPCGSpatialData* SpatialData)
 {
 	return Context.GetTargetActor(SpatialData);
 }
 
+AActor* UPCGBlueprintHelpers::GetTargetActorWithContext(const FPCGBlueprintContextHandle& ContextHandle, UPCGSpatialData* SpatialData)
+{
+	if (TSharedPtr<FPCGContextHandle> NativeContextHandle = ContextHandle.Handle.Pin())
+	{
+		if (FPCGContext* Context = NativeContextHandle->GetContext())
+		{
+			return GetTargetActor(*Context, SpatialData);
+		}
+	}
+
+	UE_LOG(LogPCG, Error, TEXT("UPCGBlueprintHelpers::GetTargetActorWithContext : Invalid context handle"));
+	return nullptr;
+}
+
 void UPCGBlueprintHelpers::SetExtents(FPCGPoint& InPoint, const FVector& InExtents)
 {
 	InPoint.SetExtents(InExtents);
@@ -253,6 +337,20 @@ int64 UPCGBlueprintHelpers::GetTaskId(FPCGContext& Context)
 	return static_cast<int64>(Context.TaskId);
 }
 
+int64 UPCGBlueprintHelpers::GetTaskIdWithContext(const FPCGBlueprintContextHandle& ContextHandle)
+{
+	if (TSharedPtr<FPCGContextHandle> NativeContextHandle = ContextHandle.Handle.Pin())
+	{
+		if (FPCGContext* Context = NativeContextHandle->GetContext())
+		{
+			return GetTaskId(*Context);
+		}
+	}
+
+	UE_LOG(LogPCG, Error, TEXT("UPCGBlueprintHelpers::GetTaskIdWithContext : Invalid context handle"));
+	return InvalidPCGTaskId;
+}
+
 bool UPCGBlueprintHelpers::FlushPCGCache()
 {
 	if (UPCGSubsystem* PCGSubsystem = UPCGSubsystem::GetSubsystemForCurrentWorld())
@@ -282,4 +380,18 @@ void UPCGBlueprintHelpers::RefreshPCGRuntimeComponent(UPCGComponent* InComponent
 UPCGData* UPCGBlueprintHelpers::DuplicateData(const UPCGData* InData, FPCGContext& Context, bool bInitializeMetadata)
 {
 	return InData ? InData->DuplicateData(&Context, bInitializeMetadata) : nullptr;
+}
+
+UPCGData* UPCGBlueprintHelpers::DuplicateDataWithContext(const UPCGData* InData, const FPCGBlueprintContextHandle& ContextHandle, bool bInitializeMetadata)
+{
+	if (TSharedPtr<FPCGContextHandle> NativeContextHandle = ContextHandle.Handle.Pin())
+	{
+		if (FPCGContext* Context = NativeContextHandle->GetContext())
+		{
+			return DuplicateData(InData, *Context, bInitializeMetadata);
+		}
+	}
+
+	UE_LOG(LogPCG, Error, TEXT("UPCGBlueprintHelpers::DuplicateDataWithContext : Invalid context handle"));
+	return nullptr;
 }
\ No newline at end of file
diff --git a/Engine/Plugins/PCG/Source/PCG/Public/Data/PCGSpatialData.h b/Engine/Plugins/PCG/Source/PCG/Public/Data/PCGSpatialData.h
--- a/Engine/Plugins/PCG/Source/PCG/Public/Data/PCGSpatialData.h
+++ b/Engine/Plugins/PCG/Source/PCG/Public/Data/PCGSpatialData.h
@@ -14,6 +14,7 @@
 #define UE_API PCG_API
 
 class AActor;
+struct FPCGBlueprintContextHandle;
 struct FPCGContext;
 class UPCGSpatialData;
 class UPCGBasePointData;
@@ -162,9 +163,13 @@ public:
 	UFUNCTION(BlueprintCallable, Category = SpatialData, meta = (DeprecatedFunction, DeprecationMessage = "The To Point Data function is deprecated - use To Point Data With Context instead."))
 	const UPCGPointData* ToPointData() const { return ToPointData(nullptr); }
 
+	/** Prefer using ToBasePointDataWithContext */
 	UFUNCTION(BlueprintCallable, Category = SpatialData)
 	const UPCGPointData* ToPointDataWithContext(UPARAM(ref) FPCGContext& Context) const { return ToPointData(&Context); }
 
+	UFUNCTION(BlueprintCallable, Category = SpatialData)
+	const UPCGBasePointData* ToBasePointDataWithContext(const FPCGBlueprintContextHandle& ContextHandle) const;
+
 	virtual const UPCGPointData* ToPointData(FPCGContext* Context, const FBox& InBounds = FBox(EForceInit::ForceInit)) const PURE_VIRTUAL(UPCGSpatialData::ToPointData, return nullptr;);
 
 	UE_API virtual const UPCGPointArrayData* ToPointArrayData(FPCGContext* Context, const FBox& InBounds = FBox(EForceInit::ForceInit)) const;
diff --git a/Engine/Plugins/PCG/Source/PCG/Public/Helpers/PCGBlueprintHelpers.h b/Engine/Plugins/PCG/Source/PCG/Public/Helpers/PCGBlueprintHelpers.h
--- a/Engine/Plugins/PCG/Source/PCG/Public/Helpers/PCGBlueprintHelpers.h
+++ b/Engine/Plugins/PCG/Source/PCG/Public/Helpers/PCGBlueprintHelpers.h
@@ -16,6 +16,7 @@ class UPCGData;
 class UPCGGraphInterface;
 class UPCGSettings;
 
+struct FPCGBlueprintContextHandle;
 struct FPCGContext;
 struct FPCGLandscapeLayerWeight;
 struct FPCGPoint;
@@ -78,24 +79,48 @@ public:
 	UFUNCTION(BlueprintCallable, Category = "PCG|Helpers", meta = (ScriptMethod))
 	static PCG_API FRandomStream GetRandomStreamFromTwoPoints(const FPCGPoint& InPointA, const FPCGPoint& InPointB, const UPCGSettings* OptionalSettings = nullptr, const UPCGComponent* OptionalComponent = nullptr);
 
+	/** Prefer using GetSettingsWithContext */
 	UFUNCTION(BlueprintCallable, Category = "PCG|Helpers", meta = (ScriptMethod))
 	static PCG_API const UPCGSettings* GetSettings(UPARAM(ref) FPCGContext& Context);
 
+	UFUNCTION(BlueprintCallable, Category = "PCG|Helpers", meta = (ScriptMethod))
+	static PCG_API const UPCGSettings* GetSettingsWithContext(const FPCGBlueprintContextHandle& ContextHandle);
+
+	/** Prefer using GetActorDataWithContext */
 	UFUNCTION(BlueprintCallable, Category = "PCG|Temporary", meta = (ScriptMethod))
 	static PCG_API UPCGData* GetActorData(UPARAM(ref) FPCGContext& Context);
 
+	UFUNCTION(BlueprintCallable, Category = "PCG|Helpers", meta = (ScriptMethod))
+	static PCG_API UPCGData* GetActorDataWithContext(const FPCGBlueprintContextHandle& ContextHandle);
+
+	/** Prefer using GetInputDataWithContext */
 	UFUNCTION(BlueprintCallable, Category = "PCG|Temporary", meta = (ScriptMethod))
 	static PCG_API UPCGData* GetInputData(UPARAM(ref) FPCGContext& Context);
 
+	UFUNCTION(BlueprintCallable, Category = "PCG|Helpers", meta = (ScriptMethod))
+	static PCG_API UPCGData* GetInputDataWithContext(const FPCGBlueprintContextHandle& ContextHandle);
+
+	/** Prefer using GetComponentWithContext */
 	UFUNCTION(BlueprintCallable, Category = "PCG|Temporary", meta = (ScriptMethod))
 	static PCG_API UPCGComponent* GetComponent(UPARAM(ref) FPCGContext& Context);
 
+	UFUNCTION(BlueprintCallable, Category = "PCG|Helpers", meta = (ScriptMethod))
+	static PCG_API UPCGComponent* GetComponentWithContext(const FPCGBlueprintContextHandle& ContextHandle);
+
+	/** Prefer using GetOriginalComponentWithContext */
 	UFUNCTION(BlueprintCallable, Category = "PCG|Temporary", meta = (ScriptMethod))
 	static PCG_API UPCGComponent* GetOriginalComponent(UPARAM(ref) FPCGContext& Context);
 
-	UFUNCTION(BlueprintCallable, Category = "PCG", meta = (ScriptMethod))
+	UFUNCTION(BlueprintCallable, Category = "PCG|Helpers", meta = (ScriptMethod))
+	static PCG_API UPCGComponent* GetOriginalComponentWithContext(const FPCGBlueprintContextHandle& ContextHandle);
+
+	/** Prefer using GetTargetActorWithContext */
+	UFUNCTION(BlueprintCallable, Category = "PCG|Temporary", meta = (ScriptMethod))
 	static PCG_API AActor* GetTargetActor(UPARAM(ref) FPCGContext& Context, UPCGSpatialData* SpatialData);
 
+	UFUNCTION(BlueprintCallable, Category = "PCG|Helpers", meta = (ScriptMethod))
+	static PCG_API AActor* GetTargetActorWithContext(const FPCGBlueprintContextHandle& ContextHandle, UPCGSpatialData* SpatialData);
+
 	UFUNCTION(BlueprintCallable, Category = "PCG|Points", meta = (ScriptMethod))
 	static PCG_API void SetExtents(UPARAM(ref) FPCGPoint& InPoint, const FVector& InExtents);
 
@@ -123,9 +148,13 @@ public:
 	UFUNCTION(BlueprintCallable, Category = "PCG|Helpers", meta = (WorldContext="WorldContextObject"))
 	static PCG_API TArray<FPCGLandscapeLayerWeight> GetInterpolatedPCGLandscapeLayerWeights(UObject* WorldContextObject, const FVector& Location);
 
-	UFUNCTION(BlueprintCallable, Category = "PCG|Helpers", meta = (ScriptMethod))
+	/** Prefer using GetTaskIdWithContext */
+	UFUNCTION(BlueprintCallable, Category = "PCG|Temporary", meta = (ScriptMethod))
 	static PCG_API int64 GetTaskId(UPARAM(ref) FPCGContext& Context);
 
+	UFUNCTION(BlueprintCallable, Category = "PCG|Helpers", meta = (ScriptMethod))
+	static PCG_API int64 GetTaskIdWithContext(const FPCGBlueprintContextHandle& ContextHandle);
+
 	/** Flush the cache, to be used if you have changed something PCG depends on at runtime. Same as `pcg.FlushCache` command. Returns true if it succeeded. */
 	UFUNCTION(BlueprintCallable, Category = "PCG", meta=(DisplayName = "Flush PCG Cache"))
 	static PCG_API bool FlushPCGCache();
@@ -134,8 +163,12 @@ public:
 	UFUNCTION(BlueprintCallable, Category = "PCG|Runtime", meta = (ScriptMethod, DisplayName = "Refresh PCG Runtime Component"))
 	static PCG_API void RefreshPCGRuntimeComponent(UPCGComponent* InComponent, const bool bFlushCache = false);
 
+	/** Prefer using DuplicateDataWithContext */
+	UFUNCTION(BlueprintCallable, Category="PCG|Temporary", meta = (ScriptMethod))
+	static PCG_API UPCGData* DuplicateData(const UPCGData* InData, UPARAM(ref) FPCGContext& Context, bool bInitializeMetadata = true);
+
 	// Implementation note: Needs to be done outside of UPCGData because of circular dependency between PCGContext.h and PCGData.h
 	/** Return a copy of the data, with Metadata inheritance for spatial data. */
-	UFUNCTION(BlueprintCallable, Category="PCG|Data", meta = (ScriptMethod))
-	static PCG_API UPCGData* DuplicateData(const UPCGData* InData, UPARAM(ref) FPCGContext& Context, bool bInitializeMetadata = true);
+	UFUNCTION(BlueprintCallable, Category = "PCG|Data", meta = (ScriptMethod))
+	static PCG_API UPCGData* DuplicateDataWithContext(const UPCGData* InData, const FPCGBlueprintContextHandle& ContextHandle, bool bInitializeMetadata = true);
 };
