[PCG] Fix GPU generation failing for non-partitioned components

#rnx

diff --git a/Engine/Plugins/PCG/Source/PCG/Private/PCGComponent.cpp b/Engine/Plugins/PCG/Source/PCG/Private/PCGComponent.cpp
--- a/Engine/Plugins/PCG/Source/PCG/Private/PCGComponent.cpp
+++ b/Engine/Plugins/PCG/Source/PCG/Private/PCGComponent.cpp
@@ -4624,7 +4624,7 @@ FPCGGridDescriptor FPCGComponentExecutionState::GetGridDescriptor(uint32 InGridS
 uint32 FPCGComponentExecutionState::GetGenerationGridSize() const
 {
 	check(Component);
-	return Component->GetGenerationGridSize();
+	return (IsPartitioned() || IsLocalSource()) ? Component->GetGenerationGridSize() : PCGHiGenGrid::UninitializedGridSize();
 }
 
 FPCGTaskId FPCGComponentExecutionState::GetGenerationTaskId() const
