[PCG] Fix compute graph element execution hanging for static mesh spawner mesh entries with 0 weight

#rnx

diff --git a/Engine/Plugins/PCG/Source/PCG/Private/Compute/DataInterfaces/Elements/PCGStaticMeshSpawnerDataInterface.cpp b/Engine/Plugins/PCG/Source/PCG/Private/Compute/DataInterfaces/Elements/PCGStaticMeshSpawnerDataInterface.cpp
--- a/Engine/Plugins/PCG/Source/PCG/Private/Compute/DataInterfaces/Elements/PCGStaticMeshSpawnerDataInterface.cpp
+++ b/Engine/Plugins/PCG/Source/PCG/Private/Compute/DataInterfaces/Elements/PCGStaticMeshSpawnerDataInterface.cpp
@@ -678,7 +678,7 @@ void UPCGStaticMeshSpawnerDataProvider::CreatePrimitiveDescriptors(FPCGContext*
 			TotalWeight += Entry.Weight;
 		}
 
-		if (!ensure(TotalWeight > UE_SMALL_NUMBER))
+		if (TotalWeight < UE_SMALL_NUMBER)
 		{
 			return;
 		}
@@ -690,6 +690,14 @@ void UPCGStaticMeshSpawnerDataProvider::CreatePrimitiveDescriptors(FPCGContext*
 		for (const FPCGMeshSelectorWeightedEntry& Entry : SelectorWeighted->MeshEntries)
 		{
 			if (Entry.Descriptor.StaticMesh.IsNull())
+			{
+				PCG_KERNEL_VALIDATION_ERR(InContext, Settings, FText::Format(PCGStaticMeshSpawnerDataInterface::CouldNotLoadStaticMeshFormat, FText::FromString(Entry.Descriptor.StaticMesh.ToString())));
+				continue;
+			}
+
+			const float Weight = float(Entry.Weight) / TotalWeight;
+
+			if (Weight < UE_SMALL_NUMBER)
 			{
 				continue;
 			}
@@ -702,45 +710,36 @@ void UPCGStaticMeshSpawnerDataProvider::CreatePrimitiveDescriptors(FPCGContext*
 				break;
 			}
 
-			if (!Entry.Descriptor.StaticMesh.IsNull())
-			{
-				const float Weight = float(Entry.Weight) / TotalWeight;
-				CumulativeWeight += Weight;
-				PrimitiveSelectionCDF.Add(CumulativeWeight);
-				PrimitiveStringKeys.Add(InBinding->GetStringTable().IndexOfByKey(Entry.Descriptor.StaticMesh.ToString()));
+			CumulativeWeight += Weight;
+			PrimitiveSelectionCDF.Add(CumulativeWeight);
+			PrimitiveStringKeys.Add(InBinding->GetStringTable().IndexOfByKey(Entry.Descriptor.StaticMesh.ToString()));
 
-				FPCGProceduralISMComponentDescriptor Descriptor;
-				Descriptor = Entry.Descriptor;
-				Descriptor.WorldBounds = InContext->ExecutionSource->GetExecutionState().GetBounds();
-				Descriptor.NumCustomFloats = CustomFloatCount;
-				Descriptor.StaticMesh = Entry.Descriptor.StaticMesh;
+			FPCGProceduralISMComponentDescriptor Descriptor;
+			Descriptor = Entry.Descriptor;
+			Descriptor.WorldBounds = InContext->ExecutionSource->GetExecutionState().GetBounds();
+			Descriptor.NumCustomFloats = CustomFloatCount;
+			Descriptor.StaticMesh = Entry.Descriptor.StaticMesh;
 
-				int32 InstanceCount = FMath::CeilToInt(TotalInputPointCount * Weight);
+			int32 InstanceCount = FMath::CeilToInt(TotalInputPointCount * Weight);
 
-				if (SelectorWeighted->MeshEntries.Num() > 1)
-				{
-					// Since we'll be selecting meshes based on random draws using the point random seeds which we don't have on CPU,
-					// we may pick more or less than the expected number of instances for each mesh. Use binomial variance to calculate
-					// the overallocation that gives 99.7% confidence (3 sigma).
-					const double Variance = double(TotalInputPointCount) * Weight * (1.0f - Weight);
-					const double Sigma = FMath::Sqrt(Variance);
-					const uint32 AdditionalAllocation = FMath::CeilToInt(3.0f * Sigma);
+			if (SelectorWeighted->MeshEntries.Num() > 1)
+			{
+				// Since we'll be selecting meshes based on random draws using the point random seeds which we don't have on CPU,
+				// we may pick more or less than the expected number of instances for each mesh. Use binomial variance to calculate
+				// the overallocation that gives 99.7% confidence (3 sigma).
+				const double Variance = double(TotalInputPointCount) * Weight * (1.0f - Weight);
+				const double Sigma = FMath::Sqrt(Variance);
+				const uint32 AdditionalAllocation = FMath::CeilToInt(3.0f * Sigma);
 
-					// Useful for debugging instance counts.
-					//UE_LOG(LogPCG, Warning, TEXT("Over allocation, N = %d, P = %.2f, NP = %d, 3sig = %f, Extra = %d (%.2f%%), Calculated = %d, Final = %d"), TotalInputPointCount, Weight, InstanceCount, float(3.0f * Sigma), AdditionalAllocation, float(AdditionalAllocation) / TotalInputPointCount, InstanceCount + AdditionalAllocation, FMath::Min(InstanceCount + AdditionalAllocation, TotalInputPointCount));
+				// Useful for debugging instance counts.
+				//UE_LOG(LogPCG, Warning, TEXT("Over allocation, N = %d, P = %.2f, NP = %d, 3sig = %f, Extra = %d (%.2f%%), Calculated = %d, Final = %d"), TotalInputPointCount, Weight, InstanceCount, float(3.0f * Sigma), AdditionalAllocation, float(AdditionalAllocation) / TotalInputPointCount, InstanceCount + AdditionalAllocation, FMath::Min(InstanceCount + AdditionalAllocation, TotalInputPointCount));
 
-					InstanceCount += AdditionalAllocation;
-				}
+				InstanceCount += AdditionalAllocation;
+			}
 
-				Descriptor.NumInstances = FMath::Min(InstanceCount, TotalInputPointCount);
+			Descriptor.NumInstances = FMath::Min(InstanceCount, TotalInputPointCount);
 
-				PrimitiveDescriptors.Add(MoveTemp(Descriptor));
-			}
-			else
-			{
-				PCG_KERNEL_VALIDATION_ERR(InContext, Settings,
-					FText::Format(PCGStaticMeshSpawnerDataInterface::CouldNotLoadStaticMeshFormat, FText::FromString(Entry.Descriptor.StaticMesh.ToString())));
-			}
+			PrimitiveDescriptors.Add(MoveTemp(Descriptor));
 		}
 	}
 
