[PCG] Bughawk fixes

#rnx

diff --git a/Engine/Plugins/PCG/Source/PCG/Private/RuntimeGen/PCGRuntimeGenScheduler.cpp b/Engine/Plugins/PCG/Source/PCG/Private/RuntimeGen/PCGRuntimeGenScheduler.cpp
--- a/Engine/Plugins/PCG/Source/PCG/Private/RuntimeGen/PCGRuntimeGenScheduler.cpp
+++ b/Engine/Plugins/PCG/Source/PCG/Private/RuntimeGen/PCGRuntimeGenScheduler.cpp
@@ -627,10 +627,10 @@ void FPCGRuntimeGenScheduler::TickQueueComponentsForGeneration(
 					// Poll for streaming to complete on the unbounded grid.
 					if (bCheckStreaming && Policy->DoesGridDependOnWorldStreaming(PCGHiGenGrid::UnboundedGridSize()))
 					{
-						FVector StreamingLocation, StreamingExtents;
-						OriginalComponentBounds.GetCenterAndExtents(StreamingLocation, StreamingExtents);
+						const FVector StreamingLocation = OriginalComponentBounds.GetCenter();
+						const FVector StreamingSize = OriginalComponentBounds.GetSize();
 
-						if (!IsWorldStreamingComplete(StreamingLocation, FMath::Max(StreamingExtents.X, StreamingExtents.Y)))
+						if (!IsWorldStreamingComplete(StreamingLocation, FMath::Max(StreamingSize.X, StreamingSize.Y)))
 						{
 							UE_LOG(LogPCG, VeryVerbose, TEXT("Partitioned original component '%s' rejected as world is not fully loaded."), *OriginalComponent->GetOwner()->GetActorNameOrLabel());
 							continue;
@@ -788,10 +788,10 @@ void FPCGRuntimeGenScheduler::TickQueueComponentsForGeneration(
 			// Poll for streaming to complete.
 			if (bCheckStreaming && Policy->DoesGridDependOnWorldStreaming(PCGHiGenGrid::UnboundedGridSize()))
 			{
-				FVector StreamingLocation, StreamingExtents;
-				OriginalComponentBounds.GetCenterAndExtents(StreamingLocation, StreamingExtents);
+				const FVector StreamingLocation = OriginalComponentBounds.GetCenter();
+				const FVector StreamingSize = OriginalComponentBounds.GetSize();
 
-				if (!IsWorldStreamingComplete(StreamingLocation, FMath::Max(StreamingExtents.X, StreamingExtents.Y)))
+				if (!IsWorldStreamingComplete(StreamingLocation, FMath::Max(StreamingSize.X, StreamingSize.Y)))
 				{
 					UE_LOG(LogPCG, VeryVerbose, TEXT("Non-partitioned original component '%s' rejected as world is not fully loaded."), *OriginalComponent->GetOwner()->GetActorNameOrLabel());
 					continue;
