[PCG] Fix crashes when spawning an absurd number of instances with GPU execution

#rnx

diff --git a/Engine/Plugins/PCG/Source/PCG/Private/Compute/DataInterfaces/Elements/PCGSkinnedMeshSpawnerDataInterface.cpp b/Engine/Plugins/PCG/Source/PCG/Private/Compute/DataInterfaces/Elements/PCGSkinnedMeshSpawnerDataInterface.cpp
--- a/Engine/Plugins/PCG/Source/PCG/Private/Compute/DataInterfaces/Elements/PCGSkinnedMeshSpawnerDataInterface.cpp
+++ b/Engine/Plugins/PCG/Source/PCG/Private/Compute/DataInterfaces/Elements/PCGSkinnedMeshSpawnerDataInterface.cpp
@@ -476,13 +476,18 @@ void UPCGSkinnedMeshSpawnerDataProvider::CreatePrimitiveDescriptors(FPCGContext*
 		return;
 	}
 
+	if (TotalInputPointCount >= MAX_INSTANCE_ID)
+	{
+		PCG_KERNEL_VALIDATION_ERR(InContext, Settings, FText::Format(LOCTEXT("TooManyInstances", "Tried to spawn too many instances ({0}), procedural ISM component creation skipped and instances will not be rendered."), TotalInputPointCount));
+		return;
+	}
+
 	UPCGComponent* SourceComponent = Cast<UPCGComponent>(InContext->ExecutionSource.Get());
 	if (!SourceComponent)
 	{
 		return;
 	}
 
-	
 	if (Settings->InstanceDataPackerParameters && Settings->InstanceDataPackerParameters->GetAttributeNames(/*OutNames=*/nullptr))
 	{
 		TArray<FName> AttributeNames;
@@ -605,22 +610,6 @@ void UPCGSkinnedMeshSpawnerDataProvider::CreatePrimitiveDescriptors(FPCGContext*
 		PrimitiveSelectionCDF.SetNumZeroed(PrimitiveDescriptors.Num());
 	}
 
-	// Validate instance count is not too large.
-	uint64 TotalInstanceCount = 0;
-	for (const FPCGSoftSkinnedMeshComponentDescriptor& Desc : PrimitiveDescriptors)
-	{
-		TotalInstanceCount += Desc.NumInstancesGPUOnly;
-	}
-
-	if (TotalInstanceCount >= MAX_INSTANCE_ID)
-	{
-		PCG_KERNEL_VALIDATION_ERR(InContext, Settings, FText::Format(
-			LOCTEXT("TooManyInstances", "Tried to spawn too many instances ({0}), procedural ISM component creation skipped and instances will not be rendered."),
-			FText::FromString(FString::FromInt(TotalInstanceCount))));
-
-		return;
-	}
-
 	FPCGSpawnerPrimitives& Primitives = InBinding->MeshSpawnersToPrimitives.FindOrAdd(ProducerKernel);
 	Primitives.NumCustomFloats = CustomFloatCount;
 	Primitives.AttributeIdOffsetStrides = AttributeIdOffsetStrides;
diff --git a/Engine/Plugins/PCG/Source/PCG/Private/Compute/DataInterfaces/Elements/PCGStaticMeshSpawnerDataInterface.cpp b/Engine/Plugins/PCG/Source/PCG/Private/Compute/DataInterfaces/Elements/PCGStaticMeshSpawnerDataInterface.cpp
--- a/Engine/Plugins/PCG/Source/PCG/Private/Compute/DataInterfaces/Elements/PCGStaticMeshSpawnerDataInterface.cpp
+++ b/Engine/Plugins/PCG/Source/PCG/Private/Compute/DataInterfaces/Elements/PCGStaticMeshSpawnerDataInterface.cpp
@@ -540,6 +540,12 @@ void UPCGStaticMeshSpawnerDataProvider::CreatePrimitiveDescriptors(FPCGContext*
 		return;
 	}
 
+	if (TotalInputPointCount >= MAX_INSTANCE_ID)
+	{
+		PCG_KERNEL_VALIDATION_ERR(InContext, Settings, FText::Format(LOCTEXT("TooManyInstances", "Tried to spawn too many instances ({0}), procedural ISM component creation skipped and instances will not be rendered."), TotalInputPointCount));
+		return;
+	}
+
 	if (!InContext->ExecutionSource.IsValid())
 	{
 		return;
@@ -743,22 +749,6 @@ void UPCGStaticMeshSpawnerDataProvider::CreatePrimitiveDescriptors(FPCGContext*
 		}
 	}
 
-	// Validate instance count is not too large.
-	uint64 TotalInstanceCount = 0;
-	for (const FPCGProceduralISMComponentDescriptor& Desc : PrimitiveDescriptors)
-	{
-		TotalInstanceCount += Desc.NumInstances;
-	}
-
-	if (TotalInstanceCount >= MAX_INSTANCE_ID)
-	{
-		PCG_KERNEL_VALIDATION_ERR(InContext, Settings, FText::Format(
-			LOCTEXT("TooManyInstances", "Tried to spawn too many instances ({0}), procedural ISM component creation skipped and instances will not be rendered."),
-			FText::FromString(FString::FromInt(TotalInstanceCount))));
-
-		return;
-	}
-
 	FPCGSpawnerPrimitives& Primitives = InBinding->MeshSpawnersToPrimitives.FindOrAdd(ProducerKernel);
 	Primitives.NumCustomFloats = CustomFloatCount;
 	Primitives.AttributeIdOffsetStrides = AttributeIdOffsetStrides;
diff --git a/Engine/Plugins/PCG/Source/PCG/Private/Compute/DataInterfaces/PCGDataCollectionDataInterface.cpp b/Engine/Plugins/PCG/Source/PCG/Private/Compute/DataInterfaces/PCGDataCollectionDataInterface.cpp
--- a/Engine/Plugins/PCG/Source/PCG/Private/Compute/DataInterfaces/PCGDataCollectionDataInterface.cpp
+++ b/Engine/Plugins/PCG/Source/PCG/Private/Compute/DataInterfaces/PCGDataCollectionDataInterface.cpp
@@ -575,9 +575,11 @@ bool UPCGDataCollectionDataProvider::PrepareForExecute_GameThread(UPCGDataBindin
 
 	if (!bBufferSizeValidated)
 	{
-		if (PCGComputeHelpers::IsBufferSizeTooLarge(PinDesc->ComputePackedSizeBytes()))
+		SizeBytes = PinDesc->ComputePackedSizeBytes();
+		bIsBufferSizeValid = !PCGComputeHelpers::IsBufferSizeTooLarge(SizeBytes);
+
+		if (!bIsBufferSizeValid)
 		{
-			PinDesc = nullptr;
 			return true;
 		}
 
@@ -630,6 +632,8 @@ FComputeDataProviderRenderProxy* UPCGDataCollectionDataProvider::GetRenderProxy(
 	{
 		.DataProvider = this,
 		.PinDesc = PinDesc,
+		.SizeBytes = SizeBytes,
+		.bIsBufferSizeValid = bIsBufferSizeValid,
 		.ExportMode = ExportMode,
 		.bZeroInitialize = bRequiresZeroInitialization,
 		.ElementCountMultiplier = ElementCountMultiplier,
@@ -646,6 +650,8 @@ void UPCGDataCollectionDataProvider::Reset()
 {
 	Super::Reset();
 
+	SizeBytes = 0;
+	bIsBufferSizeValid = false;
 	ElementCountMultiplier = 0;
 	bRequiresZeroInitialization = false;
 	DataAdaptor.Reset();
@@ -653,7 +659,9 @@ void UPCGDataCollectionDataProvider::Reset()
 }
 
 FPCGDataCollectionDataProviderProxy::FPCGDataCollectionDataProviderProxy(const FSetupParams& InParams)
-	: ExportMode(InParams.ExportMode)
+	: SizeBytes(InParams.SizeBytes)
+	, bIsBufferSizeValid(InParams.bIsBufferSizeValid)
+	, ExportMode(InParams.ExportMode)
 	, bZeroInitialize(InParams.bZeroInitialize)
 	, ElementCountMultiplier(InParams.ElementCountMultiplier)
 	, DataProviderWeakPtr(InParams.DataProvider)
@@ -664,7 +672,7 @@ FPCGDataCollectionDataProviderProxy::FPCGDataCollectionDataProviderProxy(const F
 {
 	OriginatingGenerationCount = DataProviderWeakPtr->GenerationCounter;
 
-	if (ensure(InParams.PinDesc))
+	if (InParams.PinDesc)
 	{
 		PinDesc = InParams.PinDesc;
 
@@ -675,8 +683,6 @@ FPCGDataCollectionDataProviderProxy::FPCGDataCollectionDataProviderProxy(const F
 				LargestAttributeIndex = FMath::Max(LargestAttributeIndex, PCGComputeHelpers::GetMetadataAttributeIndexFromAttributeId(AttrDesc.AttributeId));
 			}
 		}
-
-		SizeBytes = PinDesc->ComputePackedSizeBytes();
 	}
 }
 
@@ -687,12 +693,18 @@ bool FPCGDataCollectionDataProviderProxy::IsValid(FValidationData const& InValid
 		return false;
 	}
 
-	if (SizeBytes <= 0)
+	if (!bIsBufferSizeValid)
 	{
 		UE_LOG(LogPCG, Error, TEXT("FPCGDataCollectionDataProviderProxy invalid due to invalid buffer size."));
 		return false;
 	}
 
+	if (!PinDesc)
+	{
+		UE_LOG(LogPCG, Error, TEXT("FPCGDataCollectionDataProviderProxy invalid due to missing pin data description."));
+		return false;
+	}
+
 	if (bProducedByCPU && !DataAdaptor.IsValid())
 	{
 		UE_LOG(LogPCG, Warning, TEXT("FPCGDataProviderDataCollectionProxy invalid due to null DataAdaptor (and bRequiresUpload=true)."));
diff --git a/Engine/Plugins/PCG/Source/PCG/Private/Compute/DataInterfaces/PCGDataCollectionDataInterface.h b/Engine/Plugins/PCG/Source/PCG/Private/Compute/DataInterfaces/PCGDataCollectionDataInterface.h
--- a/Engine/Plugins/PCG/Source/PCG/Private/Compute/DataInterfaces/PCGDataCollectionDataInterface.h
+++ b/Engine/Plugins/PCG/Source/PCG/Private/Compute/DataInterfaces/PCGDataCollectionDataInterface.h
@@ -87,6 +87,10 @@ public:
 	bool bRequiresZeroInitialization = false;
 
 protected:
+	uint64 SizeBytes = 0;
+
+	bool bIsBufferSizeValid = false;
+
 	bool bBufferSizeValidated = false;
 
 	TSharedPtr<FPCGDataUploadAdaptor> DataAdaptor;
@@ -100,6 +104,8 @@ public:
 		// Weak pointer useful for passing back buffer handles. Do not access this directly from the render thread.
 		TWeakObjectPtr<UPCGDataCollectionDataProvider> DataProvider;
 		TSharedPtr<const FPCGDataCollectionDesc> PinDesc = nullptr;
+		uint64 SizeBytes = 0;
+		bool bIsBufferSizeValid = false;
 		EPCGExportMode ExportMode = EPCGExportMode::NoExport;
 		bool bZeroInitialize = false;
 		uint32 ElementCountMultiplier = 0;
@@ -127,7 +133,9 @@ public:
 	uint64 OriginatingGenerationCount = 0;
 
 protected:
-	uint32 SizeBytes = 0;
+	uint64 SizeBytes = 0;
+
+	bool bIsBufferSizeValid = false;
 
 	EPCGExportMode ExportMode = EPCGExportMode::NoExport;
 
diff --git a/Engine/Plugins/PCG/Source/PCG/Private/Compute/DataInterfaces/PCGTextureDataInterface.cpp b/Engine/Plugins/PCG/Source/PCG/Private/Compute/DataInterfaces/PCGTextureDataInterface.cpp
--- a/Engine/Plugins/PCG/Source/PCG/Private/Compute/DataInterfaces/PCGTextureDataInterface.cpp
+++ b/Engine/Plugins/PCG/Source/PCG/Private/Compute/DataInterfaces/PCGTextureDataInterface.cpp
@@ -280,12 +280,6 @@ bool UPCGTextureDataProvider::PrepareForExecute_GameThread(UPCGDataBinding* InBi
 		return false;
 	}
 
-	// @todo_pcg: Check if the texture size is too large
-	//if (PCGComputeHelpers::IsBufferSizeTooLarge(PinDesc->ComputePackedSizeBytes()))
-	//{
-	//	PinDesc = nullptr;
-	//}
-
 	if (bInitializeFromDataCollection)
 	{
 		BuildInfosFromDataCollection(InBinding);
