PCG Editor Mode
- Tool Working Data architecture; data used for processing managed by the PCGComponent
- "Get Tool Data" PCG Node that acts as 'Tool Data Instance'/tool parameter definition
- Paint Mode

diff --git a/Engine/Plugins/PCG/PCG.uplugin b/Engine/Plugins/PCG/PCG.uplugin
--- a/Engine/Plugins/PCG/PCG.uplugin
+++ b/Engine/Plugins/PCG/PCG.uplugin
@@ -43,6 +43,10 @@
 		{
 			"Name": "GeometryProcessing",
 			"Enabled": true
+		},
+		{
+			"Name": "MeshModelingToolset",
+			"Enabled": true
 		}
 	]
 }
diff --git a/Engine/Plugins/PCG/Source/PCG/PCG.Build.cs b/Engine/Plugins/PCG/Source/PCG/PCG.Build.cs
--- a/Engine/Plugins/PCG/Source/PCG/PCG.Build.cs
+++ b/Engine/Plugins/PCG/Source/PCG/PCG.Build.cs
@@ -45,6 +45,7 @@ namespace UnrealBuildTool.Rules
 						"UnrealEd",
 						"Settings",
 						"SourceControl",
+						"InteractiveToolsFramework"
 					});
 			}
 
diff --git a/Engine/Plugins/PCG/Source/PCG/Private/Data/PCGPointArrayData.cpp b/Engine/Plugins/PCG/Source/PCG/Private/Data/PCGPointArrayData.cpp
--- a/Engine/Plugins/PCG/Source/PCG/Private/Data/PCGPointArrayData.cpp
+++ b/Engine/Plugins/PCG/Source/PCG/Private/Data/PCGPointArrayData.cpp
@@ -208,6 +208,16 @@ void UPCGPointArrayData::SetNumPoints(int32 InNumPoints, bool bInitializeValues)
 	}
 }
 
+void UPCGPointArrayData::RemoveAt(int32 Index)
+{
+	if(Index >= 0 && Index < PointArray.GetNumPoints())
+	{
+		FlattenPropertiesIfNeeded();
+		PointArray.RemoveAt(Index);
+		DirtyCache();
+	}
+}
+
 void UPCGPointArrayData::AllocateProperties(EPCGPointNativeProperties Properties)
 {
 	FlattenPropertiesIfNeeded(Properties);
diff --git a/Engine/Plugins/PCG/Source/PCG/Private/Editor/PCGGraphToolData.cpp b/Engine/Plugins/PCG/Source/PCG/Private/Editor/PCGGraphToolData.cpp
new file mode 100644
--- /dev/null
+++ b/Engine/Plugins/PCG/Source/PCG/Private/Editor/PCGGraphToolData.cpp
@@ -0,0 +1,356 @@
+﻿// Copyright Epic Games, Inc. All Rights Reserved.
+
+#include "Editor/PCGGraphToolData.h"
+#include "PCGComponent.h"
+#include "PCGContext.h"
+#include "Data/PCGPointArrayData.h"
+#include "Helpers/PCGDynamicTrackingHelpers.h"
+
+#include "Engine/Engine.h"
+#include "Components/SplineComponent.h"
+#include "Logging/StructuredLog.h"
+#if WITH_EDITOR
+#include "InteractiveTool.h"
+#endif
+
+#if WITH_EDITORONLY_DATA
+void FPCGInteractiveToolWorkingData::Initialize(const FPCGInteractiveToolWorkingDataContext& Context)
+{
+	ensureMsgf(Context.OwningPCGComponent.IsValid() && Context.OwningActor.IsValid(), TEXT("Owning PCG Component or Owning Actor invalid! This should never happen so an assumption was hurt."));
+
+	InitializeInternal(Context);
+	
+	OnInitializedDelegate.ExecuteIfBound(this);
+
+	bInitialized = true;
+}
+
+void FPCGInteractiveToolWorkingData::InitializeInternal(const FPCGInteractiveToolWorkingDataContext& Context)
+{
+}
+
+void FPCGInteractiveToolWorkingData::OnToolStart(const FPCGInteractiveToolWorkingDataContext& Context)
+{
+}
+
+void FPCGInteractiveToolWorkingData::OnToolApply(const FPCGInteractiveToolWorkingDataContext& Context)
+{
+	GeneratedResources.Finalize();
+}
+
+void FPCGInteractiveToolWorkingData::OnToolCancel(const FPCGInteractiveToolWorkingDataContext& Context)
+{
+	GeneratedResources.Cleanup();
+}
+
+void FPCGInteractiveToolWorkingDataGeneratedResources::Finalize()
+{
+	for(TWeakObjectPtr<UActorComponent> GeneratedComponent : GeneratedComponents)
+	{
+		if(GeneratedComponent.IsValid())
+		{
+			GeneratedComponent->ClearFlags(RF_Transient);
+		}
+	}
+}
+
+void FPCGInteractiveToolWorkingDataGeneratedResources::Cleanup()
+{
+	for(TWeakObjectPtr<UActorComponent> GeneratedComponent : GeneratedComponents)
+	{
+		if(GeneratedComponent.IsValid())
+		{
+			if(UPCGComponent* PCGComponent = Cast<UPCGComponent>(GeneratedComponent))
+			{
+				PCGComponent->Cleanup(true);	
+			}
+			
+			GeneratedComponent->DestroyComponent();
+		}
+	}
+}
+
+FPCGInteractiveToolWorkingDataContext::FPCGInteractiveToolWorkingDataContext(UPCGComponent* InOwningComponent, UInteractiveTool* InInteractiveTool,	UInteractiveToolPropertySet* InPropertySet, FName InDataInstanceIdentifier)
+{
+	OwningPCGComponent = InOwningComponent;
+	OwningActor = OwningPCGComponent->GetOwner();
+	DataInstanceIdentifier = InDataInstanceIdentifier;
+	InteractiveTool = InInteractiveTool;
+	PCGSettings = InPropertySet;
+}
+#endif
+
+bool FPCGInteractiveToolWorkingData::IsValid() const
+{
+	return true;
+}
+
+#if WITH_EDITORONLY_DATA
+void FPCGInteractiveToolWorkingData_PointArrayData::InitializeInternal(const FPCGInteractiveToolWorkingDataContext& Context)
+{
+	GeneratedPointData = NewObject<UPCGPointArrayData>(Context.OwningPCGComponent.Get());
+}
+
+void FPCGInteractiveToolWorkingData_PointArrayData::OnToolStart(const FPCGInteractiveToolWorkingDataContext& Context)
+{
+	OnToolStartPointArray = Cast<UPCGPointArrayData>(StaticDuplicateObject(GeneratedPointData, GetTransientPackage()));
+	
+	FPCGInteractiveToolWorkingData::OnToolStart(Context);
+}
+
+void FPCGInteractiveToolWorkingData_PointArrayData::OnToolCancel(const FPCGInteractiveToolWorkingDataContext& Context)
+{
+	if(OnToolStartPointArray != nullptr)
+	{
+		UEngine::CopyPropertiesForUnrelatedObjects(OnToolStartPointArray, GeneratedPointData);
+	}
+	
+	FPCGInteractiveToolWorkingData::OnToolCancel(Context);
+}
+#endif
+
+void FPCGInteractiveToolWorkingData_PointArrayData::InitializeRuntimeElementData(FPCGContext* Context) const
+{
+	FPCGInteractiveToolWorkingData::InitializeRuntimeElementData(Context);
+
+	TArray<FPCGTaggedData>& Outputs = Context->OutputData.TaggedData;
+
+	FPCGTaggedData& Output = Outputs.Emplace_GetRef();
+	// We duplicate it so that PCG can take 'ownership'.
+	// If not, it can't process references and caches properly
+	// Might be relying on the fact it shouldn't access the same object over multiple generations
+	// And if we recreate it every time it gets properly managed
+	// If it gets dropped from cache, it's expected to be invalid
+	Output.Data = GeneratedPointData->DuplicateData(Context);
+	Output.Pin = PCGPinConstants::DefaultOutputLabel;
+#if WITH_EDITOR
+	FPCGSelectionKey Key = FPCGSelectionKey::CreateFromObjectPtr(GeneratedPointData);
+	FPCGDynamicTrackingHelper::AddSingleDynamicTrackingKey(Context, MoveTemp(Key), false);
+#endif
+}
+
+bool FPCGInteractiveToolWorkingData_PointArrayData::IsValid() const
+{
+	return GeneratedPointData != nullptr;
+}
+
+USplineComponent* FPCGInteractiveToolWorkingData_Spline::GetSplineComponent() const
+{
+	return SplineComponent.Get();
+}
+
+#if WITH_EDITORONLY_DATA
+void FPCGInteractiveToolWorkingData_Spline::InitializeInternal(const FPCGInteractiveToolWorkingDataContext& Context)
+{	
+	SplineComponent = FindOrGenerateSplineComponent(Context);
+}
+
+void FPCGInteractiveToolWorkingData_Spline::OnToolStart(const FPCGInteractiveToolWorkingDataContext& Context)
+{
+	FPCGInteractiveToolWorkingData::OnToolStart(Context);
+
+	if(SplineComponent.IsValid())
+	{
+		OnToolStartSplineComponent = Cast<USplineComponent>(StaticDuplicateObject(SplineComponent.Get(), GetTransientPackage()));
+
+		// We wipe the spline points on tool start so it becomes easy to edit them after the fact.
+		// @todo_pcg
+		// When the Spline Tool exposes more controls, such as selecting spline points to move them, resetting the spline etc., we will remove this.
+		SplineComponent->bSplineHasBeenEdited = true;
+		SplineComponent->ClearSplinePoints(true);
+	}
+}
+
+void FPCGInteractiveToolWorkingData_Spline::OnToolCancel(const FPCGInteractiveToolWorkingDataContext& Context)
+{
+	// If the Spline Component is valid and isn't a generated component, we revert its state.
+	if(SplineComponent.IsValid() && GeneratedResources.GeneratedComponents.Contains(SplineComponent) == false)
+	{
+		if(OnToolStartSplineComponent)
+		{
+			// We restore the spline component to its previous state
+			UEngine::CopyPropertiesForUnrelatedObjects(OnToolStartSplineComponent, SplineComponent.Get());
+			
+			SplineComponent->UpdateSpline();
+			
+			OnToolStartSplineComponent->MarkAsGarbage();
+			OnToolStartSplineComponent = nullptr;
+		}
+	}
+
+	// This will delete the spline component if it was generated
+	FPCGInteractiveToolWorkingData::OnToolCancel(Context);
+}
+
+USplineComponent* FPCGInteractiveToolWorkingData_Spline::FindOrGenerateSplineComponent(const FPCGInteractiveToolWorkingDataContext& Context)
+{
+	USplineComponent* MatchingSplineComponent = FindMatchingSplineComponent(Context);
+
+	// If we didn't find any spline component to match with, we create it now
+	if (MatchingSplineComponent == nullptr)
+	{
+		MatchingSplineComponent = GenerateMatchingSplineComponent(Context);
+	}	
+	
+	ensureMsgf(MatchingSplineComponent != nullptr, TEXT("Matching Spline Component should always be found after initialization as it was either found or created."));
+	return MatchingSplineComponent;
+}
+
+USplineComponent* FPCGInteractiveToolWorkingData_Spline::FindMatchingSplineComponent(const FPCGInteractiveToolWorkingDataContext& Context)
+{
+	TArray<USplineComponent*> AllSplineComponents;
+	Context.OwningActor->GetComponents<USplineComponent>(AllSplineComponents, false);
+	
+	USplineComponent* MatchingSplineComponent = nullptr;
+	if (Context.DataInstanceIdentifier.IsNone())
+	{
+		if(AllSplineComponents.Num() > 0)
+		{			
+			MatchingSplineComponent = AllSplineComponents[0];
+		}
+	}
+	else
+	{
+		// When the DataInstanceIdentifier is set, we only care about the splines with the given tag
+		AllSplineComponents.RemoveAll([DataInstanceIdentifier = Context.DataInstanceIdentifier](USplineComponent* Candidate)
+		{
+			return Candidate->ComponentHasTag(DataInstanceIdentifier) == false;
+		});
+		
+		// If we find a matching spline component with the specific tag, use it
+		if (AllSplineComponents.Num() > 0)
+		{
+			MatchingSplineComponent = Cast<USplineComponent>(AllSplineComponents[0]);
+		}
+	}
+
+	if(AllSplineComponents.Num() > 1 && MatchingSplineComponent != nullptr)
+	{
+		UE_LOGFMT(LogPCG, Warning, "More than one spline component found. Choosing the first available: {0}", MatchingSplineComponent->GetName());
+	}
+	
+	return MatchingSplineComponent;
+}
+
+USplineComponent* FPCGInteractiveToolWorkingData_Spline::GenerateMatchingSplineComponent(const FPCGInteractiveToolWorkingDataContext& Context)
+{
+	USplineComponent* MatchingSplineComponent = NewObject<USplineComponent>(Context.OwningActor.Get());
+	MatchingSplineComponent->SetMobility(Context.OwningActor->GetRootComponent()->GetMobility());
+	MatchingSplineComponent->SetupAttachment(Context.OwningActor->GetRootComponent());
+	// Transient until Applied. 
+	MatchingSplineComponent->SetFlags(RF_Transient);
+
+	if (Context.DataInstanceIdentifier.IsNone() == false)
+	{
+		MatchingSplineComponent->ComponentTags.Add(Context.DataInstanceIdentifier);
+	}
+
+	MatchingSplineComponent->ClearSplinePoints();
+	MatchingSplineComponent->bSplineHasBeenEdited = true;
+
+	// @todo_pcg
+	// When creating a new spline, we don't know if it should be closed or not
+	// We have access to the tool & property set; if we move them into runtime, we can read the spline tool's bClosedSpline setting
+	MatchingSplineComponent->SetClosedLoop(false);
+		
+	MatchingSplineComponent->RegisterComponent();
+	Context.OwningActor->AddInstanceComponent(MatchingSplineComponent);
+
+	return MatchingSplineComponent;
+}
+
+USplineComponent* FPCGInteractiveToolWorkingData_SplineSurface::FindMatchingSplineComponent(const FPCGInteractiveToolWorkingDataContext& Context)
+{
+	TArray<USplineComponent*> AllClosedSplineComponents;
+	Context.OwningActor->GetComponents<USplineComponent>(AllClosedSplineComponents, false);
+
+	// We remove all spline components that are open loops since this working data struct is only for spline surfaces
+	AllClosedSplineComponents.RemoveAll([](USplineComponent* Candidate)
+	{
+		return Candidate->IsClosedLoop() == false;
+	});
+	
+	USplineComponent* MatchingSplineComponent = nullptr;
+	if (Context.DataInstanceIdentifier.IsNone())
+	{
+		if(AllClosedSplineComponents.Num() > 0)
+		{
+			MatchingSplineComponent = AllClosedSplineComponents[0];
+		}
+	}
+	else
+	{
+		// When the DataInstanceIdentifier is set, we only care about the closed splines with the given tag
+		AllClosedSplineComponents.RemoveAll([DataInstanceIdentifier = Context.DataInstanceIdentifier](USplineComponent* Candidate)
+		{
+			return Candidate->ComponentHasTag(DataInstanceIdentifier) == false;
+		});
+		
+		// If we find a matching spline component with the specific tag, use i
+		if (AllClosedSplineComponents.Num() > 0)
+		{
+			MatchingSplineComponent = Cast<USplineComponent>(AllClosedSplineComponents[0]);
+		}
+	}
+
+	if(AllClosedSplineComponents.Num() > 1 && MatchingSplineComponent != nullptr)
+	{
+		UE_LOGFMT(LogPCG, Warning, "More than one spline component found. Choosing the first available: {0}", MatchingSplineComponent->GetName());
+	}
+
+	return MatchingSplineComponent;
+}
+
+USplineComponent* FPCGInteractiveToolWorkingData_SplineSurface::GenerateMatchingSplineComponent(const FPCGInteractiveToolWorkingDataContext& Context)
+{
+	USplineComponent* MatchingSplineComponent = NewObject<USplineComponent>(Context.OwningActor.Get());
+	MatchingSplineComponent->SetMobility(Context.OwningActor->GetRootComponent()->GetMobility());
+	MatchingSplineComponent->SetupAttachment(Context.OwningActor->GetRootComponent());
+	// Transient until Applied. 
+	MatchingSplineComponent->SetFlags(RF_Transient);
+
+	if (Context.DataInstanceIdentifier.IsNone() == false)
+	{
+		MatchingSplineComponent->ComponentTags.Add(Context.DataInstanceIdentifier);
+	}
+
+	MatchingSplineComponent->ClearSplinePoints();
+	MatchingSplineComponent->bSplineHasBeenEdited = true;
+	MatchingSplineComponent->SetClosedLoop(true);
+		
+	MatchingSplineComponent->RegisterComponent();
+	Context.OwningActor->AddInstanceComponent(MatchingSplineComponent);
+
+	return MatchingSplineComponent;
+}
+#endif
+
+void FPCGInteractiveToolWorkingData_Spline::InitializeRuntimeElementData(FPCGContext* Context) const
+{
+	FPCGInteractiveToolWorkingData::InitializeRuntimeElementData(Context);
+
+	TArray<FPCGTaggedData>& Outputs = Context->OutputData.TaggedData;
+
+	if(SplineComponent.IsValid())
+	{
+		FPCGTaggedData& Output = Outputs.Emplace_GetRef();
+		// We duplicate it so that PCG can take 'ownership'.
+		// If not, it can't process references and caches properly
+		// Might be relying on the fact it shouldn't access the same object over multiple generations
+		// And if we recreate it every time it gets properly managed
+		// If it gets dropped from cache, it's expected to be invalid
+		UPCGSplineData* SplineData = FPCGContext::NewObject_AnyThread<UPCGSplineData>(Context);
+		SplineData->Initialize(SplineComponent.Get());
+		Output.Data = SplineData;
+		Output.Pin = PCGPinConstants::DefaultOutputLabel;
+#if WITH_EDITOR
+		FPCGSelectionKey Key = FPCGSelectionKey::CreateFromObjectPtr(SplineComponent.Get());
+		FPCGDynamicTrackingHelper::AddSingleDynamicTrackingKey(Context, MoveTemp(Key), false);
+#endif
+	}
+}
+
+bool FPCGInteractiveToolWorkingData_Spline::IsValid() const
+{
+	return SplineComponent.IsValid();
+}
diff --git a/Engine/Plugins/PCG/Source/PCG/Private/Elements/PCGDataFromTool.cpp b/Engine/Plugins/PCG/Source/PCG/Private/Elements/PCGDataFromTool.cpp
new file mode 100644
--- /dev/null
+++ b/Engine/Plugins/PCG/Source/PCG/Private/Elements/PCGDataFromTool.cpp
@@ -0,0 +1,115 @@
+// Copyright Epic Games, Inc. All Rights Reserved.
+
+#include "Elements/PCGDataFromTool.h"
+#include "PCGActorAndComponentMapping.h"
+#include "PCGComponent.h"
+#include "PCGCustomVersion.h"
+#include "PCGSubsystem.h"
+#include "Data/PCGBasePointData.h"
+#include "Elements/PCGAddComponent.h"
+#include "Helpers/PCGDynamicTrackingHelpers.h"
+#include "Helpers/PCGHelpers.h"
+#include "Utils/PCGGraphExecutionLogging.h"
+
+#include "Algo/AnyOf.h"
+
+
+#include UE_INLINE_GENERATED_CPP_BY_NAME(PCGDataFromTool)
+
+#define LOCTEXT_NAMESPACE "PCGDataFromTool"
+
+#if WITH_EDITOR
+FString UPCGDataFromTool::GetAdditionalTitleInformation() const
+{
+	FStringBuilderBase Builder;
+	if (ToolTag.IsNone() == false)
+	{
+		Builder.Append(ToolTag.ToString());
+	}
+
+	bool bAddSeparator = false;
+	if (ToolTag.IsNone() == false && DataInstance.IsNone() == false)
+	{
+		bAddSeparator = true;	
+	}
+	
+	if (bAddSeparator)
+	{
+		Builder.Append(" | ");
+	}
+
+	if (DataInstance.IsNone() == false)
+	{
+		Builder.Append(DataInstance.ToString());
+	}
+	
+	return Builder.ToString();
+}
+
+FText UPCGDataFromTool::GetNodeTooltipText() const
+{
+	return LOCTEXT("DataFromToolTooltip", "Builds a collection of PCG-compatible data from the selected tools.");
+}
+#endif
+
+FPCGElementPtr UPCGDataFromTool::CreateElement() const
+{
+	return MakeShared<FPCGDataFromToolElement>();
+}
+
+TArray<FPCGPinProperties> UPCGDataFromTool::InputPinProperties() const
+{
+	// No input pins
+	TArray<FPCGPinProperties> PinProperties;
+	return PinProperties;
+}
+
+TArray<FPCGPinProperties> UPCGDataFromTool::OutputPinProperties() const
+{
+	// Default output pin
+	TArray<FPCGPinProperties> PinProperties;
+	PinProperties.Emplace(PCGPinConstants::DefaultOutputLabel, EPCGDataType::Any);
+	return PinProperties;
+}
+
+bool FPCGDataFromToolElement::ExecuteInternal(FPCGContext* Context) const
+{
+	TRACE_CPUPROFILER_EVENT_SCOPE(FPCGDataFromToolElement::Execute);
+
+	check(Context);
+	
+	const UPCGDataFromTool* Settings = Context->GetInputSettings<UPCGDataFromTool>();
+	check(Settings);
+	
+	UPCGComponent* PCGComponent = Cast<UPCGComponent>(Context->ExecutionSource.Get());
+
+	if (Settings->ToolTag.IsNone() == false)
+	{
+		FNameBuilder Builder(Settings->ToolTag);
+
+		// This mirrors UPCGInteractiveToolSettings::GetWorkingDataIdentifier
+		if (Settings->DataInstance.IsNone() == false)
+		{
+			Builder.AppendChar(TEXT('_'));
+			Builder.Append(Settings->DataInstance.ToString());
+		}
+
+		FName ToolDataIdentifier = Builder.ToString();
+		
+		if (PCGComponent->ToolDataContainer.ToolData.Contains(ToolDataIdentifier) && PCGComponent->ToolDataContainer.ToolData[ToolDataIdentifier].IsValid())
+		{
+			const FPCGInteractiveToolWorkingData* Data = PCGComponent->ToolDataContainer.ToolData[ToolDataIdentifier].GetPtr<FPCGInteractiveToolWorkingData>();
+			Data->InitializeRuntimeElementData(Context);
+		}
+		// @todo_pcg
+		// If we return false this causes accumulation of generation tasks, but if we return true it won't recognize the data change until we Flush Cache since the dynamic tracking key wasn't actually created yet
+		else
+		{
+			return true;
+		}
+	}
+
+	return true;
+}
+
+#undef LOCTEXT_NAMESPACE
diff --git a/Engine/Plugins/PCG/Source/PCG/Private/PCGGraph.cpp b/Engine/Plugins/PCG/Source/PCG/Private/PCGGraph.cpp
--- a/Engine/Plugins/PCG/Source/PCG/Private/PCGGraph.cpp
+++ b/Engine/Plugins/PCG/Source/PCG/Private/PCGGraph.cpp
@@ -633,6 +633,33 @@ UPCGNode* UPCGGraph::FindNodeWithSettings(const UPCGSettingsInterface* InSetting
 	return NodeFound;
 }
 
+TArray<UPCGNode*> UPCGGraph::FindNodesWithSettings(const TSubclassOf<UPCGSettingsInterface> InSettingsClass, bool bRecursive) const
+{
+	TArray<UPCGNode*> Result;
+
+	auto FindNode = [&Result, InSettingsClass](UPCGNode* InNode)
+	{
+		if (InNode && InNode->GetSettingsInterface() && InNode->GetSettingsInterface()->GetClass() == InSettingsClass)
+		{
+			Result.AddUnique(InNode);
+		}
+
+		// Always return true so we never stop execution
+		return true;
+	};
+
+	if (bRecursive)
+	{
+		ForEachNodeRecursively(FindNode);
+	}
+	else
+	{
+		ForEachNode(FindNode);
+	}
+
+	return Result;
+}
+
 UPCGNode* UPCGGraph::FindNodeByTitleName(FName NodeTitle, bool bRecursive, const TSubclassOf<const UPCGSettings> OptionalClass) const
 {
 	UPCGNode* NodeFound = nullptr;
diff --git a/Engine/Plugins/PCG/Source/PCG/Private/PCGPointArray.cpp b/Engine/Plugins/PCG/Source/PCG/Private/PCGPointArray.cpp
--- a/Engine/Plugins/PCG/Source/PCG/Private/PCGPointArray.cpp
+++ b/Engine/Plugins/PCG/Source/PCG/Private/PCGPointArray.cpp
@@ -35,6 +35,28 @@ void FPCGPointArray::SetNumPoints(int32 InNumPoints, bool bInitializeValues)
 	MetadataEntry.SetNum(NumPoints, /*bAllocate=*/false, bInitializedValues);
 }
 
+void FPCGPointArray::RemoveAt(int32 Index)
+{
+	if(Index >= 0 && Index < NumPoints)
+	{
+		NumPoints -= 1;
+
+		Transform.RemoveAt(Index);
+		Density.RemoveAt(Index);
+		BoundsMin.RemoveAt(Index);
+		BoundsMax.RemoveAt(Index);
+		Color.RemoveAt(Index);
+		Steepness.RemoveAt(Index);
+		Seed.RemoveAt(Index);
+		MetadataEntry.RemoveAt(Index);
+
+		if(NumPoints == 1)
+		{
+			bInitializedValues = false;
+		}
+	}
+}
+
 void FPCGPointArray::Allocate(EPCGPointNativeProperties InProperties)
 {
 	if (EnumHasAnyFlags(InProperties, EPCGPointNativeProperties::Transform))
diff --git a/Engine/Plugins/PCG/Source/PCG/Public/Data/PCGPointArrayData.h b/Engine/Plugins/PCG/Source/PCG/Public/Data/PCGPointArrayData.h
--- a/Engine/Plugins/PCG/Source/PCG/Public/Data/PCGPointArrayData.h
+++ b/Engine/Plugins/PCG/Source/PCG/Public/Data/PCGPointArrayData.h
@@ -36,7 +36,8 @@ public:
 	
 	virtual int32 GetNumPoints() const override { return PointArray.GetNumPoints(); }
 	UE_API virtual void SetNumPoints(int32 InNumPoints, bool bInitializeValues = true) override;
-		
+	UE_API void RemoveAt(int32 Index);
+	
 	UE_API virtual void AllocateProperties(EPCGPointNativeProperties Properties) override;
 	UE_API virtual void FreeProperties(EPCGPointNativeProperties Properties) override;
 	UE_API virtual void MoveRange(int32 RangeStartIndex, int32 MoveToIndex, int32 NumElements) override;
@@ -62,7 +63,7 @@ public:
 	UE_API virtual FPCGPointSeed::ConstValueRange GetConstSeedValueRange() const override;
 	UE_API virtual FPCGPointMetadataEntry::ConstValueRange GetConstMetadataEntryValueRange() const override;
 	//~End UPCGBasePointData Interface
-		
+
 protected:
 	UE_API virtual void InitializeSpatialDataInternal(const FPCGInitializeFromDataParams& InParams) override;
 
diff --git a/Engine/Plugins/PCG/Source/PCG/Public/Editor/PCGGraphToolData.h b/Engine/Plugins/PCG/Source/PCG/Public/Editor/PCGGraphToolData.h
new file mode 100644
--- /dev/null
+++ b/Engine/Plugins/PCG/Source/PCG/Public/Editor/PCGGraphToolData.h
@@ -0,0 +1,231 @@
+﻿// Copyright Epic Games, Inc. All Rights Reserved.
+
+#pragma once
+
+#include "PCGPointArray.h"
+#include "Data/PCGSplineData.h"
+
+#include "CoreMinimal.h"
+#include "UObject/Object.h"
+#include "GameFramework/Actor.h"
+#include "Components/BoxComponent.h"
+#include "PCGGraphToolData.generated.h"
+
+class USplineComponent;
+class UPCGPointArrayData;
+class UPCGComponent;
+class UInteractiveTool;
+class UInteractiveToolPropertySet;
+
+/** A transient struct that contains resources generated during tool usage for a specific tool data struct. Used to finalize or clean up when a tool is applied or cancelled. */
+USTRUCT()
+struct PCG_API FPCGInteractiveToolWorkingDataGeneratedResources
+{
+	GENERATED_BODY()
+
+	FPCGInteractiveToolWorkingDataGeneratedResources() = default;
+
+#if WITH_EDITORONLY_DATA
+	/** Finalizes generated resources (e.g. remove Transient flag). */
+	void Finalize();
+
+	/** Remove and destroy generated resources. */
+	void Cleanup();
+	
+	UPROPERTY(Transient)
+	TArray<TWeakObjectPtr<UActorComponent>> GeneratedComponents;
+#endif
+};
+
+/** A transient context struct that is passed in so the tool data can manage its data properly during tool usage. Used at editor-time only. */
+USTRUCT()
+struct PCG_API FPCGInteractiveToolWorkingDataContext
+{
+	GENERATED_BODY()
+
+	FPCGInteractiveToolWorkingDataContext() = default;
+
+#if WITH_EDITORONLY_DATA
+	FPCGInteractiveToolWorkingDataContext(UPCGComponent* OwningComponent, UInteractiveTool* InteractiveTool, UInteractiveToolPropertySet* PropertySet, FName DataInstanceIdentifier);
+
+	/** The owning PCG Component. Should always be valid. Can be used as an Outer for new objects. */
+	UPROPERTY(Transient)
+	TWeakObjectPtr<UPCGComponent> OwningPCGComponent;
+
+	/** An optional DataInstanceIdentifier. Can be used to find or create a specific object, such as a Spline Component with a certain tag. */
+	UPROPERTY(Transient)
+	FName DataInstanceIdentifier;
+
+	/** The owning actor. The owning PCG component's outer. Should always be valid. Convenience accessor. */
+	UPROPERTY(Transient)
+	TWeakObjectPtr<AActor> OwningActor;
+
+	/** The interactive tool currently editing the tool data. */
+	UPROPERTY(Transient)
+	TWeakObjectPtr<UInteractiveTool> InteractiveTool;
+
+	/** The PCG Setting properties of the Interactive Tool. */
+	UPROPERTY(Transient)
+	TWeakObjectPtr<UInteractiveToolPropertySet> PCGSettings;
+#endif
+};
+
+/** The base struct for working data for a tool. This data is created by the tool, stored on the PCG Component,
+ * then read back by the tool on subsequent tool invocations. */
+USTRUCT()
+struct PCG_API FPCGInteractiveToolWorkingData
+{
+	GENERATED_BODY()
+
+	FPCGInteractiveToolWorkingData() = default;
+	virtual ~FPCGInteractiveToolWorkingData() = default;
+
+	bool IsInitialized() const { return bInitialized; }
+	
+	virtual bool IsValid() const;
+
+	/** Construct and add pcg data to the runtime context.
+	 * This will translate data from this struct into UPCGData to be processed by the graph.
+	 * Add dynamic tracking keys in here if you want to refresh when a change to your source data happens. */
+	virtual void InitializeRuntimeElementData(FPCGContext* InContext) const {}
+	
+#if WITH_EDITORONLY_DATA
+public:
+	DECLARE_DELEGATE_OneParam(FOnInitialized, FPCGInteractiveToolWorkingData*)
+	
+	/** Initializes data used by the tool working data. Initialize is called once when the working data is created. */
+	void Initialize(const FPCGInteractiveToolWorkingDataContext& Context);
+	
+	/** OnToolStart is called every time the tool is started up, or when the data is created.
+	 * Can be used to save the initial state of persistent data to restore to on Cancel, or to Modify existing data. */
+	virtual void OnToolStart(const FPCGInteractiveToolWorkingDataContext& Context);
+	
+	/** Called when the tool is applied. Finalize changes here. */
+	virtual void OnToolApply(const FPCGInteractiveToolWorkingDataContext& Context);
+	
+	/** Called when the tool is canceled. Revert changes you have made here. */
+	virtual void OnToolCancel(const FPCGInteractiveToolWorkingDataContext& Context);
+
+	FOnInitialized& OnInitialized() { return OnInitializedDelegate; }
+private:
+	/** Implement your initialization logic here. Called by Initialize. */
+	virtual void InitializeInternal(const FPCGInteractiveToolWorkingDataContext& Context);
+protected:
+	/** Each working data struct keeps track of its generated resources. They have to be reverted during OnToolCancel. */
+	UPROPERTY(Transient)
+	FPCGInteractiveToolWorkingDataGeneratedResources GeneratedResources;
+
+	/** OnInitialized delegate to allow the settings class to do more work. */
+	FOnInitialized OnInitializedDelegate;
+#endif
+	
+private:
+	UPROPERTY()
+	bool bInitialized = false;
+};
+
+/** Tool Data that stores points. */
+USTRUCT(BlueprintType, DisplayName="Point Array Tool Data")
+struct FPCGInteractiveToolWorkingData_PointArrayData : public FPCGInteractiveToolWorkingData
+{
+	GENERATED_BODY()
+	
+public:
+	FPCGInteractiveToolWorkingData_PointArrayData() = default;
+
+	virtual void InitializeRuntimeElementData(FPCGContext* InContext) const override;
+	
+	virtual bool IsValid() const override;
+	
+	UPCGPointArrayData* GetPointArrayData() const { return GeneratedPointData; }
+
+#if WITH_EDITORONLY_DATA
+	virtual void InitializeInternal(const FPCGInteractiveToolWorkingDataContext& Context) override;
+	
+	virtual void OnToolStart(const FPCGInteractiveToolWorkingDataContext& Context) override;
+	
+	virtual void OnToolCancel(const FPCGInteractiveToolWorkingDataContext& Context) override;
+#endif
+	
+protected:
+	UPROPERTY(VisibleAnywhere, BlueprintReadOnly, Category="Working Data")
+	TObjectPtr<UPCGPointArrayData> GeneratedPointData;
+
+	/** A transient copy of the point array data, created on tool start/data access. We revert back to these points on Cancel. */
+	UPROPERTY(Transient)
+	TObjectPtr<UPCGPointArrayData> OnToolStartPointArray;
+};
+
+/** The working data for splines. */
+USTRUCT(BlueprintType, DisplayName="Spline Tool Data")
+struct FPCGInteractiveToolWorkingData_Spline : public FPCGInteractiveToolWorkingData
+{
+	GENERATED_BODY()
+public:
+	FPCGInteractiveToolWorkingData_Spline() = default;
+
+	virtual void InitializeRuntimeElementData(FPCGContext* InContext) const override;
+	virtual bool IsValid() const override;
+	
+	PCG_API USplineComponent* GetSplineComponent() const;
+
+#if WITH_EDITORONLY_DATA
+public:
+	virtual void OnToolStart(const FPCGInteractiveToolWorkingDataContext& Context) override;
+	virtual void OnToolCancel(const FPCGInteractiveToolWorkingDataContext& Context) override;
+	
+private:
+	virtual void InitializeInternal(const FPCGInteractiveToolWorkingDataContext& Context) override;
+	
+private:
+	USplineComponent* FindOrGenerateSplineComponent(const FPCGInteractiveToolWorkingDataContext& Context);
+	
+	virtual USplineComponent* FindMatchingSplineComponent(const FPCGInteractiveToolWorkingDataContext& Context);
+	virtual USplineComponent* GenerateMatchingSplineComponent(const FPCGInteractiveToolWorkingDataContext& Context);
+#endif
+	
+protected:
+	UPROPERTY(VisibleAnywhere, BlueprintReadOnly, Category="Working Data")
+	TWeakObjectPtr<USplineComponent> SplineComponent;
+
+	UPROPERTY(Transient)
+	TObjectPtr<USplineComponent> OnToolStartSplineComponent;
+};
+
+USTRUCT(BlueprintType, DisplayName="Spline Surface Tool Data")
+struct PCG_API FPCGInteractiveToolWorkingData_SplineSurface : public FPCGInteractiveToolWorkingData_Spline
+{
+	GENERATED_BODY()
+
+#if WITH_EDITOR
+private:
+	/** We look for and generate a closed spline component. */
+	virtual USplineComponent* FindMatchingSplineComponent(const FPCGInteractiveToolWorkingDataContext& Context) override;
+	virtual USplineComponent* GenerateMatchingSplineComponent(const FPCGInteractiveToolWorkingDataContext& Context) override;
+#endif
+};
+
+/** The tool data stored on the graph itself. Informs the UI. */
+USTRUCT(BlueprintType)
+struct PCG_API FPCGGraphToolData
+{
+	GENERATED_BODY()
+
+	UPROPERTY(EditAnywhere, Category = Tool)
+	FText DisplayName;
+	
+	UPROPERTY(EditAnywhere, Category = Tool)
+	TArray<FString> CompatibleToolTags;
+
+	UPROPERTY(EditAnywhere, Category = Tool)
+	TSubclassOf<AActor> InitialActorClassToSpawn;
+};
+
+USTRUCT(BlueprintType)
+struct PCG_API FPCGInteractiveToolWorkingDataContainer
+{
+	GENERATED_BODY()
+	
+	UPROPERTY(BlueprintReadOnly, VisibleAnywhere, Category = "Tool Data")
+	TMap<FName, TInstancedStruct<FPCGInteractiveToolWorkingData>> ToolData;
+};
\ No newline at end of file
diff --git a/Engine/Plugins/PCG/Source/PCG/Public/Elements/PCGDataFromTool.h b/Engine/Plugins/PCG/Source/PCG/Public/Elements/PCGDataFromTool.h
new file mode 100644
--- /dev/null
+++ b/Engine/Plugins/PCG/Source/PCG/Public/Elements/PCGDataFromTool.h
@@ -0,0 +1,56 @@
+// Copyright Epic Games, Inc. All Rights Reserved.
+
+#pragma once
+
+#include "PCGContext.h"
+#include "PCGSettings.h"
+#include "Elements/PCGActorSelector.h"
+#include "Elements/PCGLoadObjectsContext.h"
+
+#include "UObject/ObjectKey.h"
+
+#include "PCGDataFromTool.generated.h"
+
+/** Builds a collection of PCG-compatible data from the specified editor tools. */
+UCLASS(MinimalAPI, BlueprintType, ClassGroup = (Procedural))
+class UPCGDataFromTool : public UPCGSettings
+{
+	GENERATED_BODY()
+
+public:
+	//~Begin UPCGSettings interface
+#if WITH_EDITOR
+	virtual FName GetDefaultNodeName() const override { return FName(TEXT("GetToolData")); }
+	virtual FText GetDefaultNodeTitle() const override { return NSLOCTEXT("PCGDataFromTool", "NodeTitle", "Get Tool Data"); }
+	PCG_API virtual FString GetAdditionalTitleInformation() const override;
+	PCG_API virtual FText GetNodeTooltipText() const override;
+	virtual EPCGSettingsType GetType() const override { return EPCGSettingsType::Generic; }
+	virtual bool CanDynamicallyTrackKeys() const override { return true; }
+#endif
+
+protected:
+	PCG_API virtual TArray<FPCGPinProperties> InputPinProperties() const override;
+	PCG_API virtual TArray<FPCGPinProperties> OutputPinProperties() const override;
+	PCG_API virtual FPCGElementPtr CreateElement() const override;
+	//~End UPCGSettings interface
+
+public:
+	/** The tool tag on the pcg component to create a data collection from.
+	 * Example given: "PaintTool" will allow you to retrieve points created by the PaintTool in the level viewport in PCG mode. */
+	UPROPERTY(BlueprintReadWrite, EditAnywhere, Category = "Data Retrieval Settings", meta=(PCG_Overridable))
+	FName ToolTag = NAME_None;
+
+	/** The optional name of the data this node should retrieve. If you want to support multiple tool outputs of the same type, differentiate using this name.
+	 * Example given: "Trees" or "Bushes" will allow you to paint onto two different layers so you can process them differently in the graph.
+	 */
+	UPROPERTY(BlueprintReadWrite, EditAnywhere, Category = "Data Retrieval Settings", meta=(PCG_Overridable))
+	FName DataInstance = NAME_None;
+};
+
+class FPCGDataFromToolElement : public IPCGElement
+{
+protected:
+	virtual bool CanExecuteOnlyOnMainThread(FPCGContext* Context) const override { return true; }
+	PCG_API virtual bool ExecuteInternal(FPCGContext* Context) const override;
+	virtual bool SupportsBasePointDataInputs(FPCGContext* InContext) const override { return true; }
+};
diff --git a/Engine/Plugins/PCG/Source/PCG/Public/PCGComponent.h b/Engine/Plugins/PCG/Source/PCG/Public/PCGComponent.h
--- a/Engine/Plugins/PCG/Source/PCG/Public/PCGComponent.h
+++ b/Engine/Plugins/PCG/Source/PCG/Public/PCGComponent.h
@@ -14,6 +14,7 @@
 #include "ComponentInstanceDataCache.h"
 #include "Components/ActorComponent.h"
 #include "Containers/Map.h"
+#include "Editor/PCGGraphToolData.h"
 #include "Misc/ScopeExit.h"
 #include "UObject/ObjectKey.h"
 #include "Misc/TransactionallySafeCriticalSection.h"
@@ -363,6 +364,10 @@ public:
 	UPROPERTY(BlueprintReadOnly, VisibleAnywhere, Instanced, Category = RuntimeGeneration, meta = (EditCondition = "GenerationTrigger == EPCGComponentGenerationTrigger::GenerateAtRuntime", EditConditionHides))
 	TObjectPtr<UPCGSchedulingPolicyBase> SchedulingPolicy;
 
+	/** This stores working data of a tool; runtime property as the PCG Component will rely on the tool's working data to generate properly. */
+	UPROPERTY(BlueprintReadOnly, VisibleAnywhere, Category = "Tool Data")
+	FPCGInteractiveToolWorkingDataContainer ToolDataContainer;
+	
 #if WITH_EDITORONLY_DATA
 	UPROPERTY(BlueprintReadWrite, EditAnywhere, Category = "Editing Settings", meta = (DisplayName = "Regenerate PCG Volume In Editor", DisplayPriority = 400))
 	bool bRegenerateInEditor = true;
diff --git a/Engine/Plugins/PCG/Source/PCG/Public/PCGGraph.h b/Engine/Plugins/PCG/Source/PCG/Public/PCGGraph.h
--- a/Engine/Plugins/PCG/Source/PCG/Public/PCGGraph.h
+++ b/Engine/Plugins/PCG/Source/PCG/Public/PCGGraph.h
@@ -14,6 +14,7 @@
 #endif // WITH_EDITOR
 
 #include "ComputeFramework/ComputeGraphInstance.h"
+#include "Editor/PCGGraphToolData.h"
 #include "StructUtils/PropertyBag.h"
 #include "UObject/ObjectPtr.h"
 
@@ -178,9 +179,9 @@ public:
 	UPROPERTY(BlueprintReadWrite, EditAnywhere, Category = AssetInfo, AssetRegistrySearchable, meta = (EditCondition = "IsExportToLibraryEnabled", EditConditionHides, PCGNoHash))
 	bool bExposeToLibrary = false;
 
-	/** Signifies the graph has compatibility with the annotated tags. Multiple tags may be comma separated. Ex. "PCGTool1, PCGTool2, PCGTool3" */
+	/** Contains the data relevant for PCG Editor Mode usage. */
 	UPROPERTY(BlueprintReadWrite, EditAnywhere, Category = AssetInfo, AssetRegistrySearchable, meta=(EditCondition = "bExposeToLibrary", EditConditionHides, PCGNoHash))
-	FString CompatibleToolTags;
+	FPCGGraphToolData ToolData;
 
 	/** When enabled, this graph can be executed outside of the world using an editor execution source. */
 	UPROPERTY(BlueprintReadWrite, EditAnywhere, Category = AssetInfo, AssetRegistrySearchable, meta = (PCGNoHash, EditCondition = "CanToggleStandaloneGraph()", EditConditionHides))
@@ -393,6 +394,8 @@ public:
 
 	/** Returns the node with the given settings in the graph, if any */
 	UE_API UPCGNode* FindNodeWithSettings(const UPCGSettingsInterface* InSettings, bool bRecursive = false) const;
+	
+	UE_API TArray<UPCGNode*> FindNodesWithSettings(const TSubclassOf<UPCGSettingsInterface> InSettingsClass, bool bRecursive = false) const;
 
 	/** Returns the first node that matches the given name in the graph, if any.*/
 	UE_API UPCGNode* FindNodeByTitleName(FName NodeTitle, bool bRecursive = false, TSubclassOf<const UPCGSettings> OptionalClass = {}) const;
diff --git a/Engine/Plugins/PCG/Source/PCG/Public/PCGPointArray.h b/Engine/Plugins/PCG/Source/PCG/Public/PCGPointArray.h
--- a/Engine/Plugins/PCG/Source/PCG/Public/PCGPointArray.h
+++ b/Engine/Plugins/PCG/Source/PCG/Public/PCGPointArray.h
@@ -132,6 +132,23 @@ public:
 		return Values[Index];
 	}
 
+	inline void RemoveAt(int32 Index)
+	{
+		ensure(Index >= 0 && Index < NumValues);
+		const int32 OldNum = Values.Num();
+		
+		if (IsAllocated())
+		{
+			Values.RemoveAt(Index);
+			NumValues -= 1; 
+		}
+
+		if (OldNum > 0 && Values.Num() == 1)
+		{
+			Free();
+		}
+	}
+	
 	/** Move ranges of values(moves them back into the array).Used by the PCGAsync api */
 	inline void MoveRange(int32 RangeStartIndex, int32 MoveToIndex, int32 NumElements)
 	{
@@ -249,6 +266,7 @@ public:
 
 	UE_API void SetNumPoints(int32 InNumPoints, bool bInitializeValues);
 	int32 GetNumPoints() const { return NumPoints; }
+	UE_API void RemoveAt(int32 Index);
 		
 	UE_API EPCGPointNativeProperties GetAllocatedProperties() const;
 	UE_API void Allocate(EPCGPointNativeProperties InProperties);
diff --git a/Engine/Plugins/PCG/Source/PCGEditor/PCGEditor.Build.cs b/Engine/Plugins/PCG/Source/PCGEditor/PCGEditor.Build.cs
--- a/Engine/Plugins/PCG/Source/PCGEditor/PCGEditor.Build.cs
+++ b/Engine/Plugins/PCG/Source/PCGEditor/PCGEditor.Build.cs
@@ -69,8 +69,15 @@ namespace UnrealBuildTool.Rules
 					"TypedElementFramework",
 					"TypedElementRuntime",
 					"UnrealEd",
-					"WidgetRegistration",
+					"WidgetRegistration", 
+					"ModelingComponents",
+					"ModelingComponentsEditorOnly"
 				});
+			
+			// @todo_pcg
+			// Set to 1 for development where each tool graph creates a button
+			// Remove again when we settle on a definitive UI
+			PrivateDefinitions.Add("USE_PCG_MODE_GRAPHS_AS_BUTTONS=0");
 		}
 	}
 }
diff --git a/Engine/Plugins/PCG/Source/PCGEditor/Private/Details/PCGInteractiveToolSettingsDetails.cpp b/Engine/Plugins/PCG/Source/PCGEditor/Private/Details/PCGInteractiveToolSettingsDetails.cpp
--- a/Engine/Plugins/PCG/Source/PCGEditor/Private/Details/PCGInteractiveToolSettingsDetails.cpp
+++ b/Engine/Plugins/PCG/Source/PCGEditor/Private/Details/PCGInteractiveToolSettingsDetails.cpp
@@ -52,7 +52,8 @@ void FPCGInteractiveToolSettingsDetails::CustomizeDetails(IDetailLayoutBuilder&
 	}
 	
 	IDetailCategoryBuilder& CategoryBuilder = DetailBuilder.EditCategory("PCG");
-	if(ToolSettings.IsValid() && ToolSettings->GraphInstance != nullptr && ToolSettings->GraphInstance->GetUserParametersStruct()->GetNumPropertiesInBag() > 0)
+	UPCGGraphInstance* GraphInstance = ToolSettings.IsValid() ? ToolSettings->GetGraphInstance() : nullptr;
+	if(GraphInstance != nullptr && GraphInstance->GetUserParametersStruct()->GetNumPropertiesInBag() > 0)
 	{
 		TArray<TSharedRef<IPropertyHandle>> PropertyHandles;
 		CategoryBuilder.GetDefaultProperties(PropertyHandles);
@@ -62,10 +63,10 @@ void FPCGInteractiveToolSettingsDetails::CustomizeDetails(IDetailLayoutBuilder&
 			CategoryBuilder.AddProperty(PropertyHandle);
 
 			// We add the parameters right behind the graph property
-			if(PropertyHandle->GetProperty()->GetFName() == GET_MEMBER_NAME_CHECKED(UPCGInteractiveToolSettings, TargetGraph))
-			{			
+			if(PropertyHandle->GetProperty()->GetFName() == GET_MEMBER_NAME_CHECKED(UPCGInteractiveToolSettings, ToolGraph))
+			{
 				IDetailPropertyRow* PropertyRow = CategoryBuilder.AddExternalObjectProperty(
-				{ToolSettings->GraphInstance}, 
+				{GraphInstance}, 
 				GET_MEMBER_NAME_CHECKED(UPCGGraphInstance, ParametersOverrides),
 				 EPropertyLocation::Default,
 				 FAddPropertyParams());
@@ -86,9 +87,12 @@ void FPCGInteractiveToolSettingsDetails::Rebuild(UObject* Object, const FPropert
 {
 	if(Object == ToolSettings)
 	{
-		if(PropertyChangedChainEvent.Property->GetFName() == GET_MEMBER_NAME_CHECKED(UPCGInteractiveToolSettings, TargetGraph))
+		if(PropertyChangedChainEvent.Property && PropertyChangedChainEvent.Property->GetFName() == GET_MEMBER_NAME_CHECKED(UPCGInteractiveToolSettings, ToolGraph))
 		{
-			DetailLayoutBuilder.Pin()->ForceRefreshDetails();
+			if(DetailLayoutBuilder.IsValid())
+			{
+				DetailLayoutBuilder.Pin()->ForceRefreshDetails();
+			}
 		}
 	}
 }
diff --git a/Engine/Plugins/PCG/Source/PCGEditor/Private/EditorMode/PCGEdMode.cpp b/Engine/Plugins/PCG/Source/PCGEditor/Private/EditorMode/PCGEdMode.cpp
--- a/Engine/Plugins/PCG/Source/PCGEditor/Private/EditorMode/PCGEdMode.cpp
+++ b/Engine/Plugins/PCG/Source/PCGEditor/Private/EditorMode/PCGEdMode.cpp
@@ -28,11 +28,14 @@
 
 #include "ToolContextInterfaces.h"
 #include "ToolMenus.h"
+#include "ToolTargetManager.h"
 #include "EditorMode/PCGEdModeSettings.h"
 #include "EditorMode/PCGEdModeStyle.h"
 #include "EditorMode/Tools/Paint/PCGPaintTool.h"
 #include "Engine/World.h" // FWorldDelegates
 #include "Framework/Application/SlateApplication.h"
+#include "ToolTargets/PrimitiveComponentToolTarget.h"
+#include "ToolTargets/StaticMeshComponentToolTarget.h"
 #include "ViewportToolbar/UnrealEdViewportToolbar.h"
 
 #include UE_INLINE_GENERATED_CPP_BY_NAME(PCGEdMode)
@@ -46,7 +49,7 @@ namespace PCGEditorMode
 	namespace Constants
 	{
 		const FName PCGEditorModeOwnerName = "PCGEditorModeViewportToolBar";
-		const FText PCGEditorModeDisplayName = LOCTEXT("PCGEditorModeDisplayName", "PCG (Experimental)");
+		const FText PCGEditorModeDisplayName = LOCTEXT("PCGEditorModeDisplayName", "PCG");
 	}
 
 	namespace Helpers
@@ -157,6 +160,9 @@ void UPCGEditorMode::Enter()
 		return;
 	}
 
+	UEditorInteractiveToolsContext* InteractiveToolsContext = GetInteractiveToolsContext();
+	InteractiveToolsContext->TargetManager->AddTargetFactory(NewObject<UPrimitiveComponentToolTargetFactory>(InteractiveToolsContext->TargetManager));
+	
 	ToolManager->OnToolPostBuild.AddUObject(this, &UPCGEditorMode::OnToolPostBuild);
 	ToolManager->OnToolShutdownRequest.BindLambda([this](UInteractiveToolManager*, UInteractiveTool*, const EToolShutdownType ShutdownType)
 	{
@@ -171,20 +177,54 @@ void UPCGEditorMode::Enter()
 		}
 	});
 
+#if USE_PCG_MODE_GRAPHS_AS_BUTTONS
+	for(TSharedPtr<FUICommandInfo> ToolCommandInfo : FPCGEditorModeToolCommands::Get().GetLineContextToolCommands())
+	{
+		UPCGDrawSplineToolBuilder* SplineToolBuilder = NewObject<UPCGDrawSplineToolBuilder>();
+		SplineToolBuilder->SetToolGraph(CastChecked<UPCGGraph>(FSoftObjectPath(ToolCommandInfo->GetCommandName().ToString()).TryLoad()));
+		SplineToolBuilder->SetIsSplineSurface(false);
+		RegisterTool(ToolCommandInfo, ToolCommandInfo->GetCommandName().ToString(), SplineToolBuilder);
+	}
+
+	UPCGDrawSplineToolBuilder* SplineSurfaceToolBuilder = NewObject<UPCGDrawSplineToolBuilder>();
+	SplineSurfaceToolBuilder->SetIsSplineSurface(false);
+
+	// Add spline surface tools here
+	
+	for(TSharedPtr<FUICommandInfo> ToolCommandInfo : FPCGEditorModeToolCommands::Get().GetPaintContextToolCommands())
+	{
+		UPCGPaintToolBuilder* PaintToolBuilder = NewObject<UPCGPaintToolBuilder>();
+		PaintToolBuilder->SetToolGraph(CastChecked<UPCGGraph>(FSoftObjectPath(ToolCommandInfo->GetCommandName().ToString()).TryLoad()));
+		
+		RegisterTool(ToolCommandInfo, ToolCommandInfo->GetCommandName().ToString(), PaintToolBuilder);
+	}
+
+	for(TSharedPtr<FUICommandInfo> ToolCommandInfo : FPCGEditorModeToolCommands::Get().GetVolumeContextToolCommands())
+	{
+		UPCGCreateVolumeToolBuilder* VolumeToolBuilder = NewObject<UPCGCreateVolumeToolBuilder>();
+		VolumeToolBuilder->SetToolGraph(CastChecked<UPCGGraph>(FSoftObjectPath(ToolCommandInfo->GetCommandName().ToString()).TryLoad()));
+		RegisterTool(ToolCommandInfo, ToolCommandInfo->GetCommandName().ToString(), VolumeToolBuilder);
+	}
+#else
 	const FPCGEditorModeToolCommands& ToolManagerCommands = FPCGEditorModeToolCommands::Get();
 
-	// @todo_pcg: Could be abstracted out and have the tool identifier come from the label
-	UPCGDrawSplineToolBuilder* DrawSplineBuilder = NewObject<UPCGDrawSplineToolBuilder>();
-	DrawSplineBuilder->SetIsSplineSurface(false);
-	RegisterTool(ToolManagerCommands.EnableDrawSplineTool, TEXT("DrawSpline"), DrawSplineBuilder);
-	UPCGDrawSplineToolBuilder* DrawSurfaceBuilder = NewObject<UPCGDrawSplineToolBuilder>();
-	DrawSurfaceBuilder->SetIsSplineSurface(true);
-	RegisterTool(ToolManagerCommands.EnableDrawSurfaceTool, TEXT("DrawSurface"), DrawSurfaceBuilder);
+	UPCGDrawSplineToolBuilder* DrawSplineToolBuilder = NewObject<UPCGDrawSplineToolBuilder>();
+	DrawSplineToolBuilder->SetSplineToolClass(UPCGDrawSplineTool::StaticClass());
+	RegisterTool(ToolManagerCommands.EnableDrawSplineTool, TEXT("DrawSpline"), DrawSplineToolBuilder);
+
+	UPCGDrawSplineToolBuilder* DrawSurfaceToolBuilder = NewObject<UPCGDrawSplineToolBuilder>();
+	DrawSurfaceToolBuilder->SetSplineToolClass(UPCGDrawSplineSurfaceTool::StaticClass());
+	RegisterTool(ToolManagerCommands.EnableDrawSurfaceTool, TEXT("DrawSurface"), DrawSurfaceToolBuilder);
+
+	UPCGPaintToolBuilder* PaintToolBuilder = NewObject<UPCGPaintToolBuilder>();
+	RegisterTool(ToolManagerCommands.EnablePaintTool, TEXT("Paint"), PaintToolBuilder);
 
-	RegisterTool(ToolManagerCommands.EnablePaintTool, TEXT("Paint"), NewObject<UPCGPaintToolBuilder>());
-	RegisterTool(ToolManagerCommands.EnableCreateVolumeTool, TEXT("CreateVolume"), NewObject<UPCGCreateVolumeToolBuilder>());
+	// @TODO_PCG Reenable Volume Tool when it makes sense
+	// UPCGCreateVolumeToolBuilder* CreateVolumeToolBuilder = NewObject<UPCGCreateVolumeToolBuilder>();
+	// RegisterTool(ToolManagerCommands.EnableCreateVolumeTool, TEXT("CreateVolume"), CreateVolumeToolBuilder);
 
-	ToolManager->SelectActiveToolType(EToolSide::Left, TEXT("Line"));
+	ToolManager->SelectActiveToolType(EToolSide::Left, "DrawSpline");
+#endif
 }
 
 void UPCGEditorMode::Exit()
diff --git a/Engine/Plugins/PCG/Source/PCGEditor/Private/EditorMode/PCGEdModeCommands.cpp b/Engine/Plugins/PCG/Source/PCGEditor/Private/EditorMode/PCGEdModeCommands.cpp
--- a/Engine/Plugins/PCG/Source/PCGEditor/Private/EditorMode/PCGEdModeCommands.cpp
+++ b/Engine/Plugins/PCG/Source/PCGEditor/Private/EditorMode/PCGEdModeCommands.cpp
@@ -2,11 +2,16 @@
 
 #include "EditorMode/PCGEdModeCommands.h"
 
+#include "PCGEditorModule.h"
+#include "AssetRegistry/AssetRegistryModule.h"
 #include "EditorMode/PCGEdModeSettings.h"
 #include "EditorMode/PCGEdModeStyle.h"
+#include "EditorMode/Tools/Helpers/PCGEdModeEditorUtilities.h"
 #include "EditorMode/Tools/Line/PCGDrawSplineTool.h"
 #include "EditorMode/Tools/Paint/PCGPaintTool.h"
 #include "EditorMode/Tools/Volume/PCGCreateVolumeTool.h"
+#include "Logging/StructuredLog.h"
+#include "Misc/OutputDeviceNull.h"
 
 #define LOCTEXT_NAMESPACE "PCGEditorMode"
 
@@ -85,8 +90,6 @@ void FPCGEditorModeToolCommands::RegisterCommands()
 	PCG_TOOL_COMMAND_INFO(EnableCreateVolumeTool, "CreateVolume", "Create Volume", "Create a PCG volume in the scene.", "PCGEditorMode.Tools.CreateVolume");
 }
 
-#undef PCG_TOOL_COMMAND_INFO
-
 void FPCGEditorModeToolCommands::RegisterAllToolCommands()
 {
 	FPCGDrawSplineCommands::Register();
@@ -106,19 +109,31 @@ void FPCGEditorModeToolCommands::GetToolDefaultObjectList(TArray<UInteractiveToo
 	// @todo_pcg: To be populated with tool commands for bindings.
 }
 
-TArray<TSharedPtr<FUICommandInfo>> FPCGEditorModeToolCommands::GetLineContextToolCommands() const
+TArray<TSharedPtr<FUICommandInfo>> FPCGEditorModeToolCommands::GetLineContextToolCommands()
 {
+#if USE_PCG_MODE_GRAPHS_AS_BUTTONS
+	return PCG::EditorMode::Utility::GetOrCreateUICommandsWithToolTag(AsShared(), "SplineTool");
+#else
 	return {EnableDrawSplineTool, EnableDrawSurfaceTool};
+#endif
 }
 
-TArray<TSharedPtr<FUICommandInfo>> FPCGEditorModeToolCommands::GetPaintContextToolCommands() const
+TArray<TSharedPtr<FUICommandInfo>> FPCGEditorModeToolCommands::GetPaintContextToolCommands()
 {
+#if USE_PCG_MODE_GRAPHS_AS_BUTTONS
+	return PCG::EditorMode::Utility::GetOrCreateUICommandsWithToolTag(AsShared(), "PaintTool");
+#else
 	return {EnablePaintTool};
+#endif
 }
 
-TArray<TSharedPtr<FUICommandInfo>> FPCGEditorModeToolCommands::GetVolumeContextToolCommands() const
+TArray<TSharedPtr<FUICommandInfo>> FPCGEditorModeToolCommands::GetVolumeContextToolCommands()
 {
+#if USE_PCG_MODE_GRAPHS_AS_BUTTONS
+	return PCG::EditorMode::Utility::GetOrCreateUICommandsWithToolTag(AsShared(), "VolumeTool");
+#else
 	return {EnableCreateVolumeTool};
+#endif
 }
 
 /**
@@ -138,11 +153,13 @@ void ToolCommandClassName::GetToolDefaultObjectList(TArray<UInteractiveTool*>& T
 {\
 ToolCDOs.Add(GetMutableDefault<ToolClassName>());\
 }
-PCG_DEFINE_TOOL_COMMANDS(FPCGDrawSplineCommands, "DrawSpline", "Draw Spline", UPCGDrawSplineTool);
-PCG_DEFINE_TOOL_COMMANDS(FPCGDrawSurfaceCommands, "DrawSurface", "Draw Surface", UPCGDrawSplineTool);
-PCG_DEFINE_TOOL_COMMANDS(FPCGPaintCommands, "DrawSurface", "Draw Surface", UPCGPaintTool);
+PCG_DEFINE_TOOL_COMMANDS(FPCGDrawSplineCommands, "DrawSpline", "Draw Spline", UPCGDrawSplineToolBase);
+PCG_DEFINE_TOOL_COMMANDS(FPCGDrawSurfaceCommands, "DrawSurface", "Draw Surface", UPCGDrawSplineToolBase);
+PCG_DEFINE_TOOL_COMMANDS(FPCGPaintCommands, "Paint", "Paint", UPCGPaintTool);
 PCG_DEFINE_TOOL_COMMANDS(FPCGCreateVolumeCommands, "CreateVolume", "Create Volume", UPCGCreateVolumeTool);
 
 #undef PCG_DEFINE_TOOL_COMMANDS
 
+#undef PCG_TOOL_COMMAND_INFO
+
 #undef LOCTEXT_NAMESPACE
diff --git a/Engine/Plugins/PCG/Source/PCGEditor/Private/EditorMode/PCGEdModeToolkit.cpp b/Engine/Plugins/PCG/Source/PCGEditor/Private/EditorMode/PCGEdModeToolkit.cpp
--- a/Engine/Plugins/PCG/Source/PCGEditor/Private/EditorMode/PCGEdModeToolkit.cpp
+++ b/Engine/Plugins/PCG/Source/PCGEditor/Private/EditorMode/PCGEdModeToolkit.cpp
@@ -4,7 +4,6 @@
 
 #include "EditorMode/PCGEdModeCommands.h"
 #include "EditorMode/PCGEdModeSettings.h"
-#include "EditorMode/Tools/PCGInteractiveTool.h"
 
 #include "EditorModeManager.h"
 #include "LevelEditorViewport.h"
@@ -72,6 +71,16 @@ void FPCGEditorModeToolkit::OnToolPaletteChanged(const FName PaletteName)
 	ensure(false);
 }
 
+void FPCGEditorModeToolkit::BuildToolPalette(FName Palette, class FToolBarBuilder& ToolbarBuilder)
+{
+	if (!OwningEditorMode.IsValid())
+	{
+		return;
+	}
+
+	FModeToolkit::BuildToolPalette(Palette, ToolbarBuilder);	
+}
+
 FText FPCGEditorModeToolkit::GetActiveToolDisplayName() const
 {
 	return CachedActiveToolInfo.ToolDisplayName;
@@ -109,10 +118,7 @@ void FPCGEditorModeToolkit::OnToolEnded(UInteractiveToolManager* Manager, UInter
 	ClearToolNotification();
 	ClearToolWarning();
 
-	if (UPCGInteractiveTool* CurrentTool = CastChecked<UPCGInteractiveTool>(Tool, ECastCheckedType::NullAllowed))
-	{
-		UnregisterToolCallbacks(CurrentTool);
-	}
+	UnregisterToolCallbacks(Tool);
 }
 
 void FPCGEditorModeToolkit::PostToolNotification(const FText& Message)
@@ -204,7 +210,7 @@ EVisibility FPCGEditorModeToolkit::GetAcceptButtonVisibility() const
 
 EVisibility FPCGEditorModeToolkit::GetCancelButtonVisibility() const
 {
-	return GetScriptableEditorMode()->GetInteractiveToolsContext()->ActiveToolHasAccept() ? EVisibility::Visible : EVisibility::Collapsed;
+	return GetScriptableEditorMode()->GetInteractiveToolsContext()->CanCancelActiveTool() ? EVisibility::Visible : EVisibility::Collapsed;
 }
 
 FReply FPCGEditorModeToolkit::OnAcceptButtonClicked() const
@@ -244,23 +250,21 @@ void FPCGEditorModeToolkit::UnregisterToolkitCallbacks()
 	PCGEditorModeSettings->OnSettingChanged().RemoveAll(this);
 }
 
-void FPCGEditorModeToolkit::RegisterToolCallbacks(UPCGInteractiveTool* CurrentTool)
+void FPCGEditorModeToolkit::RegisterToolCallbacks(UInteractiveTool* CurrentTool)
 {
 	if (CurrentTool)
 	{
 		CurrentTool->OnPropertySetsModified.AddSP(this, &FPCGEditorModeToolkit::UpdateActiveTool);
 		CurrentTool->OnPropertyModifiedDirectlyByTool.AddSP(this, &FPCGEditorModeToolkit::InvalidateCachedDetailPanelState);
-		CurrentTool->OnToolWarningChanged.AddSP(this, &FPCGEditorModeToolkit::PostToolWarning);
 	}
 }
 
-void FPCGEditorModeToolkit::UnregisterToolCallbacks(UPCGInteractiveTool* CurrentTool) const
+void FPCGEditorModeToolkit::UnregisterToolCallbacks(UInteractiveTool* CurrentTool) const
 {
 	if (CurrentTool)
 	{
 		CurrentTool->OnPropertySetsModified.RemoveAll(this);
 		CurrentTool->OnPropertyModifiedDirectlyByTool.RemoveAll(this);
-		CurrentTool->OnToolWarningChanged.RemoveAll(this);
 	}
 }
 
@@ -276,45 +280,43 @@ void FPCGEditorModeToolkit::OnEditorModeSettingsChanged(UObject* InObject, FProp
 void FPCGEditorModeToolkit::UpdateActiveTool()
 {
 	UInteractiveTool* CurrentTool = GetScriptableEditorMode()->GetToolManager(EToolsContextScope::EdMode)->GetActiveTool(EToolSide::Left);
-	if (UPCGInteractiveTool* CurrentPCGTool = CastChecked<UPCGInteractiveTool>(CurrentTool, ECastCheckedType::NullAllowed))
-	{
-		CachedActiveToolInfo = CurrentPCGTool->GetToolInfo();
+	
+	CachedActiveToolInfo = CurrentTool->GetToolInfo();
 
-		ToolkitBuilder->SetActiveToolDisplayName(CachedActiveToolInfo.ToolDisplayName);
+	ToolkitBuilder->SetActiveToolDisplayName(CachedActiveToolInfo.ToolDisplayName);
 
-		if (const UPCGEditorModeSettings* Settings = GetDefault<UPCGEditorModeSettings>())
+	if (const UPCGEditorModeSettings* Settings = GetDefault<UPCGEditorModeSettings>())
+	{
+		if (Settings->bHideToolButtonsDuringActiveTool)
 		{
-			if (Settings->bHideToolButtonsDuringActiveTool)
-			{
-				ToolkitBuilder->SetActivePaletteCommandsVisibility(EVisibility::Collapsed);
-			}
+			ToolkitBuilder->SetActivePaletteCommandsVisibility(EVisibility::Collapsed);
 		}
+	}
 
-		// Before actually changing the detail panel, we need to see where the current keyboard focus is, because
-		// if it's inside the detail panel, we'll need to reset it to the detail panel as a whole, else we might
-		// lose it entirely when that detail panel element gets destroyed (which would make us unable to receive any
-		// hotkey presses until the user clicks somewhere).
-		const TSharedPtr<SWidget> FocusedWidget = FSlateApplication::Get().GetKeyboardFocusedWidget();
-		if (FocusedWidget != ModeDetailsView)
+	// Before actually changing the detail panel, we need to see where the current keyboard focus is, because
+	// if it's inside the detail panel, we'll need to reset it to the detail panel as a whole, else we might
+	// lose it entirely when that detail panel element gets destroyed (which would make us unable to receive any
+	// hotkey presses until the user clicks somewhere).
+	const TSharedPtr<SWidget> FocusedWidget = FSlateApplication::Get().GetKeyboardFocusedWidget();
+	if (FocusedWidget != ModeDetailsView)
+	{
+		// Search upward from the currently focused widget
+		TSharedPtr<SWidget> CurrentWidget = FocusedWidget;
+		while (CurrentWidget.IsValid())
 		{
-			// Search upward from the currently focused widget
-			TSharedPtr<SWidget> CurrentWidget = FocusedWidget;
-			while (CurrentWidget.IsValid())
+			if (CurrentWidget == ModeDetailsView)
 			{
-				if (CurrentWidget == ModeDetailsView)
-				{
-					// Reset focus to the detail panel as a whole to avoid losing it when the inner elements change.
-					FSlateApplication::Get().SetKeyboardFocus(ModeDetailsView);
-					break;
-				}
-
-				CurrentWidget = CurrentWidget->GetParentWidget();
+				// Reset focus to the detail panel as a whole to avoid losing it when the inner elements change.
+				FSlateApplication::Get().SetKeyboardFocus(ModeDetailsView);
+				break;
 			}
-		}
 
-		ModeDetailsView->SetObjects({CurrentPCGTool->GetToolProperties(true)});
-		RegisterToolCallbacks(CurrentPCGTool);
+			CurrentWidget = CurrentWidget->GetParentWidget();
+		}
 	}
+	
+	ModeDetailsView->SetObjects(CurrentTool->GetToolProperties(true));
+	RegisterToolCallbacks(CurrentTool);
 
 	// @todo_pcg: If a tool property needs to invalidate the cached details panel, do it here. 
 	// InvalidateCachedDetailPanelState(nullptr);
@@ -328,7 +330,7 @@ void FPCGEditorModeToolkit::InvalidateCachedDetailPanelState(UObject* ChangedObj
 void FPCGEditorModeToolkit::RegisterPalettes()
 {
 	const FPCGEditorModePaletteCommands& PaletteCommands = FPCGEditorModePaletteCommands::Get();
-	const FPCGEditorModeToolCommands& ToolCommands = FPCGEditorModeToolCommands::Get();
+	FPCGEditorModeToolCommands& ToolCommands = FPCGEditorModeToolCommands::Get();
 
 	ToolkitSections = MakeShared<FToolkitSections>();
 	ToolkitSections->DetailsView = ModeDetailsView;
diff --git a/Engine/Plugins/PCG/Source/PCGEditor/Private/EditorMode/Tools/Helpers/PCGEdModeActorHelpers.cpp b/Engine/Plugins/PCG/Source/PCGEditor/Private/EditorMode/Tools/Helpers/PCGEdModeActorHelpers.cpp
--- a/Engine/Plugins/PCG/Source/PCGEditor/Private/EditorMode/Tools/Helpers/PCGEdModeActorHelpers.cpp
+++ b/Engine/Plugins/PCG/Source/PCGEditor/Private/EditorMode/Tools/Helpers/PCGEdModeActorHelpers.cpp
@@ -8,7 +8,7 @@
 
 #define LOCTEXT_NAMESPACE "PCGEditorMode"
 
-namespace PCG::EditorMode
+namespace UE::PCG::EditorMode
 {
 	namespace Actor
 	{
@@ -31,14 +31,14 @@ namespace PCG::EditorMode
 			return UPCGActorHelpers::SpawnDefaultActor(SpawnDefaultActorParams);
 		}
 
-		AActor* SpawnWorking(UWorld* World, const FTransform& Transform)
+		AActor* SpawnWorking(UWorld* World, const FTransform& Transform, TSubclassOf<AActor> ActorClass)
 		{
 			FActorSpawnParameters SpawnParams;
 			SpawnParams.Name = DefaultWorkingActorName;
 			SpawnParams.NameMode = FActorSpawnParameters::ESpawnActorNameMode::Requested;
 			SpawnParams.ObjectFlags = RF_Transient;
 
-			UPCGActorHelpers::FSpawnDefaultActorParams SpawnDefaultActorParams(World, AActor::StaticClass(), Transform, SpawnParams);
+			UPCGActorHelpers::FSpawnDefaultActorParams SpawnDefaultActorParams(World, ActorClass, Transform, SpawnParams);
 
 			return UPCGActorHelpers::SpawnDefaultActor(SpawnDefaultActorParams);
 		}
@@ -53,8 +53,8 @@ namespace PCG::EditorMode
 				return nullptr;
 			}
 
-			TArray<UActorComponent*> ToolComponents = Owner->GetComponentsByTag(UPCGComponent::StaticClass(), Tags::ToolGeneratedTag);
-			return ToolComponents.IsEmpty() ? nullptr : Cast<UPCGComponent>(ToolComponents[0]);
+			UPCGComponent* PCGComponent = Owner->GetComponentByClass<UPCGComponent>();
+			return PCGComponent;
 		}
 
 		UPCGComponent* Create(AActor* Owner, const FName ComponentName)
@@ -75,7 +75,7 @@ namespace PCG::EditorMode
 		{
 			return Create(Owner, Actor::DefaultWorkingActorName);
 		}
-	} // namespace PCG::EditorMode::PCGComponent
-} // namespace PCG::EditorMode
+	}
+}
 
 #undef LOCTEXT_NAMESPACE
diff --git a/Engine/Plugins/PCG/Source/PCGEditor/Private/EditorMode/Tools/Helpers/PCGEdModeEditorUtilities.cpp b/Engine/Plugins/PCG/Source/PCGEditor/Private/EditorMode/Tools/Helpers/PCGEdModeEditorUtilities.cpp
new file mode 100644
--- /dev/null
+++ b/Engine/Plugins/PCG/Source/PCGEditor/Private/EditorMode/Tools/Helpers/PCGEdModeEditorUtilities.cpp
@@ -0,0 +1,124 @@
+﻿// Copyright Epic Games, Inc. All Rights Reserved.
+
+#include "EditorMode/Tools/Helpers/PCGEdModeEditorUtilities.h"
+
+#include "PCGGraph.h"
+#include "AssetRegistry/AssetRegistryModule.h"
+#include "Misc/OutputDeviceNull.h"
+#include "HAL/IConsoleManager.h"
+#include "Framework/Commands/InputBindingManager.h"
+
+
+static TAutoConsoleVariable<bool> CVarClearToolAssetCache(
+	TEXT("pcg.debug.tools.cleartoolassetscache"),
+	true,                                
+	TEXT("Enable to always regather the tool commands for the PCG Mode."),
+	ECVF_Default
+);
+
+TArray<UE::PCG::EditorMode::Utility::FPCGGraphToolEditorData> UE::PCG::EditorMode::Utility::GetGraphToolsWithToolTag(FName ToolTag)
+{
+	TArray<UE::PCG::EditorMode::Utility::FPCGGraphToolEditorData> Result;
+	
+	TArray<TSharedPtr<FUICommandInfo>> GraphAssetCommands;
+	FAssetRegistryModule& AssetRegistryModule = FModuleManager::Get().LoadModuleChecked<FAssetRegistryModule>("AssetRegistry");
+
+	FARFilter Filter;
+	Filter.ClassPaths.Add(UPCGGraph::StaticClass()->GetClassPathName());
+	Filter.TagsAndValues.Add(GET_MEMBER_NAME_CHECKED(UPCGGraph, ToolData));
+
+	TArray<FAssetData> GraphAssets;
+	AssetRegistryModule.Get().GetAssets(Filter, GraphAssets);
+
+	for(auto It = GraphAssets.CreateIterator(); It; ++It)
+	{
+		FPCGGraphToolData ToolData;
+
+		FString ToolDataText;
+		if(It->GetTagValue(GET_MEMBER_NAME_CHECKED(UPCGGraph, ToolData), ToolDataText))
+		{
+			FOutputDeviceNull NullOut;
+			StaticStruct<FPCGGraphToolData>()->ImportText(*ToolDataText, &ToolData, nullptr, PPF_None, &NullOut,
+				FPCGGraphToolData::StaticStruct()->GetName(), true);
+
+			if(ToolData.CompatibleToolTags.Contains(ToolTag))
+			{
+				FPCGGraphToolEditorData GraphToolEditorData;
+				GraphToolEditorData.AssetData = (*It);
+				GraphToolEditorData.GraphToolData = ToolData;
+				Result.Add(GraphToolEditorData);
+			}
+		}
+	}
+
+	return Result;
+}
+
+bool UE::PCG::EditorMode::Utility::DoesPCGGraphHaveToolTag(FName ToolTag, const FAssetData& AssetData)
+{
+	FPCGGraphToolData ToolData;
+
+	FString ToolDataText;
+	if(AssetData.GetTagValue(GET_MEMBER_NAME_CHECKED(UPCGGraph, ToolData), ToolDataText))
+	{
+		FOutputDeviceNull NullOut;
+		StaticStruct<FPCGGraphToolData>()->ImportText(*ToolDataText, &ToolData, nullptr, PPF_None, &NullOut,
+			FPCGGraphToolData::StaticStruct()->GetName(), true);
+
+		if(ToolData.CompatibleToolTags.Contains(ToolTag))
+		{
+			return true;
+		}
+	}
+
+	return false;
+}
+
+TArray<TSharedPtr<FUICommandInfo>> UE::PCG::EditorMode::Utility::GetOrCreateUICommandsWithToolTag(TSharedRef<FBindingContext> BindingContext, FName ToolTag)
+{
+	static TMap<FName, TOptional<TArray<TSharedPtr<FUICommandInfo>>>> CachedGraphAssetCommands;
+
+	TOptional<TArray<TSharedPtr<FUICommandInfo>>>& CommandInfos = CachedGraphAssetCommands.FindOrAdd(ToolTag);
+	if(CommandInfos.IsSet() && CVarClearToolAssetCache.GetValueOnGameThread())
+	{
+		for(TSharedPtr<FUICommandInfo> CommandInfo : CommandInfos.GetValue())
+		{
+			FInputBindingManager::Get().RemoveInputCommand(BindingContext, CommandInfo.ToSharedRef());
+		}
+		
+		CommandInfos.Reset();
+	}
+
+	if(CommandInfos.IsSet() == false)
+	{
+		TArray<TSharedPtr<FUICommandInfo>> GraphAssetCommands;
+
+		TArray<FPCGGraphToolEditorData> GraphToolEditorData = GetGraphToolsWithToolTag(ToolTag);
+
+		for(const FPCGGraphToolEditorData& Data : GraphToolEditorData)
+		{
+			TSharedPtr<FUICommandInfo>& AssetCommandInfo = GraphAssetCommands.AddDefaulted_GetRef();
+			FString ObjectPath = Data.AssetData.GetObjectPathString();
+
+			FText DisplayName = Data.GraphToolData.DisplayName;
+			if(DisplayName.IsEmpty())
+			{
+				DisplayName = FText::AsCultureInvariant(Data.AssetData.AssetName.ToString());
+			}
+
+			FName CommandName = FName(ObjectPath);
+			FUICommandInfo::MakeCommandInfo(BindingContext,
+				AssetCommandInfo,
+				CommandName,
+				DisplayName,
+				FText::GetEmpty(),
+				FSlateIcon(),
+				EUserInterfaceActionType::ToggleButton,
+				FInputChord());
+		}
+
+		CommandInfos = GraphAssetCommands;
+	}
+
+	return CommandInfos.Get({});
+}
diff --git a/Engine/Plugins/PCG/Source/PCGEditor/Private/EditorMode/Tools/Helpers/PCGEdModeSplineHelpers.cpp b/Engine/Plugins/PCG/Source/PCGEditor/Private/EditorMode/Tools/Helpers/PCGEdModeSplineHelpers.cpp
--- a/Engine/Plugins/PCG/Source/PCGEditor/Private/EditorMode/Tools/Helpers/PCGEdModeSplineHelpers.cpp
+++ b/Engine/Plugins/PCG/Source/PCGEditor/Private/EditorMode/Tools/Helpers/PCGEdModeSplineHelpers.cpp
@@ -3,11 +3,12 @@
 #include "EditorMode/Tools/Helpers/PCGEdModeSplineHelpers.h"
 
 #include "Components/SplineComponent.h"
+#include "EditorMode/Tools/Helpers/PCGEdModeActorHelpers.h"
 #include "UObject/UObjectGlobals.h"
 
 #define LOCTEXT_NAMESPACE "PCGEditorMode"
 
-namespace PCG::EditorMode::Spline
+namespace UE::PCG::EditorMode::Spline
 {
 	static const FLazyName DefaultWorkingSplineName = "PCG Tool Working Spline";
 
@@ -41,6 +42,8 @@ namespace PCG::EditorMode::Spline
 			Spline = NewObject<USplineComponent>(Actor, USplineComponent::StaticClass(), Name, Flags);
 		}
 
+		Spline->SetFlags(RF_Transient);
+		Spline->ComponentTags.Add(Tags::ToolGeneratedTag);
 		Actor->AddInstanceComponent(Spline);
 		Spline->RegisterComponent();
 		Spline->ResetRelativeTransform();
@@ -54,7 +57,7 @@ namespace PCG::EditorMode::Spline
 		return CreateComponent(Actor, DefaultWorkingSplineName, /*TemplateComponent=*/nullptr, /*bTransactional=*/false);
 	}
 
-	USplineComponent* FindComponent(const AActor* Actor, const int32 TargetIndex)
+	USplineComponent* Find(const AActor* Actor, const int32 TargetIndex)
 	{
 		if (!ensure(Actor))
 		{
diff --git a/Engine/Plugins/PCG/Source/PCGEditor/Private/EditorMode/Tools/Line/PCGDrawSplineTool.cpp b/Engine/Plugins/PCG/Source/PCGEditor/Private/EditorMode/Tools/Line/PCGDrawSplineTool.cpp
--- a/Engine/Plugins/PCG/Source/PCGEditor/Private/EditorMode/Tools/Line/PCGDrawSplineTool.cpp
+++ b/Engine/Plugins/PCG/Source/PCGEditor/Private/EditorMode/Tools/Line/PCGDrawSplineTool.cpp
@@ -2,11 +2,11 @@
 
 #include "EditorMode/Tools/Line/PCGDrawSplineTool.h"
 
+#include "InteractiveToolManager.h"
 #include "PCGComponent.h"
 #include "EditorMode/PCGEdModeSettings.h"
 #include "EditorMode/PCGEdModeStyle.h"
 #include "EditorMode/Tools/Helpers/PCGEdModeSceneVisualizeHelpers.h"
-#include "EditorMode/Tools/Helpers/PCGEdModeSplineHelpers.h"
 
 #include "BaseBehaviors/SingleClickOrDragBehavior.h"
 #include "Components/SplineComponent.h"
@@ -32,7 +32,7 @@ namespace PCGDrawSplineTool
 			// These pass the working spline to the overloads below
 			virtual void Apply(UObject* Object) override
 			{
-				if (!ensure(Object && Object->IsA<UPCGDrawSplineTool>() && WorkingSpline))
+				if (!ensure(Object && Object->IsA<UPCGDrawSplineToolBase>() && WorkingSpline))
 				{
 					return;
 				}
@@ -42,7 +42,7 @@ namespace PCGDrawSplineTool
 
 			virtual void Revert(UObject* Object) override
 			{
-				if (!ensure(Object && Object->IsA<UPCGDrawSplineTool>() && WorkingSpline))
+				if (!ensure(Object && Object->IsA<UPCGDrawSplineToolBase>() && WorkingSpline))
 				{
 					return;
 				}
@@ -193,95 +193,123 @@ namespace PCGDrawSplineTool
 	{
 		const FText ToolName = LOCTEXT("ToolName", "DrawSpline");
 		const FText ToolDescription = LOCTEXT("DrawSplineDisplayMessage", "Draw a spline and distribute spline data directly to a PCG Graph Instance.");
+		const FName SplineToolTag = "SplineTool";
+		const FName SplineSurfaceToolTag = "SplineSurfaceTool";		
 		static const FText AddPointTransactionName = LOCTEXT("AddPointTransactionName", "Add Point");
 	}
 } // namespace PCGDrawSplineTool
 
-UPCGDrawSplineToolSettings::UPCGDrawSplineToolSettings()
+UPCGInteractiveToolSettings_SplineBase::UPCGInteractiveToolSettings_SplineBase()
 	: SplineName(GetDefault<UPCGEditorModeSettings>()->DefaultNewSplineComponentName)
 {
-	if (!HasAnyFlags(RF_ClassDefaultObject))
-	{
-		UPCGDrawSplineToolSettings::TryLoadDefaultGraph();
-	}
+	
 }
 
-FText UPCGDrawSplineToolSettings::GetToolTag() const
+USplineComponent* UPCGInteractiveToolSettings_SplineBase::GetWorkingSplineComponent(FName DataInstanceIdentifier) const
 {
-	return PCGDrawSplineTool::Constants::ToolName;
+	if(const FPCGInteractiveToolWorkingData_Spline* ToolData = GetTypedWorkingData<FPCGInteractiveToolWorkingData_Spline>(DataInstanceIdentifier))
+	{
+		return ToolData->GetSplineComponent();
+	}
+	
+	return nullptr;
 }
 
-void UPCGDrawSplineTool::Setup()
+TArray<USplineComponent*> UPCGInteractiveToolSettings_SplineBase::GetWorkingSplineComponents() const
 {
-	using namespace PCGDrawSplineTool;
-
-	Super::Setup();
+	TArray<USplineComponent*> Result;
 
-	UPCGDrawSplineToolSettings* Settings = CastChecked<UPCGDrawSplineToolSettings>(ToolSettings);
-	CachedMinPointSpacingSquared = Settings->MinPointSpacing * Settings->MinPointSpacing;
+	TArray<FName> AllWorkingDataInstanceNames = GetValidDataInstanceNamesForGraph();
+	for(const FName& DataInstanceIdentifier : AllWorkingDataInstanceNames)
+	{
+		if(USplineComponent* WorkingSplineComponent = GetWorkingSplineComponent(DataInstanceIdentifier))
+		{
+			Result.Add(WorkingSplineComponent);
+		}
+	}
 
-	WorkingSpline = Spline::CreateWorkingComponent(WorkingActor.Get());
-	check(WorkingSpline);
+	return Result;
+}
 
-	WorkingSpline->ClearSplinePoints();
-	WorkingSpline->SetClosedLoop(bIsSplineSurface);
-	WorkingSpline->bSplineHasBeenEdited = true;
+const UScriptStruct* UPCGInteractiveToolSettings_Spline::GetWorkingDataType() const
+{
+	return FPCGInteractiveToolWorkingData_Spline::StaticStruct();
+}
 
-	ClickOrDragBehavior = NewObject<USingleClickOrDragInputBehavior>();
-	ClickOrDragBehavior->Initialize(this, this);
-	AddInputBehavior(ClickOrDragBehavior);
+FName UPCGInteractiveToolSettings_Spline::GetToolTag() const
+{
+	return PCGDrawSplineTool::Constants::SplineToolTag;
+}
 
-	// @todo_pcg: Make adjustments to this and check when properties change.
-	CachedRayParams =
+void UPCGInteractiveToolSettings_Spline::OnWorkingDataInitialized(FPCGInteractiveToolWorkingData* WorkingData) const
+{
+	if(USplineComponent* SplineComponent = GetWorkingSplineComponent(DataInstance))
 	{
-		.World = GetWorld()
-	};
+		SplineComponent->SetClosedLoop(bClosedSpline);
+	}
+}
 
-	Settings->WatchProperty(Settings->TargetActor, [this](TWeakObjectPtr<AActor> TargetActor)
+const UScriptStruct* UPCGInteractiveToolSettings_SplineSurface::GetWorkingDataType() const
+{
+	return FPCGInteractiveToolWorkingData_SplineSurface::StaticStruct();
+}
+
+FName UPCGInteractiveToolSettings_SplineSurface::GetToolTag() const
+{
+	return PCGDrawSplineTool::Constants::SplineSurfaceToolTag;
+}
+
+void UPCGInteractiveToolSettings_Spline::OnPropertyModified(UObject* Object, FProperty* Property)
+{
+	Super::OnPropertyModified(Object, Property);
+	
+	if(Object == this && Property->GetFName() == GET_MEMBER_NAME_CHECKED(UPCGInteractiveToolSettings_Spline, bClosedSpline))
 	{
-		if (!TargetActor.IsValid())
+		if(USplineComponent* SplineComponent = GetWorkingSplineComponent(DataInstance))
 		{
-			SetWarning(LOCTEXT("TargetActorSetToInvalid", "Target actor is invalid!"));
-		}
-	});
+			SplineComponent->SetClosedLoop(bClosedSpline, true);
 
-	Settings->SilentUpdateWatched();
+			FPropertyChangedEvent Event(nullptr);
+			FCoreUObjectDelegates::OnObjectPropertyChanged.Broadcast(SplineComponent, Event);
+		}
+	}
 }
 
-void UPCGDrawSplineTool::Shutdown(const EToolShutdownType ShutdownType)
+void UPCGDrawSplineToolBase::Setup()
 {
 	using namespace PCGDrawSplineTool;
 
-	if (ShutdownType == EToolShutdownType::Accept && WorkingSpline && WorkingSpline->GetNumberOfSplinePoints() > 0)
-	{
-		ApplyTool();
-	}
-
-	if (WorkingComponent.IsValid())
-	{
-		WorkingComponent->CleanupLocalImmediate(/*bRemoveAllComponents=*/true, /*bCleanUpLocalComponents=*/true);
-		WorkingComponent->bActivated = false;
-		WorkingComponent->DestroyComponent();
-	}
+	Super::Setup();
 
-	if (WorkingSpline)
-	{
-		WorkingSpline->DestroyComponent();
-		WorkingSpline = nullptr;
-	}
+	UPCGInteractiveToolSettings_SplineBase* ToolSettings = GetSplineSettings();
+	CachedMinPointSpacingSquared = ToolSettings->MinPointSpacing * ToolSettings->MinPointSpacing;
+	
+	ClickOrDragBehavior = NewObject<USingleClickOrDragInputBehavior>();
+	ClickOrDragBehavior->Initialize(this, this);
+	AddInputBehavior(ClickOrDragBehavior);
 
-	if (WorkingActor.IsValid())
+	// @todo_pcg: Make adjustments to this and check when properties change.
+	CachedRayParams =
 	{
-		WorkingActor->MarkAsGarbage();
-		WorkingActor = nullptr;
-	}
+		.World = GetWorld()
+	};
+}
 
+void UPCGDrawSplineToolBase::Shutdown(const EToolShutdownType ShutdownType)
+{
+	// We have to call this before the super function as it would remove the property sets before we can route the shutdown function to the pcg settings
+	UE::PCG::EditorMode::Tool::Shutdown(this, ShutdownType);
+	
 	Super::Shutdown(ShutdownType);
 }
 
-void UPCGDrawSplineTool::AddSplinePoint(const FVector& HitLocation, const FVector& HitNormal) const
+void UPCGDrawSplineToolBase::AddSplinePoint(const FVector& HitLocation, const FVector& HitNormal) const
 {
 	using namespace PCGDrawSplineTool;
+
+	UPCGInteractiveToolSettings_SplineBase* ToolSettings = GetSplineSettings();
+	
+	USplineComponent* WorkingSpline = ToolSettings->GetWorkingSplineComponent(ToolSettings->DataInstance);
 	if (!WorkingSpline)
 	{
 		return;
@@ -291,15 +319,14 @@ void UPCGDrawSplineTool::AddSplinePoint(const FVector& HitLocation, const FVecto
 	const FVector NewNormal = GetNewPointNormal(HitNormal);
 
 	FVector FinalLocation = HitLocation;
-	const UPCGDrawSplineToolSettings* Settings = CastChecked<UPCGDrawSplineToolSettings>(ToolSettings);
-	switch (Settings->Target.OffsetMode)
+	switch (ToolSettings->Target.OffsetMode)
 	{
 		case EPCGToolDrawTargetOffset::DistanceOffset:
 			check(HitNormal.IsNormalized());
-			FinalLocation += Settings->Target.OffsetDistance * NewNormal;
+			FinalLocation += ToolSettings->Target.OffsetDistance * NewNormal;
 			break;
 		case EPCGToolDrawTargetOffset::ExplicitOffset:
-			FinalLocation += Settings->Target.OffsetVector;
+			FinalLocation += ToolSettings->Target.OffsetVector;
 			break;
 		default: // fall-through
 		case EPCGToolDrawTargetOffset::HitLocation:
@@ -308,13 +335,25 @@ void UPCGDrawSplineTool::AddSplinePoint(const FVector& HitLocation, const FVecto
 
 	WorkingSpline->AddSplinePoint(FinalLocation, ESplineCoordinateSpace::World, /*bUpdate=*/false);
 	WorkingSpline->SetUpVectorAtSplinePoint(NumSplinePoints, NewNormal, ESplineCoordinateSpace::World, /*bUpdate=*/true);
+
+	FPropertyChangedEvent EmptyPropertyChangedEvent(nullptr);
+	if(UE::PCG::EditorMode::Tool::CVarCommitPropertyChangedEventOnTimer.GetValueOnAnyThread())
+	{
+		ToolSettings->PropertyChangedEventQueue.Add(WorkingSpline, EmptyPropertyChangedEvent);
+	}
+	else
+	{
+		FCoreUObjectDelegates::OnObjectPropertyChanged.Broadcast(WorkingSpline, EmptyPropertyChangedEvent);
+	}
 }
 
-FVector UPCGDrawSplineTool::GetNewPointNormal(const FVector& HitNormal) const
+FVector UPCGDrawSplineToolBase::GetNewPointNormal(const FVector& HitNormal) const
 {
-	const UPCGDrawSplineToolSettings* Settings = CastChecked<UPCGDrawSplineToolSettings>(ToolSettings);
+	UPCGInteractiveToolSettings_SplineBase* ToolSettings = GetSplineSettings();
+	
+	USplineComponent* WorkingSpline = ToolSettings->GetWorkingSplineComponent(ToolSettings->DataInstance);
 	const int32 NumSplinePoints = WorkingSpline->GetNumberOfSplinePoints();
-	switch (Settings->Target.NormalMode)
+	switch (ToolSettings->Target.NormalMode)
 	{
 		case EPCGToolDrawTargetNormal::AlignToPrevious:
 		{
@@ -341,105 +380,14 @@ FVector UPCGDrawSplineTool::GetNewPointNormal(const FVector& HitNormal) const
 		case EPCGToolDrawTargetNormal::WorldUp:
 			return FVector::UnitZ();
 		case EPCGToolDrawTargetNormal::Explicit:
-			return Settings->Target.Normal;
+			return ToolSettings->Target.Normal;
 		default: // fall-through
 		case EPCGToolDrawTargetNormal::HitNormal:
 			return HitNormal;
 	}
 }
 
-void UPCGDrawSplineTool::ApplyToolInternal() const
-{
-	GetToolManager()->BeginUndoTransaction(LOCTEXT("DrawSplineTransactionName", "[PCG Tool] Draw Spline"));
-	const UPCGDrawSplineToolSettings* Settings = CastChecked<UPCGDrawSplineToolSettings>(ToolSettings);
-
-	TOptional<FText> OutWarning;
-
-	UWorld* World = GetContextualWorldChecked();
-	USplineComponent* ToolSplineComponent = WorkingSpline.Get();
-
-	{
-		using namespace PCG::EditorMode;
-
-		check(World);
-		check(ToolSplineComponent);
-		
-		if (!Settings->TargetGraph)
-		{
-			OutWarning = LOCTEXT("InvalidGraphInstance", "[Spline Tool] Failed: Invalid Graph Instance");
-		}
-
-		AActor* Actor = nullptr;
-		UPCGComponent* PCGComponent = nullptr;
-
-		switch (Settings->Interaction)
-		{
-		case EPCGToolInteractionMode::CreateNewActor:
-			{
-				Actor = Actor::Spawn(World, Settings->NewActorName, FTransform::Identity);
-			}
-			break;
-		case EPCGToolInteractionMode::AddToSelectedActor:
-			{
-				if (!Settings->TargetActor.IsValid())
-				{
-					OutWarning = LOCTEXT("InvalidActorWarning", "Invalid Actor!");
-				}
-
-				Actor = Settings->TargetActor.Get();
-				PCGComponent = Actor ? PCGComponent::Find(Actor) : nullptr;
-			}
-			break;
-		default:
-			checkNoEntry();
-		}
-
-		// @todo_pcg: Will overwrite the current spline. Could be an option.
-		if (Actor)
-		{
-			USplineComponent* Spline = Spline::FindComponent(Actor);
-			if (!Spline)
-			{
-				Spline = Spline::CreateComponent(Actor, Settings->SplineName, ToolSplineComponent, /*bTransactional=*/true);
-				check(Spline);
-			}
-			else
-			{
-				Spline::CopySplinePoints(*ToolSplineComponent, *Spline, /*bTransactional=*/true);
-
-				// The copy points doesn't trigger a object changed event; we need to trigger one here manually so it gets picked up
-				FPropertyChangedEvent EmptyEvent{ nullptr };
-				FCoreUObjectDelegates::OnObjectPropertyChanged.Broadcast(Spline, EmptyEvent);
-			}
-
-			if(Settings->TargetGraph)
-			{
-				if (!PCGComponent)
-				{
-					PCGComponent = PCGComponent::Create(Actor, Settings->NewPCGComponentName);
-					check(PCGComponent);
-				}
-
-				PCGComponent->SetGraph(Settings->TargetGraph);
-				PCGComponent->GetGraphInstance()->CopyParameterOverrides(WorkingComponent->GetGraphInstance());
-
-				if(Settings->bAutoGenerate)
-				{
-					PCGComponent->Generate();
-				}
-			}
-		}
-
-		if(OutWarning.IsSet())
-		{
-			NotifyError(OutWarning.GetValue());
-		}
-	}
-
-	GetToolManager()->EndUndoTransaction();
-}
-
-FInputRayHit UPCGDrawSplineTool::IsHitByClick(const FInputDeviceRay& ClickPosition)
+FInputRayHit UPCGDrawSplineToolBase::IsHitByClick(const FInputDeviceRay& ClickPosition)
 {
 	using namespace PCG::EditorMode;
 	if (TOptional<Scene::FViewHitResult> Result = ViewportRay(ClickPosition, CachedRayParams))
@@ -450,14 +398,16 @@ FInputRayHit UPCGDrawSplineTool::IsHitByClick(const FInputDeviceRay& ClickPositi
 	return {};
 }
 
-void UPCGDrawSplineTool::OnClicked(const FInputDeviceRay& ClickPosition)
+void UPCGDrawSplineToolBase::OnClicked(const FInputDeviceRay& ClickPosition)
 {
 	using namespace PCGDrawSplineTool;
 
 	if (TOptional<Scene::FViewHitResult> Result = ViewportRay(ClickPosition, CachedRayParams))
 	{
-		const UPCGDrawSplineToolSettings* Settings = CastChecked<UPCGDrawSplineToolSettings>(ToolSettings);
-		switch (Settings->DrawMode)
+	UPCGInteractiveToolSettings_SplineBase* ToolSettings = GetSplineSettings();
+		
+		USplineComponent* WorkingSpline = ToolSettings->GetWorkingSplineComponent(ToolSettings->DataInstance);
+		switch (ToolSettings->DrawMode)
 		{
 			case EPCGDrawSplineDrawMode::ClickAutoTangent:
 			case EPCGDrawSplineDrawMode::FreeDraw:
@@ -467,7 +417,7 @@ void UPCGDrawSplineTool::OnClicked(const FInputDeviceRay& ClickPosition)
 				const int32 PointIndex = WorkingSpline->GetNumberOfSplinePoints() - 1;
 				GetToolManager()->EmitObjectChange(this,
 												   MakeUnique<CommandChange::FSimplePointInsertionChange>(
-													   WorkingSpline.Get(),
+													   WorkingSpline,
 													   Result->ImpactPosition,
 													   WorkingSpline->GetUpVectorAtSplinePoint(PointIndex, ESplineCoordinateSpace::World)),
 												   Constants::AddPointTransactionName);
@@ -485,7 +435,7 @@ void UPCGDrawSplineTool::OnClicked(const FInputDeviceRay& ClickPosition)
 				GetToolManager()->EmitObjectChange(
 					this,
 					MakeUnique<CommandChange::FTangentPointInsertionChange>(
-						WorkingSpline.Get(),
+						WorkingSpline,
 						Result->ImpactPosition,
 						WorkingSpline->GetUpVectorAtSplinePoint(PointIndex, ESplineCoordinateSpace::World),
 						FVector::Zero()),
@@ -496,7 +446,7 @@ void UPCGDrawSplineTool::OnClicked(const FInputDeviceRay& ClickPosition)
 	}
 }
 
-FInputRayHit UPCGDrawSplineTool::CanBeginClickDragSequence(const FInputDeviceRay& DragPosition)
+FInputRayHit UPCGDrawSplineToolBase::CanBeginClickDragSequence(const FInputDeviceRay& DragPosition)
 {
 	using namespace PCG::EditorMode::Scene;
 	if (TOptional<FViewHitResult> Result = ViewportRay(DragPosition, CachedRayParams))
@@ -507,26 +457,30 @@ FInputRayHit UPCGDrawSplineTool::CanBeginClickDragSequence(const FInputDeviceRay
 	return {};
 }
 
-void UPCGDrawSplineTool::OnClickPress(const FInputDeviceRay& ClickPosition)
+void UPCGDrawSplineToolBase::OnClickPress(const FInputDeviceRay& ClickPosition)
 {
 	using namespace PCG::EditorMode::Scene;
 	if (TOptional<FViewHitResult> Result = ViewportRay(ClickPosition, CachedRayParams))
 	{
+		UPCGInteractiveToolSettings_SplineBase* ToolSettings = GetSplineSettings();
+		
 		AddSplinePoint(Result->ImpactPosition, Result->ImpactNormal);
-
-		const UPCGDrawSplineToolSettings* Settings = CastChecked<UPCGDrawSplineToolSettings>(ToolSettings);
-		if (Settings->DrawMode == EPCGDrawSplineDrawMode::FreeDraw)
+		
+		if (ToolSettings->DrawMode == EPCGDrawSplineDrawMode::FreeDraw)
 		{
 			// Cache the index for where this stroke began
-			FreeDrawStrokeStartIndex = WorkingSpline->GetNumberOfSplinePoints() - 1;
+			FreeDrawStrokeStartIndex = ToolSettings->GetWorkingSplineComponent(ToolSettings->DataInstance)->GetNumberOfSplinePoints() - 1;
 		}
 	}
 }
 
-void UPCGDrawSplineTool::OnClickDrag(const FInputDeviceRay& DragPosition)
+void UPCGDrawSplineToolBase::OnClickDrag(const FInputDeviceRay& DragPosition)
 {
 	using namespace PCG::EditorMode::Scene;
 
+	UPCGInteractiveToolSettings_SplineBase* ToolSettings = GetSplineSettings();
+	
+	USplineComponent* WorkingSpline = ToolSettings->GetWorkingSplineComponent(ToolSettings->DataInstance);
 	const int32 NumSplinePoints = WorkingSpline->GetNumberOfSplinePoints();
 	if (!ensure(NumSplinePoints > 0))
 	{
@@ -536,8 +490,7 @@ void UPCGDrawSplineTool::OnClickDrag(const FInputDeviceRay& DragPosition)
 	if (TOptional<FViewHitResult> Result = ViewportRay(DragPosition, CachedRayParams))
 	{
 		const int32 FinalSplinePoint = NumSplinePoints - 1;
-		const UPCGDrawSplineToolSettings* Settings = CastChecked<UPCGDrawSplineToolSettings>(ToolSettings);
-		switch (Settings->DrawMode)
+		switch (ToolSettings->DrawMode)
 		{
 			case EPCGDrawSplineDrawMode::ClickAutoTangent:
 			{
@@ -574,14 +527,14 @@ void UPCGDrawSplineTool::OnClickDrag(const FInputDeviceRay& DragPosition)
 					FVector FinalLocation = Result->ImpactPosition;
 					const FVector Normal = GetNewPointNormal(Result->ImpactNormal);
 
-					switch (Settings->Target.OffsetMode)
+					switch (ToolSettings->Target.OffsetMode)
 					{
 						case EPCGToolDrawTargetOffset::DistanceOffset:
 							check(Result->ImpactNormal.IsNormalized());
-							FinalLocation += Settings->Target.OffsetDistance * Normal;
+							FinalLocation += ToolSettings->Target.OffsetDistance * Normal;
 							break;
 						case EPCGToolDrawTargetOffset::ExplicitOffset:
-							FinalLocation += Settings->Target.OffsetVector;
+							FinalLocation += ToolSettings->Target.OffsetVector;
 							break;
 						default: // fall-through
 						case EPCGToolDrawTargetOffset::HitLocation:
@@ -602,16 +555,20 @@ void UPCGDrawSplineTool::OnClickDrag(const FInputDeviceRay& DragPosition)
 	}
 }
 
-void UPCGDrawSplineTool::OnClickRelease(const FInputDeviceRay& ReleasePosition)
+void UPCGDrawSplineToolBase::OnClickRelease(const FInputDeviceRay& ReleasePosition)
 {
 	// Activate the final position
 	OnClickDrag(ReleasePosition);
 	OnTerminateDragSequence();
 }
 
-void UPCGDrawSplineTool::OnTerminateDragSequence()
+void UPCGDrawSplineToolBase::OnTerminateDragSequence()
 {
 	using namespace PCGDrawSplineTool;
+
+	UPCGInteractiveToolSettings_SplineBase* ToolSettings = GetSplineSettings();
+	USplineComponent* WorkingSpline = ToolSettings->GetWorkingSplineComponent(ToolSettings->DataInstance);
+
 	check(WorkingSpline);
 
 	bDrawTangentForPreviousPoint = false;
@@ -621,17 +578,16 @@ void UPCGDrawSplineTool::OnTerminateDragSequence()
 	const int32 PreviousSplinePoint = NumSplinePoints - 1;
 
 	check(PreviousSplinePoint >= 0);
-
-	const UPCGDrawSplineToolSettings* Settings = CastChecked<UPCGDrawSplineToolSettings>(ToolSettings);
+	
 	// Emit the appropriate undo transaction
-	switch (Settings->DrawMode)
+	switch (ToolSettings->DrawMode)
 	{
 		case EPCGDrawSplineDrawMode::ClickAutoTangent:
 		{
 			GetToolManager()->EmitObjectChange(
 				this,
 				MakeUnique<CommandChange::FSimplePointInsertionChange>(
-					WorkingSpline.Get(),
+					WorkingSpline,
 					WorkingSpline->GetLocationAtSplinePoint(PreviousSplinePoint, ESplineCoordinateSpace::World),
 					WorkingSpline->GetUpVectorAtSplinePoint(PreviousSplinePoint, ESplineCoordinateSpace::World)),
 				Constants::AddPointTransactionName);
@@ -642,7 +598,7 @@ void UPCGDrawSplineTool::OnTerminateDragSequence()
 			GetToolManager()->EmitObjectChange(
 				this,
 				MakeUnique<CommandChange::FTangentPointInsertionChange>(
-					WorkingSpline.Get(),
+					WorkingSpline,
 					WorkingSpline->GetLocationAtSplinePoint(PreviousSplinePoint, ESplineCoordinateSpace::World),
 					WorkingSpline->GetUpVectorAtSplinePoint(PreviousSplinePoint, ESplineCoordinateSpace::World),
 					WorkingSpline->GetTangentAtSplinePoint(PreviousSplinePoint, ESplineCoordinateSpace::World)),
@@ -662,15 +618,35 @@ void UPCGDrawSplineTool::OnTerminateDragSequence()
 
 			GetToolManager()->EmitObjectChange(
 				this,
-				MakeUnique<CommandChange::FStrokeInsertionChange>(WorkingSpline.Get(), HitLocations, UpVectors),
+				MakeUnique<CommandChange::FStrokeInsertionChange>(WorkingSpline, HitLocations, UpVectors),
 				Constants::AddPointTransactionName);
 			break;
 		}
 	}
 }
 
-void UPCGDrawSplineTool::OnTick(const float DeltaTime)
+UPCGInteractiveToolSettings_SplineBase* UPCGDrawSplineToolBase::GetSplineSettings() const
+{
+	TArray<UObject*> AllToolProperties = GetToolProperties(false);
+
+	UObject** FoundProperties = AllToolProperties.FindByPredicate([](const UObject* Candidate)
+	{
+		return Candidate->IsA<UPCGInteractiveToolSettings_SplineBase>();
+	});
+
+	if(FoundProperties)
+	{
+		return Cast<UPCGInteractiveToolSettings_SplineBase>(*FoundProperties);
+	}
+
+	return nullptr;
+}
+
+void UPCGDrawSplineToolBase::OnTick(const float DeltaTime)
 {
+	UPCGInteractiveToolSettings_SplineBase* ToolSettings = GetSplineSettings();
+	
+	USplineComponent* WorkingSpline = ToolSettings->GetWorkingSplineComponent(ToolSettings->DataInstance);
 	if (!WorkingSpline)
 	{
 		GetToolManager()->PostActiveToolShutdownRequest(
@@ -683,10 +659,13 @@ void UPCGDrawSplineTool::OnTick(const float DeltaTime)
 	Super::OnTick(DeltaTime);
 }
 
-void UPCGDrawSplineTool::Render(IToolsContextRenderAPI* RenderAPI)
+void UPCGDrawSplineToolBase::Render(IToolsContextRenderAPI* RenderAPI)
 {
 	using namespace PCG::EditorMode::Scene;
 
+	UPCGInteractiveToolSettings_SplineBase* ToolSettings = GetSplineSettings();
+	
+	USplineComponent* WorkingSpline = ToolSettings->GetWorkingSplineComponent(ToolSettings->DataInstance);
 	if (WorkingSpline)
 	{
 		const int32 LastPointIndex = WorkingSpline->GetNumberOfSplinePoints() - 1;
@@ -700,31 +679,36 @@ void UPCGDrawSplineTool::Render(IToolsContextRenderAPI* RenderAPI)
 	}
 }
 
-void UPCGDrawSplineTool::OnRefresh()
+bool UPCGDrawSplineToolBase::CanAccept() const
 {
-	if (WorkingSpline)
-	{
-		WorkingSpline->PostEditChange();
-	}
+	UPCGInteractiveToolSettings_SplineBase* ToolSettings = GetSplineSettings();
 
-	Super::OnRefresh();
+	USplineComponent* WorkingSpline = ToolSettings->GetWorkingSplineComponent(ToolSettings->DataInstance);
+	return WorkingSpline && WorkingSpline->GetNumberOfSplinePoints() > 0;
 }
 
-bool UPCGDrawSplineTool::CanAccept() const
+void UPCGDrawSplineTool::Setup()
 {
-	return WorkingSpline && WorkingSpline->GetNumberOfSplinePoints() > 0;
+	AddToolPropertySource(ToolSettings);
+	
+	Super::Setup();
+}
+
+void UPCGDrawSplineSurfaceTool::Setup()
+{
+	AddToolPropertySource(ToolSettings);
+	
+	Super::Setup();
 }
 
 bool UPCGDrawSplineToolBuilder::CanBuildTool(const FToolBuilderState& SceneState) const
 {
-	return true;
+	return SplineToolClass != nullptr;
 }
 
 UInteractiveTool* UPCGDrawSplineToolBuilder::BuildTool(const FToolBuilderState& SceneState) const
 {
-	UPCGDrawSplineTool* NewTool = NewObject<UPCGDrawSplineTool>(SceneState.ToolManager);
-
-	NewTool->SetIsSplineSurface(bIsSplineSurface);
+	UPCGDrawSplineToolBase* NewTool = NewObject<UPCGDrawSplineToolBase>(SceneState.ToolManager, SplineToolClass);
 
 	NewTool->SetToolInfo(
 	{
@@ -732,8 +716,18 @@ UInteractiveTool* UPCGDrawSplineToolBuilder::BuildTool(const FToolBuilderState&
 		.ToolDisplayMessage = PCGDrawSplineTool::Constants::ToolDescription,
 		.ToolIcon = FPCGEditorModeStyle::Get().GetBrush("PCGEditorMode.Tools.DrawSpline")
 	});
+	
+	if(UE::PCG::EditorMode::Tool::BuildTool(NewTool) == false)
+	{
+		return nullptr;
+	}
 
 	return NewTool;
 }
 
+void UPCGDrawSplineToolBuilder::SetSplineToolClass(TSubclassOf<UPCGDrawSplineToolBase> InClass)
+{
+	SplineToolClass = InClass;
+}
+
 #undef LOCTEXT_NAMESPACE
diff --git a/Engine/Plugins/PCG/Source/PCGEditor/Private/EditorMode/Tools/PCGInteractiveTool.cpp b/Engine/Plugins/PCG/Source/PCGEditor/Private/EditorMode/Tools/PCGInteractiveTool.cpp
deleted file mode 100644
diff --git a/Engine/Plugins/PCG/Source/PCGEditor/Private/EditorMode/Tools/PCGInteractiveToolBuilder.cpp b/Engine/Plugins/PCG/Source/PCGEditor/Private/EditorMode/Tools/PCGInteractiveToolBuilder.cpp
--- a/Engine/Plugins/PCG/Source/PCGEditor/Private/EditorMode/Tools/PCGInteractiveToolBuilder.cpp
+++ b/Engine/Plugins/PCG/Source/PCGEditor/Private/EditorMode/Tools/PCGInteractiveToolBuilder.cpp
@@ -3,9 +3,59 @@
 #include "EditorMode/Tools/PCGInteractiveToolBuilder.h"
 
 #include "InteractiveToolManager.h"
+#include "EditorMode/Tools/PCGInteractiveToolSettings.h"
 
 #include UE_INLINE_GENERATED_CPP_BY_NAME(PCGInteractiveToolBuilder)
 
+#define LOCTEXT_NAMESPACE "PCGInteractiveToolBuilder"
+
+bool UE::PCG::EditorMode::Tool::BuildTool(UInteractiveTool* InTool)
+{
+	// This function uses reflection to instantiate and initialize the tool settings object.
+	// This allows BP to only create a settings property for it to work, in theory
+	
+	using namespace PCG::EditorMode;
+	
+	const FObjectProperty* SettingsProperty = nullptr;
+	const void* PropertyAddress = nullptr;
+	UClass* SettingsClass = nullptr;
+	for(TPropertyValueIterator<FObjectProperty> PropertyIt(InTool->GetClass(), InTool); PropertyIt; ++PropertyIt)
+	{
+		if(PropertyIt.Key()->IsA(FObjectProperty::StaticClass()))
+		{
+			const FObjectProperty* ObjectProperty = CastFieldChecked<FObjectProperty>(PropertyIt.Key());
+			if(ObjectProperty->PropertyClass->IsChildOf(UPCGInteractiveToolSettings::StaticClass()) && ObjectProperty->PropertyClass->HasAnyClassFlags(EClassFlags::CLASS_Abstract) == false)
+			{
+				SettingsProperty = ObjectProperty;
+				PropertyAddress = PropertyIt.Value();
+				SettingsClass = ObjectProperty->PropertyClass;
+				break;
+			}
+		}
+	}
+	
+	if(SettingsProperty == nullptr || PropertyAddress == nullptr || SettingsClass == nullptr)
+	{
+		FText Message = FText::FormatOrdered(LOCTEXT("PCGSettingsPropertyMissingMessage", "Settings {0} of tool {1} does not have a valid PCGInteractiveToolSettings property. Did you forget to mark it as UPROPERTY?"), SettingsClass ? SettingsClass->GetClass()->GetDisplayNameText() : FText::GetEmpty(), InTool->GetClass()->GetDisplayNameText());
+		InTool->GetToolManager()->DisplayMessage(Message, EToolMessageLevel::UserError);
+		return false;
+	}
+	
+	UPCGInteractiveToolSettings* Settings = NewObject<UPCGInteractiveToolSettings>(InTool, SettingsClass);
+	if(Settings->Initialize(InTool) == false)
+	{
+		Settings->Cancel(InTool);
+		Settings->Shutdown();
+		
+		FText Message = FText::FormatOrdered(LOCTEXT("PCGSettingsInitializationFailedMessage", "Tool {0} failed initialization on settings {1}"), InTool->GetClass()->GetDisplayNameText(), Settings->GetClass()->GetDisplayNameText());
+		InTool->GetToolManager()->DisplayMessage(Message, EToolMessageLevel::UserError);
+		return false;
+	}
+	
+	SettingsProperty->SetObjectPropertyValue(const_cast<void*>(PropertyAddress), Settings);
+	return true;
+}
+
 bool UPCGInteractiveToolBuilder::CanBuildTool(const FToolBuilderState& SceneState) const
 {
 	return Super::CanBuildTool(SceneState);
@@ -16,8 +66,42 @@ UInteractiveTool* UPCGInteractiveToolBuilder::BuildTool(const FToolBuilderState&
 	return Super::BuildTool(SceneState);
 }
 
-void UPCGInteractiveToolBuilder::PostBuildTool(UInteractiveTool* Tool, const FToolBuilderState& SceneState) const
+void UPCGInteractiveToolBuilder::PostSetupTool(UInteractiveTool* Tool, const FToolBuilderState& SceneState) const
+{
+	Super::PostBuildTool(Tool, SceneState);
+	SceneState.ToolManager->DisplayMessage(Tool->GetToolInfo().ToolDisplayMessage, EToolMessageLevel::UserNotification);
+
+	TArray<UObject*> ToolProperties = Tool->GetToolProperties();
+
+	if(ToolGraph)
+	{
+		for(UObject* ToolPropertySet : ToolProperties)
+		{
+			if(UPCGInteractiveToolSettings* Settings = Cast<UPCGInteractiveToolSettings>(ToolPropertySet))
+			{
+				Settings->SetToolGraph(ToolGraph);
+			}
+		}
+	}
+}
+
+void UPCGInteractiveToolWithToolTargetsBuilder::PostSetupTool(UInteractiveTool* Tool, const FToolBuilderState& SceneState) const
 {
 	Super::PostBuildTool(Tool, SceneState);
 	SceneState.ToolManager->DisplayMessage(Tool->GetToolInfo().ToolDisplayMessage, EToolMessageLevel::UserNotification);
+
+	TArray<UObject*> ToolProperties = Tool->GetToolProperties();
+
+	if(ToolGraph)
+	{
+		for(UObject* ToolPropertySet : ToolProperties)
+		{
+			if(UPCGInteractiveToolSettings* Settings = Cast<UPCGInteractiveToolSettings>(ToolPropertySet))
+			{
+				Settings->SetToolGraph(ToolGraph);
+			}
+		}
+	}
 }
+
+#undef LOCTEXT_NAMESPACE
diff --git a/Engine/Plugins/PCG/Source/PCGEditor/Private/EditorMode/Tools/PCGInteractiveToolSettings.cpp b/Engine/Plugins/PCG/Source/PCGEditor/Private/EditorMode/Tools/PCGInteractiveToolSettings.cpp
--- a/Engine/Plugins/PCG/Source/PCGEditor/Private/EditorMode/Tools/PCGInteractiveToolSettings.cpp
+++ b/Engine/Plugins/PCG/Source/PCGEditor/Private/EditorMode/Tools/PCGInteractiveToolSettings.cpp
@@ -1,25 +1,522 @@
 // Copyright Epic Games, Inc. All Rights Reserved.
 
 #include "EditorMode/Tools/PCGInteractiveToolSettings.h"
-
+#include "PCGComponent.h"
 #include "PCGEditorModule.h"
 #include "PCGEditorSettings.h"
 #include "EditorMode/PCGEdModeHelpers.h"
 #include "EditorMode/PCGEdModeSettings.h"
-#include "EditorMode/Tools/PCGInteractiveTool.h"
-#include "Helpers/PCGHelpers.h"
+#include "EditorMode/Tools/Helpers/PCGEdModeEditorUtilities.h"
+#include "Elements/PCGDataFromTool.h"
+#include "Subsystems/PCGSubsystem.h"
+
+#include "InteractiveToolManager.h"
+#include "BaseBehaviors/KeyInputBehavior.h"
 #include "Logging/StructuredLog.h"
 
 #include UE_INLINE_GENERATED_CPP_BY_NAME(PCGInteractiveToolSettings)
 
+#define LOCTEXT_NAMESPACE "PCGInteractiveToolSettings"
+
+namespace UE::PCG::EditorMode::Tool
+{
+	TAutoConsoleVariable<bool> CVarCommitPropertyChangedEventOnTimer(TEXT("pcg.Editor.Tools.CommitPropertyChangedEventOnTimer"), true,
+		TEXT("If true, will queue up and automatically commit property changed events on a timer instead of committing directly."));
+}
+
+void UE::PCG::EditorMode::Tool::Shutdown(UInteractiveTool* InTool, EToolShutdownType ShutdownType)
+{
+	for(UObject* ToolPropertyObject : InTool->GetToolProperties(false))
+	{
+		if(UPCGInteractiveToolSettings* PCGSettings = Cast<UPCGInteractiveToolSettings>(ToolPropertyObject))
+		{
+			if(ShutdownType == EToolShutdownType::Cancel)
+			{
+				PCGSettings->Cancel(InTool);
+			}
+			else if(ShutdownType == EToolShutdownType::Accept)
+			{
+				PCGSettings->Apply(InTool);
+			}
+		}
+	}
+
+	for(UObject* ToolPropertyObject : InTool->GetToolProperties(false))
+	{
+		if(UPCGInteractiveToolSettings* PCGSettings = Cast<UPCGInteractiveToolSettings>(ToolPropertyObject))
+		{
+			PCGSettings->Shutdown();
+		}
+	}	
+}
+
 UPCGInteractiveToolSettings::UPCGInteractiveToolSettings()
 	: NewActorName(GetDefault<UPCGEditorModeSettings>()->DefaultNewActorName)
-	, NewPCGComponentName(GetDefault<UPCGEditorModeSettings>()->DefaultNewPCGComponentName) {}
+	, NewPCGComponentName(GetDefault<UPCGEditorModeSettings>()->DefaultNewPCGComponentName)
+{
+	
+}
+
+bool UPCGInteractiveToolSettings::Initialize(UInteractiveTool* InTool)
+{
+	using namespace UE::PCG::EditorMode;
+
+	UWorld* World = InTool->GetToolManager()->GetContextQueriesAPI()->GetCurrentEditingWorld();
+
+	FToolBuilderState State;
+	InTool->GetToolManager()->GetContextQueriesAPI()->GetCurrentSelectionState(State);
+	TArray<AActor*> SelectedActors = State.SelectedActors;
+
+	// If we have selected one actor, that's the actor we want to edit
+	// This actor might:
+	// - Not have a PCG Component
+	// - Have a PCG Component but no valid working data
+	// - Have a PCG Component with valid working data we want to edit
+	if(SelectedActors.Num() == 1)
+	{		
+		EditActor = SelectedActors[0];
+	}
+	// If we have no selected actor, we spawn a working actor that needs to fulfill the same actor requirements. This actor is fully instantiated when Applied.
+	else
+	{
+		AActor* SpawnedActor = Actor::SpawnWorking(World, FTransform::Identity, GetWorkingActorClass());
+		// our spawned actor is also the actor we are editing
+		EditActor = SpawnedActor;
+
+		GeneratedResources.SpawnedActor = SpawnedActor;
+	}
+
+	if(!ensureMsgf(EditActor.IsValid(), TEXT("The Edit Actor has to be valid at this point. Either as a previously existing actor, or a new spawned actor.")))
+	{
+		return false;
+	}
+	
+	TWeakObjectPtr<UPCGComponent> PCGComponent = FindOrGeneratePCGComponent(*EditActor.Get());
+	if(!ensureMsgf(PCGComponent.IsValid(), TEXT("The PCG Component has to be valid at this point. Either as a previously existing PCG Component, or a transient generated PCG Component.")))
+	{
+		return false;
+	}
+
+	OnModified.AddUObject(this, &UPCGInteractiveToolSettings::OnPropertyModified);
+	
+	// If we don't have a valid graph already, we try to load a default graph
+	if(PCGComponent->GetGraph() == nullptr)
+	{
+		TryLoadDefaultGraph();
+	}
+	// Otherwise, we use the graph we already have
+	else
+	{
+		SetToolGraph(PCGComponent->GetGraph());
+	}
+	
+	// Keyboard setup for switching between data instance identifiers
+	UKeyInputBehavior* KeyInputBehavior = NewObject<UKeyInputBehavior>();
+	KeyInputBehavior->Initialize(this,
+		{
+			EKeys::One, EKeys::Two, EKeys::Three, EKeys::Four, EKeys::Five,
+			EKeys::Six, EKeys::Seven, EKeys::Eight, EKeys::Nine, EKeys::Zero
+		});
+	KeyInputBehavior->bRequireAllKeys = false;
+	InTool->AddInputBehavior(KeyInputBehavior);
+
+	return true;
+}
+
+void UPCGInteractiveToolSettings::Apply(UInteractiveTool* OwningTool)
+{
+	using namespace UE::PCG::EditorMode;
+
+	OwningTool->GetToolManager()->BeginUndoTransaction(LOCTEXT("ApplyToolTransactionName", "Applied PCG Tool"));
+	
+	if (GeneratedResources.SpawnedActor.IsValid())
+	{
+		GeneratedResources.SpawnedActor->ClearFlags(RF_Transient);
+	}
+	
+	if (GeneratedResources.GeneratedPCGComponent.IsValid())
+	{
+		// On Apply, if we have a PCG Component that was generated by the tool, we remove the transient flag
+		GeneratedResources.GeneratedPCGComponent->ClearFlags(RF_Transient);
+	}
+
+	/** Now each working data struct gets to Apply its changes. */
+	for (const auto& WorkingDataMap : GetMutableTypedWorkingDataMap<FPCGInteractiveToolWorkingData>())
+	{
+		FPCGInteractiveToolWorkingDataContext Context;
+		Context.OwningPCGComponent = GetWorkingPCGComponent();
+		Context.DataInstanceIdentifier = WorkingDataMap.Key;
+		Context.OwningActor = Context.OwningPCGComponent->GetOwner();
+		Context.InteractiveTool = OwningTool;
+		Context.PCGSettings = this;
+		
+		WorkingDataMap.Value->OnToolApply(Context);
+	}
+
+   OwningTool->GetToolManager()->EndUndoTransaction();
+}
+
+void UPCGInteractiveToolSettings::Cancel(UInteractiveTool* OwningTool)
+{
+	/** Each working data struct is responsible for wiping or restoring data to the previous state. */
+	for(const auto& WorkingDataMap : GetMutableTypedWorkingDataMap<FPCGInteractiveToolWorkingData>())
+	{
+		FPCGInteractiveToolWorkingDataContext Context;
+		Context.OwningPCGComponent = GetWorkingPCGComponent();
+		Context.DataInstanceIdentifier = WorkingDataMap.Key;
+		Context.OwningActor = Context.OwningPCGComponent->GetOwner();
+		Context.InteractiveTool = OwningTool;
+		Context.PCGSettings = this;
+
+		// OnToolCancel will delete resources generated by the tool data itself
+		WorkingDataMap.Value->OnToolCancel(Context);
+	}
+
+	DeleteGeneratedResources();
+}
+
+void UPCGInteractiveToolSettings::Shutdown()
+{
+	using namespace UE::PCG::EditorMode;
+	
+	OnModified.RemoveAll(this);
+	
+	// We mark it as garbage so it stops ticking
+	MarkAsGarbage();
+	
+	// On Shutdown, we flush the cache and regenerate to make sure our results are up to date
+	if(UPCGComponent* Component = GetWorkingPCGComponent())
+	{
+		if (UPCGSubsystem* PCGSubsystem = Component->GetSubsystem())
+		{
+			PCGSubsystem->FlushCache();
+		}
+		
+		Component->Generate(true);
+	}
+}
+
+void UPCGInteractiveToolSettings::InitializeWorkingData(FName DataInstanceIdentifier)
+{
+	UPCGComponent* PCGComponent = GetWorkingPCGComponent();
+	const UScriptStruct* Struct = GetWorkingDataType();
+	ensureMsgf(Struct->IsChildOf<FPCGInteractiveToolWorkingData>(), TEXT("Working Data Type must inherit from FPCGInteractiveToolWorkingData."));
+
+	FPCGInteractiveToolWorkingDataContext Context;
+	Context.OwningPCGComponent = PCGComponent;
+	Context.DataInstanceIdentifier = DataInstanceIdentifier;
+	Context.OwningActor = Context.OwningPCGComponent->GetOwner();
+	Context.InteractiveTool = GetTypedOuter<UInteractiveTool>();
+	Context.PCGSettings = this;
+	
+	FName WorkingDataIdentifier = GetWorkingDataIdentifier(DataInstanceIdentifier);
+	
+	if (EditActor.IsValid())
+	{
+		// If we don't find the requested working data, we create it
+		if(PCGComponent->ToolDataContainer.ToolData.Contains(WorkingDataIdentifier) == false)
+		{
+			TInstancedStruct<FPCGInteractiveToolWorkingData> NewToolWorkingData;
+			NewToolWorkingData.InitializeAsScriptStruct(Struct);
+			
+			PCGComponent->ToolDataContainer.ToolData.Add(WorkingDataIdentifier, MoveTemp(NewToolWorkingData));
+			GeneratedResources.AddedWorkingDataIdentifiers.Add(WorkingDataIdentifier);
+		}
+	}
+
+	FPCGInteractiveToolWorkingData* WorkingData = GetMutableTypedWorkingData(DataInstanceIdentifier);
+	
+	// If it was freshly constructed, we initialize it. This will allow the structs to create or find components to work with.
+	if(ensureMsgf(WorkingData, TEXT("Working Data has to be allocated at this point!.")))
+	{
+		if(WorkingData->IsInitialized() == false)
+		{
+			WorkingData->OnInitialized().BindUObject(this, &UPCGInteractiveToolSettings::OnWorkingDataInitialized);
+			WorkingData->Initialize(Context);
+			WorkingData->OnInitialized().Unbind();
+		}
+	}
+	else
+	{
+		return;
+	}
+	
+	// The working data needs to be initialized at this point
+	ensureMsgf(WorkingData->IsInitialized(), TEXT("Working Data isn't initialized!. Have you forgotten to call the parent Initialize function?"));
+	
+	// If the data is valid, we call start. If not, we delete it.
+	if(WorkingData->IsValid())
+	{
+		// Notify the working data that we have begun accessing it. This allows to make a 'last state' that the tool can revert to on Cancel.
+		WorkingData->OnToolStart(Context);
+	}
+	else
+	{
+		// If the working data is invalid despite being initialized, we delete and request it again.
+		// This can happen when data dependencies have been deleted outside of PCG, such as deleting a spline component on the actor
+		WorkingData->OnToolCancel(Context);
+		PCGComponent->ToolDataContainer.ToolData.Remove(WorkingDataIdentifier);
+
+		// In which case, requesting new working data will recreate a new spline
+		InitializeWorkingData(DataInstanceIdentifier);
+	}	
+}
+
+void UPCGInteractiveToolSettings::SetToolGraph(UPCGGraphInterface* InGraph)
+{
+	ToolGraph = InGraph;
+
+	// This will call OnPropertyModified, forwarding the change
+	FPropertyChangedEvent Event(FindFProperty<FProperty>(UPCGInteractiveToolSettings::StaticClass(), GET_MEMBER_NAME_CHECKED(UPCGInteractiveToolSettings, ToolGraph)));
+	PostEditChangeProperty(Event);
+}
+
+void UPCGInteractiveToolSettings::SetDataInstance(FName DataInstanceIdentifier)
+{
+	DataInstance = DataInstanceIdentifier;
+	
+	// This will call OnPropertyModified, forwarding the change
+	FPropertyChangedEvent Event(FindFProperty<FProperty>(UPCGInteractiveToolSettings::StaticClass(), GET_MEMBER_NAME_CHECKED(UPCGInteractiveToolSettings, DataInstance)));
+	PostEditChangeProperty(Event);
+}
+
+void UPCGInteractiveToolSettings::SetActorClassToSpawn(TSubclassOf<AActor> Class)
+{
+	ActorClassToSpawn = Class;
+
+	// This will call OnPropertyModified, forwarding the change
+	FPropertyChangedEvent Event(FindFProperty<FProperty>(UPCGInteractiveToolSettings::StaticClass(), GET_MEMBER_NAME_CHECKED(UPCGInteractiveToolSettings, ActorClassToSpawn)));
+	PostEditChangeProperty(Event);
+}
+
+UPCGComponent* UPCGInteractiveToolSettings::FindOrGeneratePCGComponent(AActor& OwningActor) const
+{
+	using namespace UE::PCG::EditorMode;
+	
+	UPCGComponent* PCGComponent = PCGComponent::Find(&OwningActor);
+
+	if(PCGComponent == nullptr)
+	{
+		PCGComponent = PCGComponent::Create(&OwningActor, "ToolPCGComponent");
+		PCGComponent->SetFlags(RF_Transient);
+		GeneratedResources.GeneratedPCGComponent = PCGComponent;
+	}
+
+	return PCGComponent;
+}
+
+void UPCGInteractiveToolSettings::DeleteGeneratedResources() const
+{
+	if(GeneratedResources.SpawnedActor.IsValid())
+	{
+		GeneratedResources.SpawnedActor->Destroy();
+		GeneratedResources.SpawnedActor.Reset();
+	}
+	
+	if(GeneratedResources.GeneratedPCGComponent.IsValid())
+	{
+		GeneratedResources.GeneratedPCGComponent->Cleanup(true);
+		GeneratedResources.GeneratedPCGComponent->DestroyComponent();
+		GeneratedResources.GeneratedPCGComponent.Reset();
+	}
+
+	// If we still have a PCG Component, this means it was not a generated one. We have to remove the working data.
+	WorkingComponentCache.Reset();
+	if(UPCGComponent* PCGComponent = GetWorkingPCGComponent())
+	{
+		for(const FName& AddedWorkingDataIdentifier : GeneratedResources.AddedWorkingDataIdentifiers)
+		{
+			PCGComponent->ToolDataContainer.ToolData.Remove(AddedWorkingDataIdentifier);
+		}
+	}
+}
 
 void UPCGInteractiveToolSettings::TryLoadDefaultGraph()
 {
 	check(!HasAnyFlags(RF_ClassDefaultObject));
-	TargetGraph = CastChecked<UPCGGraphInterface>(GetDefaultGraph().TryLoad(), ECastCheckedType::NullAllowed);
+	SetToolGraph(CastChecked<UPCGGraphInterface>(GetDefaultGraph().TryLoad(), ECastCheckedType::NullAllowed));
+}
+
+void UPCGInteractiveToolSettings::Tick(float DeltaTime)
+{
+	GraphRefreshTimer -= DeltaTime;
+	if (GraphRefreshTimer < 0.f)
+	{
+		OnRefresh();
+	}
+}
+
+bool UPCGInteractiveToolSettings::IsTickable() const
+{
+	return HasAnyFlags(RF_ClassDefaultObject | RF_MirroredGarbage) == false;
+}
+
+TStatId UPCGInteractiveToolSettings::GetStatId() const
+{
+	RETURN_QUICK_DECLARE_CYCLE_STAT(FPCGInteractiveToolSettings, STATGROUP_Tickables);
+}
+
+void UPCGInteractiveToolSettings::OnRefresh()
+{
+	if(HasAnyFlags(RF_MirroredGarbage) == false)
+	{
+		if(bAutoGenerate)
+		{
+			UPCGComponent* PCGComponent = GetWorkingPCGComponent();
+			if(PCGComponent && !PCGComponent->IsCleaningUp() && !PCGComponent->IsGenerating())
+			{
+				PCGComponent->GenerateLocal(true);
+			}
+		}
+
+		// We commit the property changed events now to let the graph regenerate, if the object is tracked
+		CommitPropertyChangedEvents();
+		
+		// Reset the refresh timer
+		GraphRefreshTimer = GetDefault<UPCGEditorModeSettings>()->GraphRefreshRate;
+
+		OnRefreshInternal();
+	}
+}
+
+bool UPCGInteractiveToolSettings::CanEditChange(const FProperty* InProperty) const
+{
+	bool bSuperResult =  Super::CanEditChange(InProperty);
+
+	if(bSuperResult == true)
+	{
+		if(InProperty->GetFName() == GET_MEMBER_NAME_CHECKED(UPCGInteractiveToolSettings, ActorClassToSpawn))
+		{
+			return EditActor.IsValid();
+		}
+	}
+
+	return bSuperResult;
+}
+
+void UPCGInteractiveToolSettings::CommitPropertyChangedEvents()
+{
+	for(auto& Pair : PropertyChangedEventQueue)
+	{
+		if(Pair.Key.IsValid())
+		{
+			FCoreUObjectDelegates::OnObjectPropertyChanged.Broadcast(Pair.Key.Get(), Pair.Value);
+		}
+	}
+
+	PropertyChangedEventQueue.Empty();
+}
+
+void UPCGInteractiveToolSettings::OnKeyPressed(const FKey& InKeyID)
+{
+	TArray<FName> ValidDataInstanceNames = GetValidDataInstanceNamesForGraph();
+
+	int32 Index = INDEX_NONE;
+	
+	if(InKeyID == EKeys::One)
+	{
+		Index = 1;
+	}
+	else if(InKeyID == EKeys::Two)
+	{
+		Index = 2;
+	}
+	else if(InKeyID == EKeys::Three)
+	{
+		Index = 3;
+	}
+	else if(InKeyID == EKeys::Four)
+	{
+		Index = 4;
+	}
+	else if(InKeyID == EKeys::Five)
+	{
+		Index = 5;
+	}
+	else if(InKeyID == EKeys::Six)
+	{
+		Index = 6;
+	}
+	else if(InKeyID == EKeys::Seven)
+	{
+		Index = 7;
+	}
+	else if(InKeyID == EKeys::Eight)
+	{
+		Index = 8;
+	}
+	else if(InKeyID == EKeys::Nine)
+	{
+		Index = 9;
+	}
+	else if(InKeyID == EKeys::Zero)
+	{
+		Index = 10;
+	}
+
+	Index--;
+	
+	if (ValidDataInstanceNames.IsValidIndex(Index))
+	{
+		SetDataInstance(ValidDataInstanceNames[Index]);
+	}
+}
+
+void UPCGInteractiveToolSettings::OnWorkingDataInitialized(FPCGInteractiveToolWorkingData* WorkingData) const
+{
+}
+
+void UPCGInteractiveToolSettings::OnPropertyModified(UObject* Object, FProperty* Property)
+{
+	if(Object == this && Property->GetFName() == GET_MEMBER_NAME_CHECKED(UPCGInteractiveToolSettings, ToolGraph))
+	{
+		GetWorkingPCGComponent()->SetGraph(ToolGraph);
+
+		if(ToolGraph)
+		{
+			// If the graph changes, we also set the default data instance
+			SetActorClassToSpawn(ToolGraph->ToolData.InitialActorClassToSpawn);
+		}
+
+		// We set the data instance to the first one provided by the graph. Uses the default as fallback.
+		FName DefaultDataInstanceName = GetDefaultDataInstanceName();
+		SetDataInstance(DefaultDataInstanceName);
+		
+		// If we have a graph, we update the working data to the available data instances ones from the graph
+		if(ToolGraph)
+		{
+			// We update the working data
+			TArray<FName> ValidDataInstanceNames = GetValidDataInstanceNamesForGraph();
+
+			if(ValidDataInstanceNames.Num() > 0)
+			{
+				for(const FName& DataInstanceIdentifier : ValidDataInstanceNames)
+				{
+					InitializeWorkingData(DataInstanceIdentifier);
+				}
+			}
+			// We make sure at least the default is available so that the tool data exists and can be used even if the graph doesn't make use of it
+			else
+			{
+				InitializeWorkingData(DefaultDataInstanceName);
+			}
+			
+		}
+		// If not, we use the default
+		else
+		{
+			InitializeWorkingData(DefaultDataInstanceName);
+		}
+	}
+	if(Object == this && Property->GetFName() == GET_MEMBER_NAME_CHECKED(UPCGInteractiveToolSettings, ActorClassToSpawn))
+	{
+		// @todo_pcg
+		// Delete previously spawned actor and set up the new one. Ideally recycling the existing PCG Component.
+	}
+}
+
+void UPCGInteractiveToolSettings::OnRefreshInternal()
+{
 }
 
 FSoftObjectPath UPCGInteractiveToolSettings::GetDefaultGraph() const
@@ -35,31 +532,109 @@ FSoftObjectPath UPCGInteractiveToolSettings::GetDefaultGraph() const
 		{
 			return Settings->DefaultGraph.ToSoftObjectPath();
 		}
-	
-		UE_LOGFMT(LogPCGEditor, Warning, "Could not get default graph for settings class {0} in tool {1}", this->GetClass()->GetDisplayNameText().ToString(), GetTypedOuter<UPCGInteractiveTool>()->GetClass()->GetDisplayNameText().ToString());
+		
 		return FSoftObjectPath();
 	}
-
-	UE_LOGFMT(LogPCGEditor, Warning, "Could not find tool settings for settings class {0} in tool {1}. Please add them in the PCG Editor Settings.", this->GetClass()->GetDisplayNameText().ToString(), GetTypedOuter<UPCGInteractiveTool>()->GetClass()->GetDisplayNameText().ToString());
+	
 	return FSoftObjectPath();
 }
 
+TSubclassOf<AActor> UPCGInteractiveToolSettings::GetWorkingActorClass() const
+{
+	return ActorClassToSpawn;
+}
+
 bool UPCGInteractiveToolSettings::GraphAssetFilter(const FAssetData& AssetData) const
 {
-	// This can be checked by subclasses by either finding a tag or more specifically checking for values.
-	if (AssetData.FindTag(GET_MEMBER_NAME_CHECKED(UPCGGraph, CompatibleToolTags)))
+	using namespace UE::PCG::EditorMode;
+	
+	if(AssetData.GetClass()->IsChildOf<UPCGGraphInterface>() == false)
+	{
+		return true;
+	}
+	
+	return !Utility::DoesPCGGraphHaveToolTag(GetToolTag(), AssetData);
+}
+
+UPCGComponent* UPCGInteractiveToolSettings::GetWorkingPCGComponent() const
+{
+	if(WorkingComponentCache.IsValid() == false)
 	{
-		FString ToolTagsString;
-		TArray<FString> Tags;
-		if (AssetData.GetTagValue<FString>(GET_MEMBER_NAME_CHECKED(UPCGGraph, CompatibleToolTags), ToolTagsString))
+		WorkingComponentCache = UE::PCG::EditorMode::PCGComponent::Find(EditActor.Get());
+	}
+	
+	return WorkingComponentCache.Get();
+}
+
+UPCGGraphInstance* UPCGInteractiveToolSettings::GetGraphInstance() const
+{
+	if(UPCGComponent* PCGComponent = GetWorkingPCGComponent())
+	{
+		return PCGComponent->GetGraphInstance();
+	}
+	
+	return nullptr;
+}
+
+TArray<FName> UPCGInteractiveToolSettings::GetValidDataInstanceNamesForGraph() const
+{
+	TArray<FName> Result;
+
+	if (ToolGraph == nullptr)
+	{
+		return Result;
+	}
+
+	TArray<UPCGNode*> Nodes = ToolGraph->GetGraph()->FindNodesWithSettings(UPCGDataFromTool::StaticClass(), false);
+
+	FName ToolTag = GetToolTag();
+	for (UPCGNode* Node : Nodes)
+	{
+		if (UPCGDataFromTool* DataFromTool = Cast<UPCGDataFromTool>(Node->GetSettings()))
 		{
-			Tags = PCGHelpers::GetStringArrayFromCommaSeparatedList(ToolTagsString);
+			if (DataFromTool->ToolTag != ToolTag)
+			{
+				continue;
+			}
+			
+			Result.AddUnique(DataFromTool->DataInstance);
 		}
+	}
+	
+	return Result;
+}
 
-		return !Tags.FindByPredicate([ToolTag = GetToolTag().ToString()](const FString& InTag) { return InTag.Equals(ToolTag, ESearchCase::IgnoreCase); });
+FName UPCGInteractiveToolSettings::GetDefaultDataInstanceName() const
+{
+	TArray<FName> ValidDataInstanceNames = GetValidDataInstanceNamesForGraph();
+	
+	// By default, we use the regular tool tag as the working data identifier, if available
+	// This means a NONE instance name.
+	if (ValidDataInstanceNames.Contains(NAME_None))
+	{
+		return NAME_None;
 	}
 
-	// True == filter the asset
-	return true;
+	// If we have multiple available and they don't include, we choose the first available one
+	if (ValidDataInstanceNames.Num() > 0)
+	{
+		return ValidDataInstanceNames[0];
+	}
+	
+	return NAME_None;
+}
+
+FName UPCGInteractiveToolSettings::GetWorkingDataIdentifier(FName DataInstanceIdentifier) const
+{
+	FNameBuilder DataIdentifierBuilder(GetToolTag());
+
+	if (DataInstanceIdentifier.IsNone() == false)
+	{
+		DataIdentifierBuilder.AppendChar('_');
+		DataIdentifierBuilder.Append(DataInstanceIdentifier.ToString());
+	}
+
+	return DataIdentifierBuilder.ToString();
 }
 
+#undef LOCTEXT_NAMESPACE
diff --git a/Engine/Plugins/PCG/Source/PCGEditor/Private/EditorMode/Tools/Paint/PCGPaintTool.cpp b/Engine/Plugins/PCG/Source/PCGEditor/Private/EditorMode/Tools/Paint/PCGPaintTool.cpp
--- a/Engine/Plugins/PCG/Source/PCGEditor/Private/EditorMode/Tools/Paint/PCGPaintTool.cpp
+++ b/Engine/Plugins/PCG/Source/PCGEditor/Private/EditorMode/Tools/Paint/PCGPaintTool.cpp
@@ -2,7 +2,19 @@
 
 #include "EditorMode/Tools/Paint/PCGPaintTool.h"
 
+#include "InteractiveGizmoManager.h"
+#include "InteractiveToolManager.h"
+#include "MaterialEditingLibrary.h"
+#include "PreviewMesh.h"
+#include "ToolTargetManager.h"
+#include "BaseBehaviors/TwoAxisPropertyEditBehavior.h"
+#include "BaseGizmos/VolumetricBrushStampIndicator.h"
+#include "Components/BoxComponent.h"
 #include "EditorMode/PCGEdModeStyle.h"
+#include "Helpers/PCGHelpers.h"
+#include "TargetInterfaces/PrimitiveComponentBackedTarget.h"
+#include "ToolTargets/PrimitiveComponentToolTarget.h"
+#include "Materials/MaterialInstanceDynamic.h"
 
 #include UE_INLINE_GENERATED_CPP_BY_NAME(PCGPaintTool)
 
@@ -10,73 +22,358 @@
 
 namespace PCGPaintTool::Constants
 {
-	const FText ToolName = LOCTEXT("ToolName", "PaintTool");
-	const FText ToolDescription = LOCTEXT("DisplayMessage", "Tool not yet implemented.");
+	const FText ToolDisplayName = LOCTEXT("ToolName", "PaintTool");
+	const FName ToolTag = "PaintTool";
+	const FText ToolDescription = LOCTEXT("DisplayMessage", "Draw points onto a surface and let the PCG graph process them.");
 }
 
-UPCGPaintSettings::UPCGPaintSettings(const FObjectInitializer& ObjectInitializer)
-	: Brush(Cast<UPCGEdModeBrushBase>(
-		ObjectInitializer.CreateEditorOnlyDefaultSubobject(
-			this,
-			GET_MEMBER_NAME_CHECKED(UPCGPaintSettings, Brush),
-			UPCGEdModeBrushBase::StaticClass(),
-			/*bIsTransient=*/true)))
-{} // @todo_pcg: The brush will inevitably be serialized somewhere as the user modifies it.
-
-FText UPCGPaintSettings::GetToolTag() const
+UPCGInteractiveToolSettings_PaintTool::UPCGInteractiveToolSettings_PaintTool()
 {
-	return PCGPaintTool::Constants::ToolName;
-}
+	// We turn auto generation off, as we'll make use of the tracking system to regenerate instead
+	bAutoGenerate = false;
+} 
 
-FSoftObjectPath UPCGPaintSettings::GetDefaultGraph() const
+FName UPCGInteractiveToolSettings_PaintTool::GetToolTag() const
 {
-	return FSoftObjectPath{};
+	return PCGPaintTool::Constants::ToolTag;
 }
 
-bool UPCGPaintSettings::ApplyTool(UWorld* World, FText& OutWarning) const
+const UScriptStruct* UPCGInteractiveToolSettings_PaintTool::GetWorkingDataType() const
 {
-	// @todo_pcg: Apply the output of the tool here.
-	return false;
+	return FPCGInteractiveToolWorkingData_PointArrayData::StaticStruct();
 }
 
 void UPCGPaintTool::Setup()
 {
-	// @todo_pcg: Set up the tool here.
-	Super::Setup();
+	PropertyClass = UPCGEdModeBrushBase::StaticClass();
+	
+	AddToolPropertySource(ToolSettings);
+
+	AActor* TargetActor = Cast<UPrimitiveComponentToolTarget>(GetTarget())->GetOwnerActor();
 
-	UPCGPaintSettings* Settings = CastChecked<UPCGPaintSettings>(ToolSettings);
+	if(ToolSettings->GeneratedResources.SpawnedActor.IsValid())
+	{
+		ToolSettings->GeneratedResources.SpawnedActor->SetActorLocation(TargetActor->GetActorLocation());
+	}
+	
+	Super::Setup();	
+}
 
-	NotifyError(LOCTEXT("ToolNotYetImplemented", "The tool is not yet implemented."));
+void UPCGPaintTool::Shutdown(EToolShutdownType ShutdownType)
+{
+	UE::PCG::EditorMode::Tool::Shutdown(this, ShutdownType);
+	
+	Super::Shutdown(ShutdownType);
 }
 
-TSubclassOf<UPCGInteractiveToolSettings> UPCGPaintTool::GetSettingsClass() const
+bool UPCGPaintTool::HasAccept() const
 {
-	return UPCGPaintSettings::StaticClass();
+	return true;
 }
 
 bool UPCGPaintTool::CanAccept() const
 {
-	// @todo_pcg: true once implemented
-	return false;
+	return true;
 }
 
-bool UPCGPaintToolBuilder::CanBuildTool(const FToolBuilderState& SceneState) const
+bool UPCGPaintTool::HasCancel() const
 {
 	return true;
 }
 
+void UPCGPaintTool::OnTick(float DeltaTime)
+{
+	Super::OnTick(DeltaTime);
+
+	bool bShiftPressedThisTick = GetShiftToggle();
+	bool bShiftPressedChanged = bShiftPressedThisTick != bShiftPressedLastTick;
+	
+	if(bPaintOnTick)
+	{
+		if(bShiftPressedThisTick)
+		{
+			TryDeletePointsFromBrush();	
+		}
+		else
+		{
+			AddSinglePointFromBrush();
+		}
+
+		bPaintOnTick = false;
+	}
+
+	if(bShiftPressedChanged)
+	{
+		if(UVolumetricBrushStampIndicator* VolumetricStampIndicator = Cast<UVolumetricBrushStampIndicator>(BrushStampIndicator))
+		{
+			if(UMaterialInstanceDynamic* Instance = Cast<UMaterialInstanceDynamic>(VolumetricStampIndicator->GetPreviewMesh()->GetActiveMaterial()))
+			{
+				if(bShiftPressedThisTick)
+				{
+					Instance->SetVectorParameterValue("IndicatorColor", FLinearColor::Red);
+				}
+				else
+				{
+					FLinearColor DefaultColor;
+					Instance->GetVectorParameterDefaultValue(FName("IndicatorColor"), DefaultColor);
+					Instance->SetVectorParameterValue("IndicatorColor", DefaultColor);
+				}
+			}
+		}
+	}
+
+	bShiftPressedLastTick = bShiftPressedThisTick;
+}
+
+void UPCGPaintTool::OnClickPress(const FInputDeviceRay& PressPos)
+{
+	Super::OnClickPress(PressPos);
+
+	bPaintOnTick = true;
+}
+
+void UPCGPaintTool::OnUpdateDrag(const FRay& Ray)
+{
+	Super::OnUpdateDrag(Ray);
+
+	bPaintOnTick = true;
+}
+
+void UPCGPaintTool::SetupBrushStampIndicator()
+{
+	if (!BrushStampIndicator)
+	{
+		// See USkinWeightsPaintTool::SetupBrushStampIndicator
+		FString VolumetricBrushIndicatorGizmoType = TEXT("VolumetricBrushIndicatorGizmoType");
+		// register and spawn brush indicator gizmo
+		GetToolManager()->GetPairedGizmoManager()->RegisterGizmoType(VolumetricBrushIndicatorGizmoType, NewObject<UVolumetricBrushStampIndicatorBuilder>());
+		BrushStampIndicator = GetToolManager()->GetPairedGizmoManager()->CreateGizmo<UVolumetricBrushStampIndicator>(VolumetricBrushIndicatorGizmoType, FString(), this);
+		if (UVolumetricBrushStampIndicator* VolumetricBrushIndicator = Cast<UVolumetricBrushStampIndicator>(BrushStampIndicator))
+		{
+			VolumetricBrushIndicator->MakeBrushIndicatorMesh(this, GetTargetWorld());
+		}
+	}
+}
+
+void UPCGPaintTool::UpdateBrushStampIndicator()
+{
+	if (BrushStampIndicator)
+	{
+		if (BrushEditBehavior.IsValid())
+		{
+			FLinearColor NewColor;
+			if(BrushEditBehavior->IsEditing())
+			{
+				NewColor = FLinearColor::White;
+			}
+			else
+			{
+				if(GetShiftToggle())
+				{
+					NewColor = FLinearColor::Red;
+				}
+				else
+				{
+					NewColor = FLinearColor::Green;
+				}
+			}
+			
+			BrushStampIndicator->LineColor = NewColor;	
+		}
+		
+		BrushStampIndicator->Update(BrushProperties->BrushRadius, LastBrushStamp.WorldPosition, LastBrushStamp.WorldNormal, BrushProperties->BrushFalloffAmount, BrushProperties->BrushStrength);
+	}
+}
+
+void UPCGPaintTool::ShutdownBrushStampIndicator()
+{
+	if (BrushStampIndicator)
+	{
+		FString VolumetricBrushIndicatorGizmoType = TEXT("VolumetricBrushIndicatorGizmoType");
+		
+		GetToolManager()->GetPairedGizmoManager()->DestroyGizmo(BrushStampIndicator);
+		BrushStampIndicator = nullptr;
+		GetToolManager()->GetPairedGizmoManager()->DeregisterGizmoType(VolumetricBrushIndicatorGizmoType);
+	}
+}
+
+double UPCGPaintTool::EstimateMaximumTargetDimension()
+{
+	// We return the target actor's max dimension
+	return ToolSettings->EditActor->GetComponentsBoundingBox().GetExtent().GetAbsMax();
+}
+
+void UPCGPaintTool::AddSinglePointFromBrush()
+{
+	UPCGPointArrayData* PointData = ToolSettings->GetTypedWorkingData<FPCGInteractiveToolWorkingData_PointArrayData>(ToolSettings->DataInstance)->GetPointArrayData();
+	
+	const int32 PointIndex = PointData->GetNumPoints();
+	PointData->SetNumPoints(PointIndex + 1);
+	
+	// Values to assign
+	const float PointSteepness = 1.0f;
+	const FTransform PointTransform = FTransform(LastBrushStamp.WorldPosition);
+	const FVector& Position = PointTransform.GetLocation();
+	const int32 PointSeed = PCGHelpers::ComputeSeed((int)Position.X, (int)Position.Y, (int)Position.Z);
+	const FBox LocalBounds = FBox(FVector(-CurrentBrushRadius), FVector(CurrentBrushRadius));
+	const FVector PointBoundsMin = LocalBounds.Min;
+	const FVector PointBoundsMax = LocalBounds.Max;
+	const int64 PointMetadataEntry = PointData->Metadata->AddEntry();
+
+	// SteepnessRange - Initialize and Allocate if needed 
+	const TConstPCGValueRange<float> ConstSteepnessRange = PointData->GetConstSteepnessValueRange();
+	TOptional<const float> SteepnessSingleValue = ConstSteepnessRange.GetSingleValue();
+	const bool bAllocateSteepness = SteepnessSingleValue.IsSet() && SteepnessSingleValue.GetValue() != PointSteepness;
+	TPCGValueRange<float> SteepnessRange = PointData->GetSteepnessValueRange(bAllocateSteepness);
+	
+	// TransformRange - Initialize and Allocate if needed 
+	const TConstPCGValueRange<FTransform> ConstTransformRange = PointData->GetConstTransformValueRange();
+	TOptional<const FTransform> TransformSingleValue = ConstTransformRange.GetSingleValue();
+	const bool bAllocateTransform = TransformSingleValue.IsSet() && !TransformSingleValue.GetValue().Equals(PointTransform);
+	TPCGValueRange<FTransform> TransformRange = PointData->GetTransformValueRange(bAllocateTransform);
+
+	// SeedRange - Initialize and Allocate if needed 
+	const TConstPCGValueRange<int32> ConstSeedRange = PointData->GetConstSeedValueRange();
+	TOptional<const int32> SeedSingleValue = ConstSeedRange.GetSingleValue();
+	const bool bAllocateSeed = SeedSingleValue.IsSet() && SeedSingleValue.GetValue() != PointSeed;
+	TPCGValueRange<int32> SeedRange = PointData->GetSeedValueRange(bAllocateSeed);
+
+	// BoundsMinRange - Initialize and Allocate if needed 
+	const TConstPCGValueRange<FVector> ConstBoundsMinRange = PointData->GetConstBoundsMinValueRange();
+	TOptional<const FVector> BoundsMinSingleValue = ConstBoundsMinRange.GetSingleValue();
+	const bool bAllocateBoundsMin = BoundsMinSingleValue.IsSet() && BoundsMinSingleValue.GetValue() != PointBoundsMin;
+	TPCGValueRange<FVector> BoundsMinRange = PointData->GetBoundsMinValueRange(bAllocateBoundsMin);
+		
+	// BoundsMaxRange - Initialize and Allocate if needed 
+	const TConstPCGValueRange<FVector> ConstBoundsMaxRange = PointData->GetConstBoundsMaxValueRange();
+	TOptional<const FVector> BoundsMaxSingleValue = ConstBoundsMaxRange.GetSingleValue();
+	const bool bAllocateBoundsMax = BoundsMaxSingleValue.IsSet() && BoundsMaxSingleValue.GetValue() != PointBoundsMax;
+	TPCGValueRange<FVector> BoundsMaxRange = PointData->GetBoundsMaxValueRange(bAllocateBoundsMax);
+
+	// MetadataEntryRange - Initialize and Allocate if needed 
+	const TConstPCGValueRange<int64> ConstMetadataEntryRange = PointData->GetConstMetadataEntryValueRange();
+	TOptional<const int64> MetadataEntrySingleValue = ConstMetadataEntryRange.GetSingleValue();
+	const bool bAllocateMetadataEntry = MetadataEntrySingleValue.IsSet() && MetadataEntrySingleValue.GetValue() != PointMetadataEntry;
+	TPCGValueRange<int64> MetadataEntryRange = PointData->GetMetadataEntryValueRange(bAllocateMetadataEntry);
+
+	// Assign values
+	SteepnessRange[PointIndex] = PointSteepness;
+	TransformRange[PointIndex] = PointTransform;
+	SeedRange[PointIndex] = PointSeed;
+	BoundsMinRange[PointIndex] = PointBoundsMin;
+	BoundsMaxRange[PointIndex] = PointBoundsMax;
+	MetadataEntryRange[PointIndex] = PointMetadataEntry;
+
+	FPCGAttributeIdentifier BrushSizeAttributeIdentifier;
+	BrushSizeAttributeIdentifier.Name = "BrushSize";
+	BrushSizeAttributeIdentifier.MetadataDomain = PointData->Metadata->GetConstDefaultMetadataDomain()->GetDomainID();
+	
+	FPCGMetadataAttribute<float>* BrushSizeAttribute = PointData->Metadata->FindOrCreateAttribute(BrushSizeAttributeIdentifier, 100.f, /*bAllowsInterpolation=*/false, /*bOverrideParent=*/false, /*bOverwriteIfTypeMismatch=*/false);
+	if (BrushSizeAttribute)
+	{
+		BrushSizeAttribute->SetValue(PointMetadataEntry, CurrentBrushRadius);
+	}
+
+	FPCGAttributeIdentifier FalloffAttributeIdentifier;
+	FalloffAttributeIdentifier.Name = "BrushFalloff";
+	FalloffAttributeIdentifier.MetadataDomain = PointData->Metadata->GetConstDefaultMetadataDomain()->GetDomainID();
+	FPCGMetadataAttribute<float>* FalloffAttribute = PointData->Metadata->FindOrCreateAttribute(FalloffAttributeIdentifier, 100.f, /*bAllowsInterpolation=*/false, /*bOverrideParent=*/false, /*bOverwriteIfTypeMismatch=*/false);
+	if (FalloffAttribute)
+	{
+		FalloffAttribute->SetValue(PointMetadataEntry, BrushProperties->BrushFalloffAmount);
+	}
+	
+
+	// @todo_pcg
+	// Instead of using a queue on timer, can we put this on another thread? Direct broadcasting takes too much perf on tick, timed queue results in delays
+
+	// GetToolData dynamically tracks the pointdata object and will rerun the graph when it changes
+	// To avoid doing this on tick, we queue up the changes
+	
+	FPropertyChangedEvent EmptyPropertyChangedEvent(nullptr);
+	if(UE::PCG::EditorMode::Tool::CVarCommitPropertyChangedEventOnTimer.GetValueOnAnyThread())
+	{
+		ToolSettings->PropertyChangedEventQueue.Add(PointData, EmptyPropertyChangedEvent);
+	}
+	else
+	{
+		FCoreUObjectDelegates::OnObjectPropertyChanged.Broadcast(PointData, EmptyPropertyChangedEvent);
+	}
+}
+
+void UPCGPaintTool::TryDeletePointsFromBrush()
+{
+	auto DeletePointsFromDataInstance = [this](FName DataInstanceName)
+	{
+		UPCGPointArrayData* PointData = ToolSettings->GetTypedWorkingData<FPCGInteractiveToolWorkingData_PointArrayData>(DataInstanceName)->GetPointArrayData();
+		FSphere Sphere(LastBrushStamp.WorldPosition, CurrentBrushRadius);
+
+		TArray<FTransform> Transforms = PointData->GetTransformsCopy();
+		for(int32 Index = PointData->GetNumPoints() - 1; Index >= 0 ; --Index)
+		{
+			if(Sphere.IsInside(Transforms[Index].GetLocation()))
+			{
+				PointData->RemoveAt(Index);
+			}
+		}
+
+		FPropertyChangedEvent EmptyPropertyChangedEvent(nullptr);
+		if(UE::PCG::EditorMode::Tool::CVarCommitPropertyChangedEventOnTimer.GetValueOnAnyThread())
+		{
+			ToolSettings->PropertyChangedEventQueue.Add(PointData, EmptyPropertyChangedEvent);
+		}
+		else
+		{
+			FCoreUObjectDelegates::OnObjectPropertyChanged.Broadcast(PointData, EmptyPropertyChangedEvent);
+		}
+	};
+	
+	if(ToolSettings->bErasePointsOfSelectedDataInstance)
+	{
+		DeletePointsFromDataInstance(ToolSettings->DataInstance);
+	}
+	else
+	{
+		for(const FName& DataInstanceName : ToolSettings->GetValidDataInstanceNamesForGraph())
+		{
+			DeletePointsFromDataInstance(DataInstanceName);
+		}
+	}
+	
+}
+
+bool UPCGPaintToolBuilder::CanBuildTool(const FToolBuilderState& SceneState) const
+{
+	return SceneState.TargetManager->CountSelectedAndTargetable(SceneState, GetTargetRequirements()) > 0;
+}
+
 UInteractiveTool* UPCGPaintToolBuilder::BuildTool(const FToolBuilderState& SceneState) const
 {
 	UPCGPaintTool* NewTool = NewObject<UPCGPaintTool>(SceneState.ToolManager);
-
+	
 	NewTool->SetToolInfo(
 	{
-		.ToolDisplayName = PCGPaintTool::Constants::ToolName,
+		.ToolDisplayName = PCGPaintTool::Constants::ToolDisplayName,
 		.ToolDisplayMessage = PCGPaintTool::Constants::ToolDescription,
 		.ToolIcon = FPCGEditorModeStyle::Get().GetBrush("PCGEditorMode.Tools.Paint")
 	});
 
+	UToolTarget* const Target = SceneState.TargetManager->BuildFirstSelectedTargetable(SceneState, GetTargetRequirements());
+	NewTool->SetTarget(Target);
+	NewTool->SetWorld(SceneState.World);
+
+	if(UE::PCG::EditorMode::Tool::BuildTool(NewTool) == false)
+	{
+		return nullptr;
+	}
+	
 	return NewTool;
 }
 
+const FToolTargetTypeRequirements& UPCGPaintToolBuilder::GetTargetRequirements() const
+{
+	static FToolTargetTypeRequirements TypeRequirements(UPrimitiveComponentBackedTarget::StaticClass());
+	return TypeRequirements;
+}
+
 #undef LOCTEXT_NAMESPACE
diff --git a/Engine/Plugins/PCG/Source/PCGEditor/Private/EditorMode/Tools/Volume/PCGCreateVolumeTool.cpp b/Engine/Plugins/PCG/Source/PCGEditor/Private/EditorMode/Tools/Volume/PCGCreateVolumeTool.cpp
--- a/Engine/Plugins/PCG/Source/PCGEditor/Private/EditorMode/Tools/Volume/PCGCreateVolumeTool.cpp
+++ b/Engine/Plugins/PCG/Source/PCGEditor/Private/EditorMode/Tools/Volume/PCGCreateVolumeTool.cpp
@@ -2,6 +2,7 @@
 
 #include "EditorMode/Tools/Volume/PCGCreateVolumeTool.h"
 
+#include "InteractiveToolManager.h"
 #include "EditorMode/PCGEdModeStyle.h"
 
 #include UE_INLINE_GENERATED_CPP_BY_NAME(PCGCreateVolumeTool)
@@ -10,20 +11,27 @@
 
 namespace PCGCreateVolumeTool::Constants
 {
-	const FText ToolName = LOCTEXT("ToolName", "CreateVolume");
+	const FText ToolDisplayName = LOCTEXT("ToolDisplayName", "Create Volume");
+	const FName ToolName = "CreateVolumeTool";
 }
 
-FText UPCGCreateVolumeToolSettings::GetToolTag() const
+FName UPCGInteractiveToolSettings_CreateVolume::GetToolTag() const
 {
 	return PCGCreateVolumeTool::Constants::ToolName;
 }
 
-FSoftObjectPath UPCGCreateVolumeToolSettings::GetDefaultGraph() const
+const UScriptStruct* UPCGInteractiveToolSettings_CreateVolume::GetWorkingDataType() const
+{
+	// @todo_pcg Add createvolume tool data struct here 
+	return nullptr;
+}
+
+FSoftObjectPath UPCGInteractiveToolSettings_CreateVolume::GetDefaultGraph() const
 {
 	return FSoftObjectPath{};
 }
 
-bool UPCGCreateVolumeToolSettings::ApplyTool(UWorld* World, FText& OutWarning) const
+bool UPCGInteractiveToolSettings_CreateVolume::ApplyTool(UWorld* World, FText& OutWarning) const
 {
 	// @todo_pcg: Apply the output of the tool here.
 	return false;
@@ -33,15 +41,15 @@ void UPCGCreateVolumeTool::Setup()
 {
 	// @todo_pcg: Set up the tool here.
 	Super::Setup();
-
-	UPCGCreateVolumeToolSettings* Settings = CastChecked<UPCGCreateVolumeToolSettings>(ToolSettings);
-
-	NotifyError(LOCTEXT("ToolNotYetImplemented", "The tool is not yet implemented."));
+	
+	AddToolPropertySource(ToolSettings);
 }
 
 void UPCGCreateVolumeTool::Shutdown(EToolShutdownType ShutdownType)
 {
-	// @todo_pcg: Clean up the tool here.
+	UE::PCG::EditorMode::Tool::Shutdown(this, ShutdownType);
+
+	Super::Shutdown(ShutdownType);
 }
 
 void UPCGCreateVolumeTool::OnTick(const float DeltaTime)
@@ -60,7 +68,7 @@ void UPCGCreateVolumeTool::Render(IToolsContextRenderAPI* RenderAPI)
 	FPrimitiveDrawInterface* PDI = RenderAPI->GetPrimitiveDrawInterface();
 	check(PDI);
 
-	const UPCGCreateVolumeToolSettings* Settings = CastChecked<UPCGCreateVolumeToolSettings>(ToolSettings);
+	const UPCGInteractiveToolSettings_CreateVolume* Settings = CastChecked<UPCGInteractiveToolSettings_CreateVolume>(ToolSettings);
 
 	/** @todo_pcg: Three click process -
 	 *	First click - Point (Location)
@@ -69,11 +77,6 @@ void UPCGCreateVolumeTool::Render(IToolsContextRenderAPI* RenderAPI)
 	 */
 }
 
-TSubclassOf<UPCGInteractiveToolSettings> UPCGCreateVolumeTool::GetSettingsClass() const
-{
-	return UPCGCreateVolumeToolSettings::StaticClass();
-}
-
 FInputRayHit UPCGCreateVolumeTool::IsHitByClick(const FInputDeviceRay& ClickPosition)
 {
 	using namespace PCG::EditorMode::Scene;
@@ -87,7 +90,7 @@ FInputRayHit UPCGCreateVolumeTool::IsHitByClick(const FInputDeviceRay& ClickPosi
 
 void UPCGCreateVolumeTool::OnClicked(const FInputDeviceRay& ClickPosition)
 {
-	UPCGCreateVolumeToolSettings* Settings = CastChecked<UPCGCreateVolumeToolSettings>(ToolSettings);
+	UPCGInteractiveToolSettings_CreateVolume* Settings = CastChecked<UPCGInteractiveToolSettings_CreateVolume>(ToolSettings);
 
 	/** @todo_pcg: Three click process -
 	 *	First click - Point (Location)
@@ -98,6 +101,10 @@ void UPCGCreateVolumeTool::OnClicked(const FInputDeviceRay& ClickPosition)
 	// @todo_pcg: Different chords to translate, rotate, scale at any step of the process?
 }
 
+void UPCGCreateVolumeTool::ApplyTool() const
+{
+}
+
 bool UPCGCreateVolumeToolBuilder::CanBuildTool(const FToolBuilderState& SceneState) const
 {
 	return true;
@@ -109,10 +116,15 @@ UInteractiveTool* UPCGCreateVolumeToolBuilder::BuildTool(const FToolBuilderState
 
 	NewTool->SetToolInfo(
 	{
-		.ToolDisplayName = PCGCreateVolumeTool::Constants::ToolName,
+		.ToolDisplayName = PCGCreateVolumeTool::Constants::ToolDisplayName,
 		.ToolDisplayMessage = LOCTEXT("CreateVolumeToolDisplayMessage", "Tool not yet implemented."),
 		.ToolIcon = FPCGEditorModeStyle::Get().GetBrush("PCGEditorMode.Tools.CreateVolume")
 	});
+	
+	if(UE::PCG::EditorMode::Tool::BuildTool(NewTool) == false)
+	{
+		return nullptr;
+	}	
 
 	return NewTool;
 }
diff --git a/Engine/Plugins/PCG/Source/PCGEditor/Public/EditorMode/PCGEdModeCommands.h b/Engine/Plugins/PCG/Source/PCGEditor/Public/EditorMode/PCGEdModeCommands.h
--- a/Engine/Plugins/PCG/Source/PCGEditor/Public/EditorMode/PCGEdModeCommands.h
+++ b/Engine/Plugins/PCG/Source/PCGEditor/Public/EditorMode/PCGEdModeCommands.h
@@ -62,15 +62,15 @@ public:
 	/*** PCG EDITOR MODE TOOLS - Line Context **/
 	TSharedPtr<FUICommandInfo> EnableDrawSplineTool;
 	TSharedPtr<FUICommandInfo> EnableDrawSurfaceTool;
-	TArray<TSharedPtr<FUICommandInfo>> GetLineContextToolCommands() const;
+	TArray<TSharedPtr<FUICommandInfo>> GetLineContextToolCommands();
 
 	/*** PCG EDITOR MODE TOOLS - Paint Context **/
 	TSharedPtr<FUICommandInfo> EnablePaintTool;
-	TArray<TSharedPtr<FUICommandInfo>> GetPaintContextToolCommands() const;
+	TArray<TSharedPtr<FUICommandInfo>> GetPaintContextToolCommands();
 
 	/*** PCG EDITOR MODE TOOLS - Volume Context **/
 	TSharedPtr<FUICommandInfo> EnableCreateVolumeTool;
-	TArray<TSharedPtr<FUICommandInfo>> GetVolumeContextToolCommands() const;
+	TArray<TSharedPtr<FUICommandInfo>> GetVolumeContextToolCommands();
 };
 
 /**
diff --git a/Engine/Plugins/PCG/Source/PCGEditor/Public/EditorMode/PCGEdModeToolkit.h b/Engine/Plugins/PCG/Source/PCGEditor/Public/EditorMode/PCGEdModeToolkit.h
--- a/Engine/Plugins/PCG/Source/PCGEditor/Public/EditorMode/PCGEdModeToolkit.h
+++ b/Engine/Plugins/PCG/Source/PCGEditor/Public/EditorMode/PCGEdModeToolkit.h
@@ -24,6 +24,7 @@ public:
 	virtual FText GetToolPaletteDisplayName(FName PaletteName) const override;
 	virtual void OnToolPaletteChanged(FName PaletteName) override;
 	virtual bool HasIntegratedToolPalettes() const override { return false; }
+	virtual void BuildToolPalette(FName Palette, class FToolBarBuilder& ToolbarBuilder) override;
 
 	virtual FText GetActiveToolDisplayName() const override;
 	virtual FText GetActiveToolMessage() const override;
@@ -59,8 +60,8 @@ private:
 	void RegisterToolkitCallbacks();
 	void UnregisterToolkitCallbacks();
 
-	void RegisterToolCallbacks(UPCGInteractiveTool* CurrentTool);
-	void UnregisterToolCallbacks(UPCGInteractiveTool* CurrentTool) const;
+	void RegisterToolCallbacks(UInteractiveTool* CurrentTool);
+	void UnregisterToolCallbacks(UInteractiveTool* CurrentTool) const;
 
 	void OnEditorModeSettingsChanged(UObject* InObject, FPropertyChangedEvent& InChangedEvent);
 
diff --git a/Engine/Plugins/PCG/Source/PCGEditor/Public/EditorMode/Tools/Helpers/PCGEdModeActorHelpers.h b/Engine/Plugins/PCG/Source/PCGEditor/Public/EditorMode/Tools/Helpers/PCGEdModeActorHelpers.h
--- a/Engine/Plugins/PCG/Source/PCGEditor/Public/EditorMode/Tools/Helpers/PCGEdModeActorHelpers.h
+++ b/Engine/Plugins/PCG/Source/PCGEditor/Public/EditorMode/Tools/Helpers/PCGEdModeActorHelpers.h
@@ -4,7 +4,7 @@
 
 #include "Helpers/PCGActorHelpers.h"
 
-namespace PCG::EditorMode
+namespace UE::PCG::EditorMode
 {
 	namespace Tags
 	{
@@ -18,7 +18,7 @@ namespace PCG::EditorMode
 		/** Spawn an actor for the toolkit with explicit parameters. */
 		AActor* Spawn(UWorld* World, FName ActorName, const FTransform& Transform = FTransform::Identity);
 		/** Helper to spawn a temporary working actor for a tool. */
-		AActor* SpawnWorking(UWorld* World, const FTransform& Transform = FTransform::Identity);
+		AActor* SpawnWorking(UWorld* World, const FTransform& Transform = FTransform::Identity, TSubclassOf<AActor> ActorClass = AActor::StaticClass());
 	}
 
 	namespace PCGComponent
diff --git a/Engine/Plugins/PCG/Source/PCGEditor/Public/EditorMode/Tools/Helpers/PCGEdModeEditorUtilities.h b/Engine/Plugins/PCG/Source/PCGEditor/Public/EditorMode/Tools/Helpers/PCGEdModeEditorUtilities.h
new file mode 100644
--- /dev/null
+++ b/Engine/Plugins/PCG/Source/PCGEditor/Public/EditorMode/Tools/Helpers/PCGEdModeEditorUtilities.h
@@ -0,0 +1,23 @@
+﻿// Copyright Epic Games, Inc. All Rights Reserved.
+
+#pragma once
+
+#include "Templates/SharedPointer.h"
+#include "Editor/PCGGraphToolData.h"
+#include "AssetRegistry/AssetData.h"
+
+class FUICommandInfo;
+class FBindingContext;
+ 
+namespace UE::PCG::EditorMode::Utility
+{
+	struct FPCGGraphToolEditorData
+	{
+		FAssetData AssetData;
+		FPCGGraphToolData GraphToolData;
+	};
+	
+	TArray<FPCGGraphToolEditorData> GetGraphToolsWithToolTag(FName ToolTag);
+	bool DoesPCGGraphHaveToolTag(FName ToolTag, const FAssetData& AssetData);
+	TArray<TSharedPtr<FUICommandInfo>> GetOrCreateUICommandsWithToolTag(TSharedRef<FBindingContext> BindingContext, FName ToolTag);
+}
diff --git a/Engine/Plugins/PCG/Source/PCGEditor/Public/EditorMode/Tools/Helpers/PCGEdModeSplineHelpers.h b/Engine/Plugins/PCG/Source/PCGEditor/Public/EditorMode/Tools/Helpers/PCGEdModeSplineHelpers.h
--- a/Engine/Plugins/PCG/Source/PCGEditor/Public/EditorMode/Tools/Helpers/PCGEdModeSplineHelpers.h
+++ b/Engine/Plugins/PCG/Source/PCGEditor/Public/EditorMode/Tools/Helpers/PCGEdModeSplineHelpers.h
@@ -7,8 +7,8 @@
 class USplineComponent;
 class AActor;
 
-namespace PCG::EditorMode::Spline
-{
+namespace UE::PCG::EditorMode::Spline
+{	
 	USplineComponent* CreateComponent(
 		AActor* Actor,
 		FName ComponentName,
@@ -17,7 +17,7 @@ namespace PCG::EditorMode::Spline
 
 	USplineComponent* CreateWorkingComponent(AActor* Actor);
 
-	USplineComponent* FindComponent(const AActor* Actor, const int32 TargetIndex = INDEX_NONE);
+	USplineComponent* Find(const AActor* Actor, const int32 TargetIndex = INDEX_NONE);
 
 	void CopySplinePoints(const USplineComponent& Source, USplineComponent& Destination, const bool bTransactional = false);
 }
diff --git a/Engine/Plugins/PCG/Source/PCGEditor/Public/EditorMode/Tools/Line/PCGDrawSplineTool.h b/Engine/Plugins/PCG/Source/PCGEditor/Public/EditorMode/Tools/Line/PCGDrawSplineTool.h
--- a/Engine/Plugins/PCG/Source/PCGEditor/Public/EditorMode/Tools/Line/PCGDrawSplineTool.h
+++ b/Engine/Plugins/PCG/Source/PCGEditor/Public/EditorMode/Tools/Line/PCGDrawSplineTool.h
@@ -2,13 +2,13 @@
 
 #pragma once
 
-#include "EditorMode/Tools/PCGInteractiveTool.h"
 #include "EditorMode/Tools/PCGInteractiveToolBuilder.h"
 #include "EditorMode/Tools/PCGInteractiveToolCommon.h"
 #include "EditorMode/Tools/PCGInteractiveToolSettings.h"
 #include "EditorMode/Tools/Helpers/PCGEdModeSceneQueryHelpers.h"
 
 #include "BaseBehaviors/BehaviorTargetInterfaces.h"
+#include "Components/SplineComponent.h"
 
 #include "PCGDrawSplineTool.generated.h"
 
@@ -23,21 +23,20 @@ enum class EPCGDrawSplineDrawMode : uint8
 	FreeDraw UMETA(ToolTip = "Drag to place multiple points, with spacing controlled by Min Point Spacing.")
 };
 
-UCLASS(BlueprintType)
-class PCGEDITOR_API UPCGDrawSplineToolSettings : public UPCGInteractiveToolSettings
+UCLASS(Abstract, BlueprintType)
+class PCGEDITOR_API UPCGInteractiveToolSettings_SplineBase : public UPCGInteractiveToolSettings
 {
 	GENERATED_BODY()
 
 public:
-	UPCGDrawSplineToolSettings();
-
-	// ~Begin UPCGInteractiveToolSettings interface
-	virtual FText GetToolTag() const override;
-	// ~End UPCGInteractiveToolSettings interface
-
+	UPCGInteractiveToolSettings_SplineBase();
+	
+	USplineComponent* GetWorkingSplineComponent(FName DataInstanceIdentifier) const;
+	TArray<USplineComponent*> GetWorkingSplineComponents() const;
+	
 	UPROPERTY(EditAnywhere, Category = "Spline")
-	EPCGDrawSplineDrawMode DrawMode = EPCGDrawSplineDrawMode::TangentDrag;
-
+	EPCGDrawSplineDrawMode DrawMode = EPCGDrawSplineDrawMode::FreeDraw;
+	
 	/** Minimum world distance to cover while free drawing before a point will be generated. */
 	UPROPERTY(EditAnywhere, Category = "Spline", meta = (
 			ClampMin = 0,
@@ -54,8 +53,35 @@ public:
 	FName SplineName;
 };
 
-UCLASS(Experimental)
-class PCGEDITOR_API UPCGDrawSplineTool : public UPCGInteractiveTool, public IClickBehaviorTarget, public IClickDragBehaviorTarget
+UCLASS(BlueprintType)
+class PCGEDITOR_API UPCGInteractiveToolSettings_Spline : public UPCGInteractiveToolSettings_SplineBase
+{
+	GENERATED_BODY()
+
+	
+protected:
+	virtual const UScriptStruct* GetWorkingDataType() const override;
+	virtual FName GetToolTag() const override;
+	virtual void OnWorkingDataInitialized(FPCGInteractiveToolWorkingData* WorkingData) const override;
+	virtual void OnPropertyModified(UObject* Object, FProperty* Property) override;
+public:
+	UPROPERTY(EditAnywhere, Category = "Spline")
+	bool bClosedSpline = false;
+};
+
+
+UCLASS(BlueprintType)
+class PCGEDITOR_API UPCGInteractiveToolSettings_SplineSurface : public UPCGInteractiveToolSettings_SplineBase
+{
+	GENERATED_BODY()
+
+protected:
+	virtual const UScriptStruct* GetWorkingDataType() const override;
+	virtual FName GetToolTag() const override;
+};
+
+UCLASS(Abstract, Experimental)
+class PCGEDITOR_API UPCGDrawSplineToolBase : public UInteractiveTool, public IClickBehaviorTarget, public IClickDragBehaviorTarget
 {
 	GENERATED_BODY()
 
@@ -71,11 +97,6 @@ public:
 	virtual void Render(IToolsContextRenderAPI* RenderAPI) override;
 	// ~End UInteractiveTool interface
 
-	// ~Begin UPCGInteractiveTool interface
-	virtual void OnRefresh() override;
-	virtual TSubclassOf<UPCGInteractiveToolSettings> GetSettingsClass() const override { return UPCGDrawSplineToolSettings::StaticClass(); }
-	// ~End UPCGInteractiveTool interface
-
 	// ~Begin IClickBehaviorTarget interface
 	virtual FInputRayHit IsHitByClick(const FInputDeviceRay& ClickPosition) override;
 	virtual void OnClicked(const FInputDeviceRay& ClickPosition) override;
@@ -89,28 +110,47 @@ public:
 	virtual void OnTerminateDragSequence() override;
 	// ~End IClickDragBehaviorTarget interface
 
-	void SetIsSplineSurface(const bool bInIsSplineSurface) { bIsSplineSurface = bInIsSplineSurface; }
-
-protected:
-	virtual void ApplyToolInternal() const override;
+	UPCGInteractiveToolSettings_SplineBase* GetSplineSettings() const;
 private:
 	void AddSplinePoint(const FVector& HitLocation, const FVector& HitNormal) const;
 	FVector GetNewPointNormal(const FVector& HitNormal) const;
 
 	UPROPERTY(Instanced)
 	TObjectPtr<USingleClickOrDragInputBehavior> ClickOrDragBehavior = nullptr;
-
-	UPROPERTY(Transient)
-	TObjectPtr<USplineComponent> WorkingSpline = nullptr;
-
+	
 	PCG::EditorMode::Scene::FViewRayParams CachedRayParams{};
-	bool bIsSplineSurface = false;
 	bool bDrawTangentForPreviousPoint = false;
 	bool bFreeDrawPlacedPreviewPoint = false;
 	int32 FreeDrawStrokeStartIndex = 0;
+
+protected:
 	double CachedMinPointSpacingSquared = 0;
 };
 
+UCLASS(Experimental)
+class PCGEDITOR_API UPCGDrawSplineTool : public UPCGDrawSplineToolBase
+{
+	GENERATED_BODY()
+
+public:
+	virtual void Setup() override;
+public:
+	UPROPERTY(Transient)
+	TObjectPtr<UPCGInteractiveToolSettings_Spline> ToolSettings;
+};
+
+UCLASS(Experimental)
+class PCGEDITOR_API UPCGDrawSplineSurfaceTool : public UPCGDrawSplineToolBase
+{
+	GENERATED_BODY()
+
+public:
+	virtual void Setup() override;
+public:
+	UPROPERTY(Transient)
+	TObjectPtr<UPCGInteractiveToolSettings_SplineSurface> ToolSettings;
+};
+
 UCLASS(Transient)
 class PCGEDITOR_API UPCGDrawSplineToolBuilder : public UPCGInteractiveToolBuilder
 {
@@ -120,8 +160,8 @@ public:
 	virtual bool CanBuildTool(const FToolBuilderState& SceneState) const override;
 	virtual UInteractiveTool* BuildTool(const FToolBuilderState& SceneState) const override;
 
-	void SetIsSplineSurface(const bool bInIsSplineSurface) { bIsSplineSurface = bInIsSplineSurface; }
-
+	void SetSplineToolClass(TSubclassOf<UPCGDrawSplineToolBase> InClass);
 private:
-	bool bIsSplineSurface = false;
+	UPROPERTY()
+	TSubclassOf<UPCGDrawSplineToolBase> SplineToolClass;
 };
diff --git a/Engine/Plugins/PCG/Source/PCGEditor/Public/EditorMode/Tools/PCGInteractiveTool.h b/Engine/Plugins/PCG/Source/PCGEditor/Public/EditorMode/Tools/PCGInteractiveTool.h
deleted file mode 100644
diff --git a/Engine/Plugins/PCG/Source/PCGEditor/Public/EditorMode/Tools/PCGInteractiveToolBuilder.h b/Engine/Plugins/PCG/Source/PCGEditor/Public/EditorMode/Tools/PCGInteractiveToolBuilder.h
--- a/Engine/Plugins/PCG/Source/PCGEditor/Public/EditorMode/Tools/PCGInteractiveToolBuilder.h
+++ b/Engine/Plugins/PCG/Source/PCGEditor/Public/EditorMode/Tools/PCGInteractiveToolBuilder.h
@@ -3,9 +3,15 @@
 #pragma once
 
 #include "InteractiveToolBuilder.h"
+#include "PCGGraph.h"
 
 #include "PCGInteractiveToolBuilder.generated.h"
 
+namespace UE::PCG::EditorMode::Tool
+{
+	bool BuildTool(UInteractiveTool* InTool);
+}
+
 UCLASS()
 class PCGEDITOR_API UPCGInteractiveToolBuilder : public UInteractiveToolBuilder
 {
@@ -15,6 +21,29 @@ public:
 	// ~Begin UInteractiveToolBuilder interface
 	virtual bool CanBuildTool(const FToolBuilderState& SceneState) const override;
 	virtual UInteractiveTool* BuildTool(const FToolBuilderState& SceneState) const override;
-	virtual void PostBuildTool(UInteractiveTool* Tool, const FToolBuilderState& SceneState) const override;
+	virtual void PostSetupTool(UInteractiveTool* Tool, const FToolBuilderState& SceneState) const override;
 	// ~End UInteractiveToolBuilder interface
+
+	void SetToolGraph(UPCGGraph* InGraph) { ToolGraph = InGraph; }
+	UPCGGraph* GetToolGraph() const { return ToolGraph; }
+private:
+	UPROPERTY()
+	TObjectPtr<UPCGGraph> ToolGraph;
 };
+
+UCLASS()
+class PCGEDITOR_API UPCGInteractiveToolWithToolTargetsBuilder : public UInteractiveToolWithToolTargetsBuilder
+{
+	GENERATED_BODY()
+
+public:
+	// ~Begin UInteractiveToolBuilder interface
+	virtual void PostSetupTool(UInteractiveTool* Tool, const FToolBuilderState& SceneState) const override;
+	// ~End UInteractiveToolBuilder interface
+
+	void SetToolGraph(UPCGGraph* InGraph) { ToolGraph = InGraph; }
+	UPCGGraph* GetToolGraph() const { return ToolGraph; }
+private:
+	UPROPERTY()
+	TObjectPtr<UPCGGraph> ToolGraph;
+};
\ No newline at end of file
diff --git a/Engine/Plugins/PCG/Source/PCGEditor/Public/EditorMode/Tools/PCGInteractiveToolCommon.h b/Engine/Plugins/PCG/Source/PCGEditor/Public/EditorMode/Tools/PCGInteractiveToolCommon.h
--- a/Engine/Plugins/PCG/Source/PCGEditor/Public/EditorMode/Tools/PCGInteractiveToolCommon.h
+++ b/Engine/Plugins/PCG/Source/PCGEditor/Public/EditorMode/Tools/PCGInteractiveToolCommon.h
@@ -6,13 +6,6 @@
 
 #include "PCGInteractiveToolCommon.generated.h"
 
-UENUM()
-enum class EPCGToolInteractionMode : uint8
-{
-	CreateNewActor,
-	AddToSelectedActor
-};
-
 UENUM()
 enum class EPCGToolDrawTargetOffset : uint8
 {
diff --git a/Engine/Plugins/PCG/Source/PCGEditor/Public/EditorMode/Tools/PCGInteractiveToolSettings.h b/Engine/Plugins/PCG/Source/PCGEditor/Public/EditorMode/Tools/PCGInteractiveToolSettings.h
--- a/Engine/Plugins/PCG/Source/PCGEditor/Public/EditorMode/Tools/PCGInteractiveToolSettings.h
+++ b/Engine/Plugins/PCG/Source/PCGEditor/Public/EditorMode/Tools/PCGInteractiveToolSettings.h
@@ -8,59 +8,226 @@
 
 #include "InteractiveTool.h"
 #include "NativeGameplayTags.h"
-
+#include "PCGComponent.h"
+#include "TickableEditorObject.h"
+#include "BaseBehaviors/BehaviorTargetInterfaces.h"
 #include "PCGInteractiveToolSettings.generated.h"
 
-UCLASS(Abstract)
-class PCGEDITOR_API UPCGInteractiveToolSettings : public UInteractiveToolPropertySet
+struct FPCGInteractiveToolWorkingData;
+
+namespace UE::PCG::EditorMode::Tool
+{
+	PCGEDITOR_API void Shutdown(UInteractiveTool* InTool, EToolShutdownType ShutdownType);
+	
+	PCGEDITOR_API extern TAutoConsoleVariable<bool> CVarCommitPropertyChangedEventOnTimer;
+}
+
+UCLASS(Abstract, Blueprintable)
+class PCGEDITOR_API UPCGInteractiveToolSettings : public UInteractiveToolPropertySet, public FTickableEditorObject, public IKeyInputBehaviorTarget
 {
 	GENERATED_BODY()
 
 public:
 	UPCGInteractiveToolSettings();
+	virtual ~UPCGInteractiveToolSettings() override {}
 
-	void Initialize(UPCGGraphInstance* InGraphInstance)
-	{
-		GraphInstance = InGraphInstance;
-	}
+	virtual bool Initialize(UInteractiveTool* InTool);
+	virtual void Apply(UInteractiveTool* OwningTool);
+	virtual void Cancel(UInteractiveTool* OwningTool);
+	virtual void Shutdown();
+
+	/** Initializes or finds the working data associated with the data instance identifier.
+	 * If NAME_None, the ToolTag is used as identifier. */
+	void InitializeWorkingData(FName DataInstanceIdentifier = NAME_None);
+	
+	UPCGGraphInterface* GetToolGraph() const { return ToolGraph; }
+	void SetToolGraph(UPCGGraphInterface* InGraph);
+
+	void SetDataInstance(FName DataInstanceIdentifier);
+
+	void SetActorClassToSpawn(TSubclassOf<AActor> Class);
 	
-	/** Tool tag for use with filtering compatible graphs. */
-	virtual FText GetToolTag() const PURE_VIRTUAL(UPCGInteractiveToolSettings::GetToolTag, return FText::GetEmpty(););
 	/** A soft object path to the graph asset to load by default for this tool. */
 	virtual FSoftObjectPath GetDefaultGraph() const;
+	
+	/** Tool tag for use with filtering compatible graphs. */
+	virtual FName GetToolTag() const PURE_VIRTUAL(UPCGInteractiveToolSettings::GetToolTag, return NAME_None;);
+
+	/** The working data type for this settings class. */
+	virtual const UScriptStruct* GetWorkingDataType() const PURE_VIRTUAL(UPCGInteractiveToolSettings::GetWorkingDataType, return FPCGInteractiveToolWorkingData::StaticStruct(););
+	
+	/** The actor to instantiate for edit-purposes when we aren't editing an existing actor. 
+	 * If your target actor has any special requirements, create an actor for that purpose and return it in your settings. */
+	virtual TSubclassOf<AActor> GetWorkingActorClass() const;
 
 	/** Filter details view graph selection. */
 	UFUNCTION()
 	virtual bool GraphAssetFilter(const FAssetData& AssetData) const;
-protected:
 	
-	/** Try to load the default graph into TargetGraph. */
+	UPCGComponent* GetWorkingPCGComponent() const;
+	UPCGGraphInstance* GetGraphInstance() const;
+
+	UFUNCTION()
+	TArray<FName> GetValidDataInstanceNamesForGraph() const;
+	FName GetDefaultDataInstanceName() const;
+	
+	template<typename T = AActor>
+	T* GetTypedWorkingActor() const
+	{
+		static_assert(TIsDerivedFrom<T, AActor>::IsDerived, "T needs to be derived from AActor.");
+		
+		return Cast<T>(EditActor.Get());
+	}
+
+	FName GetWorkingDataIdentifier(FName DataInstanceIdentifier = NAME_None) const;
+	
+	template<typename T = FPCGInteractiveToolWorkingData>
+	T* GetMutableTypedWorkingData(FName DataInstanceIdentifier = NAME_None)
+	{
+		static_assert(TIsDerivedFrom<T, FPCGInteractiveToolWorkingData>::IsDerived, "T needs to be derived from FPCGInteractiveToolWorkingData.");
+
+		FName DataIdentifier = GetWorkingDataIdentifier(DataInstanceIdentifier);
+		
+		UPCGComponent* PCGComponent = GetWorkingPCGComponent();
+		if (PCGComponent->ToolDataContainer.ToolData.Contains(DataIdentifier))
+		{
+			return PCGComponent->ToolDataContainer.ToolData[DataIdentifier].GetMutablePtr<T>(); 
+		}
+
+		checkNoEntry();
+		return nullptr;
+	}
+	
+	template<typename T = FPCGInteractiveToolWorkingData>
+	const T* GetTypedWorkingData(FName DataInstanceIdentifier = NAME_None) const
+	{
+		static_assert(TIsDerivedFrom<T, FPCGInteractiveToolWorkingData>::IsDerived, "T needs to be derived from FPCGInteractiveToolWorkingData.");
+
+		FName WorkingDataIdentifier = GetWorkingDataIdentifier(DataInstanceIdentifier);
+
+		UPCGComponent* PCGComponent = GetWorkingPCGComponent();
+		if(PCGComponent->ToolDataContainer.ToolData.Contains(WorkingDataIdentifier))
+		{
+			return PCGComponent->ToolDataContainer.ToolData[WorkingDataIdentifier].GetPtr<T>(); 
+		}
+		
+		return nullptr;
+	}
+
+	template<typename T = FPCGInteractiveToolWorkingData>
+	TMap<FName, T*> GetMutableTypedWorkingDataMap() const
+	{
+		static_assert(TIsDerivedFrom<T, FPCGInteractiveToolWorkingData>::IsDerived, "T needs to be derived from FPCGInteractiveToolWorkingData.");
+
+		TMap<FName, T*> Result;
+
+		UPCGComponent* PCGComponent = GetWorkingPCGComponent();
+		
+		for(const FName& DataInstanceName : GetValidDataInstanceNamesForGraph())
+		{
+			FName WorkingDataIdentifier = GetWorkingDataIdentifier(DataInstanceName);
+			
+			if(PCGComponent->ToolDataContainer.ToolData.Contains(WorkingDataIdentifier))
+			{
+				Result.Add(WorkingDataIdentifier, PCGComponent->ToolDataContainer.ToolData[WorkingDataIdentifier].GetMutablePtr<T>()); 
+			}
+		}
+
+		return Result;
+	}
+
+	template<typename T = FPCGInteractiveToolWorkingData>
+	TMap<FName, const T*> GetTypedWorkingDataMap() const
+	{
+		static_assert(TIsDerivedFrom<T, FPCGInteractiveToolWorkingData>::IsDerived, "T needs to be derived from FPCGInteractiveToolWorkingData.");
+
+		TMap<FName, const T*> Result;
+
+		UPCGComponent* PCGComponent = GetWorkingPCGComponent();
+		
+		for(const FName& DataInstanceName : GetValidDataInstanceNamesForGraph())
+		{
+			FName WorkingDataIdentifier = GetWorkingDataIdentifier(DataInstanceName);
+			
+			if(PCGComponent->ToolDataContainer.ToolData.Contains(WorkingDataIdentifier))
+			{
+				Result.Add(WorkingDataIdentifier, PCGComponent->ToolDataContainer.ToolData[WorkingDataIdentifier].GetPtr<T>()); 
+			}
+		}
+
+		return Result;
+	}
+protected:
+	/** Try to load the default graph into ToolGraph. */
 	virtual void TryLoadDefaultGraph();
+	
+	virtual void Tick(float DeltaTime) override;
+	virtual bool IsTickable() const override;
+	virtual TStatId GetStatId() const override;
+
+	void CommitPropertyChangedEvents();
+
+	/** IKeyInputBehaviorTarget */
+	virtual void OnKeyPressed(const FKey& KeyID) override;
+
+	/** Allows to modify the just initialized working data from the editor. This can be used if the initialization of working data requires editor data access. */
+	virtual void OnWorkingDataInitialized(FPCGInteractiveToolWorkingData* WorkingData) const;
+	virtual void OnPropertyModified(UObject* Object, FProperty* Property);
+	virtual void OnRefreshInternal();
+private:
+	UPCGComponent* FindOrGeneratePCGComponent(AActor& OwningActor) const;
+	
+	void DeleteGeneratedResources() const;
+	
+	void OnRefresh();
+
+	virtual bool CanEditChange(const FProperty* InProperty) const override;
 public:
-	UPROPERTY(BlueprintReadOnly, EditAnywhere, Category = "PCG", meta = (NoResetToDefault, GetAssetFilter = "GraphAssetFilter"))
-	TObjectPtr<UPCGGraphInterface> TargetGraph = nullptr;
+	UPROPERTY(EditAnywhere, BlueprintReadOnly, Category = "PCG", meta = (NoResetToDefault, GetAssetFilter = "GraphAssetFilter"))
+	TObjectPtr<UPCGGraphInterface> ToolGraph = nullptr;
 
-	UPROPERTY(Transient)
-	TObjectPtr<UPCGGraphInstance> GraphInstance;
+	UPROPERTY(EditAnywhere, BlueprintReadOnly, Category = "PCG", meta = (NoResetToDefault, EditCondition="ToolGraph != nullptr", EditConditionHides, GetOptions="GetValidDataInstanceNamesForGraph"))
+	FName DataInstance = NAME_None;
 	
 	/** Automatically generate the PCG component while using the tool. */
 	UPROPERTY(EditAnywhere, Category = "PCG")
 	bool bAutoGenerate = true;
 
-	/** The tool can create a new actor or modify an existing one. */
-	UPROPERTY(EditAnywhere, Category = "PCG")
-	EPCGToolInteractionMode Interaction = EPCGToolInteractionMode::CreateNewActor;
-
 	// @todo_pcg: Convert to a target object, rather than actor, to ease transition into SceneGraph
 	/** The selected actor to modify with the tool. */
-	UPROPERTY(EditAnywhere, Category = "PCG", meta = (EditCondition = "Interaction == EPCGToolInteractionMode::AddToSelectedActor", EditConditionHides, TransientToolProperty))
-	TWeakObjectPtr<AActor> TargetActor = nullptr;
+	UPROPERTY(EditAnywhere, Category = "PCG", meta = (TransientToolProperty))
+	TWeakObjectPtr<AActor> EditActor = nullptr;
 
 	/** The name to be issued to a newly created actor when the tool is applied. */
-	UPROPERTY(EditAnywhere, DisplayName="Actor Label", Category = "PCG", meta = (EditCondition = "Interaction == EPCGToolInteractionMode::CreateNewActor", EditConditionHides, TransientToolProperty))
+	UPROPERTY(EditAnywhere, DisplayName="Actor Label", Category = "PCG", meta = (TransientToolProperty, EditConditionHides))
 	FName NewActorName;
 
 	/** The name to be issued to a newly created PCG component when the tool is applied. */
 	UPROPERTY(EditAnywhere, DisplayName="Component Name", Category = "PCG", meta = (TransientToolProperty))
 	FName NewPCGComponentName;
+
+	UPROPERTY(EditAnywhere, Category = "PCG", meta=(TransientToolProperty, EditConditionHides))
+	TSubclassOf<AActor> ActorClassToSpawn;
+
+	/** The current time left before a graph refresh (in seconds). */
+	UPROPERTY(meta = (TransientToolProperty))
+	float GraphRefreshTimer = 0.f;
+	
+	UPROPERTY(meta = (TransientToolProperty))
+	mutable TWeakObjectPtr<UPCGComponent> WorkingComponentCache;
+	
+	UPROPERTY(meta = (TransientToolProperty))
+	mutable bool bLockActorToSpawn = false;
+	
+	TMap<TWeakObjectPtr<UObject>, FPropertyChangedEvent> PropertyChangedEventQueue;
+
+	struct FGeneratedResources
+	{
+		TWeakObjectPtr<AActor> SpawnedActor;
+		TWeakObjectPtr<UPCGComponent> GeneratedPCGComponent;
+		/** A list of working data identifiers added this tool instantiation. Used to delete data on Cancel. */
+		TArray<FName> AddedWorkingDataIdentifiers;
+	};
+	
+	mutable FGeneratedResources GeneratedResources;
 };
diff --git a/Engine/Plugins/PCG/Source/PCGEditor/Public/EditorMode/Tools/Paint/PCGPaintTool.h b/Engine/Plugins/PCG/Source/PCGEditor/Public/EditorMode/Tools/Paint/PCGPaintTool.h
--- a/Engine/Plugins/PCG/Source/PCGEditor/Public/EditorMode/Tools/Paint/PCGPaintTool.h
+++ b/Engine/Plugins/PCG/Source/PCGEditor/Public/EditorMode/Tools/Paint/PCGPaintTool.h
@@ -3,53 +3,79 @@
 #pragma once
 
 #include "PCGEdModeBrushBase.h"
-#include "EditorMode/Tools/PCGInteractiveTool.h"
+#include "BaseBehaviors/KeyAsModifierInputBehavior.h"
+#include "BaseBehaviors/BehaviorTargetInterfaces.h"
+#include "Components/BoxComponent.h"
+#include "Data/PCGPointArrayData.h"
 #include "EditorMode/Tools/PCGInteractiveToolBuilder.h"
 #include "EditorMode/Tools/PCGInteractiveToolSettings.h"
 
 #include "PCGPaintTool.generated.h"
 
 UCLASS()
-class PCGEDITOR_API UPCGPaintSettings : public UPCGInteractiveToolSettings
+class PCGEDITOR_API UPCGInteractiveToolSettings_PaintTool : public UPCGInteractiveToolSettings
 {
 	GENERATED_BODY()
 
-	UPCGPaintSettings(const FObjectInitializer& ObjectInitializer);
+public:
+	UPCGInteractiveToolSettings_PaintTool();
 
 	// ~Begin UPCGInteractiveToolSettings interface
-	virtual FText GetToolTag() const override;
-	virtual FSoftObjectPath GetDefaultGraph() const override;
+	virtual FName GetToolTag() const override;
+	virtual const UScriptStruct* GetWorkingDataType() const override;
 	// ~End UPCGInteractiveToolSettings interface
 
-	bool ApplyTool(UWorld* World, FText& OutWarning) const;
-
-public:
-	UPROPERTY(EditAnywhere, Category = "Brush", Instanced, meta = (ShowOnlyInnerProperties))
-	TObjectPtr<UPCGEdModeBrushBase> Brush = nullptr;
+	/** If false, will erase points on all paint data instances. */
+	UPROPERTY(BlueprintReadOnly, EditAnywhere, Category = "Paint")
+	bool bErasePointsOfSelectedDataInstance = true;
 };
 
 UCLASS(Transient)
-class UPCGPaintTool : public UPCGInteractiveTool
+class UPCGPaintTool : public UBaseBrushTool
 {
 	GENERATED_BODY()
 
 public:
 	// ~Begin UPCGInteractiveToolSettings interface
 	virtual void Setup() override;
-	virtual TSubclassOf<UPCGInteractiveToolSettings> GetSettingsClass() const override;
+	virtual void Shutdown(EToolShutdownType ShutdownType) override;
+	virtual bool HasAccept() const override;
 	virtual bool CanAccept() const override;
+	virtual bool HasCancel() const override;
+	virtual void OnTick(float DeltaTime) override;
 	// ~End UPCGInteractiveToolSettings interface
 
 protected:
-	virtual void ApplyToolInternal() const override {}
+	virtual void OnClickPress(const FInputDeviceRay& PressPos) override;
+	virtual void OnUpdateDrag(const FRay& Ray) override;
+	
+	virtual void SetupBrushStampIndicator() override;
+	virtual void UpdateBrushStampIndicator() override;
+	virtual void ShutdownBrushStampIndicator() override;
+	
+	/** This will be used as the brush base, if not specifying a radius directly */
+	virtual double EstimateMaximumTargetDimension() override;
+private:
+	void AddSinglePointFromBrush();
+	void TryDeletePointsFromBrush();
+private:
+	UPROPERTY(Transient)
+	TObjectPtr<UPCGInteractiveToolSettings_PaintTool> ToolSettings;
+
+	UPROPERTY(Transient)
+	bool bPaintOnTick = false;
+
+	UPROPERTY(Transient)
+	bool bShiftPressedLastTick = false;
 };
 
 UCLASS(Transient)
-class UPCGPaintToolBuilder : public UPCGInteractiveToolBuilder
+class UPCGPaintToolBuilder : public UPCGInteractiveToolWithToolTargetsBuilder
 {
 	GENERATED_BODY()
 
 public:
 	virtual bool CanBuildTool(const FToolBuilderState& SceneState) const override;
 	virtual UInteractiveTool* BuildTool(const FToolBuilderState& SceneState) const override;
+	virtual const FToolTargetTypeRequirements& GetTargetRequirements() const override;
 };
diff --git a/Engine/Plugins/PCG/Source/PCGEditor/Public/EditorMode/Tools/Volume/PCGCreateVolumeTool.h b/Engine/Plugins/PCG/Source/PCGEditor/Public/EditorMode/Tools/Volume/PCGCreateVolumeTool.h
--- a/Engine/Plugins/PCG/Source/PCGEditor/Public/EditorMode/Tools/Volume/PCGCreateVolumeTool.h
+++ b/Engine/Plugins/PCG/Source/PCGEditor/Public/EditorMode/Tools/Volume/PCGCreateVolumeTool.h
@@ -2,7 +2,6 @@
 
 #pragma once
 
-#include "EditorMode/Tools/PCGInteractiveTool.h"
 #include "EditorMode/Tools/PCGInteractiveToolBuilder.h"
 #include "EditorMode/Tools/PCGInteractiveToolSettings.h"
 #include "EditorMode/Tools/Helpers/PCGEdModeSceneQueryHelpers.h"
@@ -11,13 +10,14 @@
 #include "PCGCreateVolumeTool.generated.h"
 
 UCLASS()
-class PCGEDITOR_API UPCGCreateVolumeToolSettings : public UPCGInteractiveToolSettings
+class PCGEDITOR_API UPCGInteractiveToolSettings_CreateVolume : public UPCGInteractiveToolSettings
 {
 	GENERATED_BODY()
 
 public:
 	// ~Begin UPCGInteractiveToolSettings interface
-	virtual FText GetToolTag() const override;
+	virtual FName GetToolTag() const override;
+	virtual const UScriptStruct* GetWorkingDataType() const override;
 	virtual FSoftObjectPath GetDefaultGraph() const override;
 	// ~End UPCGInteractiveToolSettings interface
 
@@ -41,7 +41,7 @@ public:
 };
 
 UCLASS(Transient, MinimalAPI)
-class UPCGCreateVolumeTool : public UPCGInteractiveTool, public IClickBehaviorTarget
+class UPCGCreateVolumeTool : public UInteractiveTool, public IClickBehaviorTarget
 {
 	GENERATED_BODY()
 
@@ -57,19 +57,18 @@ public:
 	virtual void Render(IToolsContextRenderAPI* RenderAPI) override;
 	// ~End UInteractiveTool interface
 
-	// ~Begin UPCGInteractiveTool interface
-	virtual TSubclassOf<UPCGInteractiveToolSettings> GetSettingsClass() const override;
-	// ~End UPCGInteractiveTool interface
-
 	// ~Begin IClickBehaviorTarget interface
 	virtual FInputRayHit IsHitByClick(const FInputDeviceRay& ClickPosition) override;
 	virtual void OnClicked(const FInputDeviceRay& ClickPosition) override;
 	// ~End IClickBehaviorTarget interface
 
 protected:
-	virtual void ApplyToolInternal() const override {}
+	void ApplyTool() const;
 private:
 	PCG::EditorMode::Scene::FViewRayParams CachedRayParams;
+
+	UPROPERTY(Transient)
+	TObjectPtr<UPCGInteractiveToolSettings_CreateVolume> ToolSettings;
 };
 
 UCLASS(Transient, MinimalAPI)
diff --git a/Engine/Plugins/Runtime/MeshModelingToolset/Source/ModelingComponents/Private/BaseGizmos/VolumetricBrushStampIndicator.cpp b/Engine/Plugins/Runtime/MeshModelingToolset/Source/ModelingComponents/Private/BaseGizmos/VolumetricBrushStampIndicator.cpp
--- a/Engine/Plugins/Runtime/MeshModelingToolset/Source/ModelingComponents/Private/BaseGizmos/VolumetricBrushStampIndicator.cpp
+++ b/Engine/Plugins/Runtime/MeshModelingToolset/Source/ModelingComponents/Private/BaseGizmos/VolumetricBrushStampIndicator.cpp
@@ -51,6 +51,16 @@ void UVolumetricBrushStampIndicator::MakeBrushIndicatorMesh(UInteractiveTool* Pa
 	}
 }
 
+UPreviewMesh* UVolumetricBrushStampIndicator::GetPreviewMesh() const
+{
+	if (SphereMesh)
+	{
+		return SphereMesh;
+	}
+
+	return nullptr;
+}
+
 void UVolumetricBrushStampIndicator::Shutdown()
 {
 	if (SphereMesh)
diff --git a/Engine/Plugins/Runtime/MeshModelingToolset/Source/ModelingComponents/Public/BaseGizmos/VolumetricBrushStampIndicator.h b/Engine/Plugins/Runtime/MeshModelingToolset/Source/ModelingComponents/Public/BaseGizmos/VolumetricBrushStampIndicator.h
--- a/Engine/Plugins/Runtime/MeshModelingToolset/Source/ModelingComponents/Public/BaseGizmos/VolumetricBrushStampIndicator.h
+++ b/Engine/Plugins/Runtime/MeshModelingToolset/Source/ModelingComponents/Public/BaseGizmos/VolumetricBrushStampIndicator.h
@@ -39,6 +39,7 @@ public:
 
 	MODELINGCOMPONENTS_API void MakeBrushIndicatorMesh(UInteractiveTool* ParentTool, UWorld* World);
 
+	MODELINGCOMPONENTS_API UPreviewMesh* GetPreviewMesh() const;
 private:
 	UPROPERTY()
 	TObjectPtr<UPreviewMesh> SphereMesh;
