[Backout] - CL46083892
#rnx
Original CL Desc
-----------------------------------------------------------------
PCG Editor Mode
- Tool Working Data architecture; data used for processing managed by the PCGComponent
- "Get Tool Data" PCG Node that acts as 'Tool Data Instance'/tool parameter definition
- Paint Mode

diff --git a/Engine/Plugins/PCG/PCG.uplugin b/Engine/Plugins/PCG/PCG.uplugin
--- a/Engine/Plugins/PCG/PCG.uplugin
+++ b/Engine/Plugins/PCG/PCG.uplugin
@@ -43,10 +43,6 @@
 		{
 			"Name": "GeometryProcessing",
 			"Enabled": true
-		},
-		{
-			"Name": "MeshModelingToolset",
-			"Enabled": true
 		}
 	]
 }
diff --git a/Engine/Plugins/PCG/Source/PCG/PCG.Build.cs b/Engine/Plugins/PCG/Source/PCG/PCG.Build.cs
--- a/Engine/Plugins/PCG/Source/PCG/PCG.Build.cs
+++ b/Engine/Plugins/PCG/Source/PCG/PCG.Build.cs
@@ -45,7 +45,6 @@ namespace UnrealBuildTool.Rules
 						"UnrealEd",
 						"Settings",
 						"SourceControl",
-						"InteractiveToolsFramework"
 					});
 			}
 
diff --git a/Engine/Plugins/PCG/Source/PCG/Private/Data/PCGPointArrayData.cpp b/Engine/Plugins/PCG/Source/PCG/Private/Data/PCGPointArrayData.cpp
--- a/Engine/Plugins/PCG/Source/PCG/Private/Data/PCGPointArrayData.cpp
+++ b/Engine/Plugins/PCG/Source/PCG/Private/Data/PCGPointArrayData.cpp
@@ -208,16 +208,6 @@ void UPCGPointArrayData::SetNumPoints(int32 InNumPoints, bool bInitializeValues)
 	}
 }
 
-void UPCGPointArrayData::RemoveAt(int32 Index)
-{
-	if(Index >= 0 && Index < PointArray.GetNumPoints())
-	{
-		FlattenPropertiesIfNeeded();
-		PointArray.RemoveAt(Index);
-		DirtyCache();
-	}
-}
-
 void UPCGPointArrayData::AllocateProperties(EPCGPointNativeProperties Properties)
 {
 	FlattenPropertiesIfNeeded(Properties);
diff --git a/Engine/Plugins/PCG/Source/PCG/Private/Editor/PCGGraphToolData.cpp b/Engine/Plugins/PCG/Source/PCG/Private/Editor/PCGGraphToolData.cpp
deleted file mode 100644
diff --git a/Engine/Plugins/PCG/Source/PCG/Private/Elements/PCGDataFromTool.cpp b/Engine/Plugins/PCG/Source/PCG/Private/Elements/PCGDataFromTool.cpp
deleted file mode 100644
diff --git a/Engine/Plugins/PCG/Source/PCG/Private/PCGGraph.cpp b/Engine/Plugins/PCG/Source/PCG/Private/PCGGraph.cpp
--- a/Engine/Plugins/PCG/Source/PCG/Private/PCGGraph.cpp
+++ b/Engine/Plugins/PCG/Source/PCG/Private/PCGGraph.cpp
@@ -633,33 +633,6 @@ UPCGNode* UPCGGraph::FindNodeWithSettings(const UPCGSettingsInterface* InSetting
 	return NodeFound;
 }
 
-TArray<UPCGNode*> UPCGGraph::FindNodesWithSettings(const TSubclassOf<UPCGSettingsInterface> InSettingsClass, bool bRecursive) const
-{
-	TArray<UPCGNode*> Result;
-
-	auto FindNode = [&Result, InSettingsClass](UPCGNode* InNode)
-	{
-		if (InNode && InNode->GetSettingsInterface() && InNode->GetSettingsInterface()->GetClass() == InSettingsClass)
-		{
-			Result.AddUnique(InNode);
-		}
-
-		// Always return true so we never stop execution
-		return true;
-	};
-
-	if (bRecursive)
-	{
-		ForEachNodeRecursively(FindNode);
-	}
-	else
-	{
-		ForEachNode(FindNode);
-	}
-
-	return Result;
-}
-
 UPCGNode* UPCGGraph::FindNodeByTitleName(FName NodeTitle, bool bRecursive, const TSubclassOf<const UPCGSettings> OptionalClass) const
 {
 	UPCGNode* NodeFound = nullptr;
diff --git a/Engine/Plugins/PCG/Source/PCG/Private/PCGPointArray.cpp b/Engine/Plugins/PCG/Source/PCG/Private/PCGPointArray.cpp
--- a/Engine/Plugins/PCG/Source/PCG/Private/PCGPointArray.cpp
+++ b/Engine/Plugins/PCG/Source/PCG/Private/PCGPointArray.cpp
@@ -35,28 +35,6 @@ void FPCGPointArray::SetNumPoints(int32 InNumPoints, bool bInitializeValues)
 	MetadataEntry.SetNum(NumPoints, /*bAllocate=*/false, bInitializedValues);
 }
 
-void FPCGPointArray::RemoveAt(int32 Index)
-{
-	if(Index >= 0 && Index < NumPoints)
-	{
-		NumPoints -= 1;
-
-		Transform.RemoveAt(Index);
-		Density.RemoveAt(Index);
-		BoundsMin.RemoveAt(Index);
-		BoundsMax.RemoveAt(Index);
-		Color.RemoveAt(Index);
-		Steepness.RemoveAt(Index);
-		Seed.RemoveAt(Index);
-		MetadataEntry.RemoveAt(Index);
-
-		if(NumPoints == 1)
-		{
-			bInitializedValues = false;
-		}
-	}
-}
-
 void FPCGPointArray::Allocate(EPCGPointNativeProperties InProperties)
 {
 	if (EnumHasAnyFlags(InProperties, EPCGPointNativeProperties::Transform))
diff --git a/Engine/Plugins/PCG/Source/PCG/Public/Data/PCGPointArrayData.h b/Engine/Plugins/PCG/Source/PCG/Public/Data/PCGPointArrayData.h
--- a/Engine/Plugins/PCG/Source/PCG/Public/Data/PCGPointArrayData.h
+++ b/Engine/Plugins/PCG/Source/PCG/Public/Data/PCGPointArrayData.h
@@ -36,8 +36,7 @@ public:
 	
 	virtual int32 GetNumPoints() const override { return PointArray.GetNumPoints(); }
 	UE_API virtual void SetNumPoints(int32 InNumPoints, bool bInitializeValues = true) override;
-	UE_API void RemoveAt(int32 Index);
-	
+		
 	UE_API virtual void AllocateProperties(EPCGPointNativeProperties Properties) override;
 	UE_API virtual void FreeProperties(EPCGPointNativeProperties Properties) override;
 	UE_API virtual void MoveRange(int32 RangeStartIndex, int32 MoveToIndex, int32 NumElements) override;
@@ -63,7 +62,7 @@ public:
 	UE_API virtual FPCGPointSeed::ConstValueRange GetConstSeedValueRange() const override;
 	UE_API virtual FPCGPointMetadataEntry::ConstValueRange GetConstMetadataEntryValueRange() const override;
 	//~End UPCGBasePointData Interface
-
+		
 protected:
 	UE_API virtual void InitializeSpatialDataInternal(const FPCGInitializeFromDataParams& InParams) override;
 
diff --git a/Engine/Plugins/PCG/Source/PCG/Public/Editor/PCGGraphToolData.h b/Engine/Plugins/PCG/Source/PCG/Public/Editor/PCGGraphToolData.h
deleted file mode 100644
diff --git a/Engine/Plugins/PCG/Source/PCG/Public/Elements/PCGDataFromTool.h b/Engine/Plugins/PCG/Source/PCG/Public/Elements/PCGDataFromTool.h
deleted file mode 100644
diff --git a/Engine/Plugins/PCG/Source/PCG/Public/PCGComponent.h b/Engine/Plugins/PCG/Source/PCG/Public/PCGComponent.h
--- a/Engine/Plugins/PCG/Source/PCG/Public/PCGComponent.h
+++ b/Engine/Plugins/PCG/Source/PCG/Public/PCGComponent.h
@@ -14,7 +14,6 @@
 #include "ComponentInstanceDataCache.h"
 #include "Components/ActorComponent.h"
 #include "Containers/Map.h"
-#include "Editor/PCGGraphToolData.h"
 #include "Misc/ScopeExit.h"
 #include "UObject/ObjectKey.h"
 #include "Misc/TransactionallySafeCriticalSection.h"
@@ -364,10 +363,6 @@ public:
 	UPROPERTY(BlueprintReadOnly, VisibleAnywhere, Instanced, Category = RuntimeGeneration, meta = (EditCondition = "GenerationTrigger == EPCGComponentGenerationTrigger::GenerateAtRuntime", EditConditionHides))
 	TObjectPtr<UPCGSchedulingPolicyBase> SchedulingPolicy;
 
-	/** This stores working data of a tool; runtime property as the PCG Component will rely on the tool's working data to generate properly. */
-	UPROPERTY(BlueprintReadOnly, VisibleAnywhere, Category = "Tool Data")
-	FPCGInteractiveToolWorkingDataContainer ToolDataContainer;
-	
 #if WITH_EDITORONLY_DATA
 	UPROPERTY(BlueprintReadWrite, EditAnywhere, Category = "Editing Settings", meta = (DisplayName = "Regenerate PCG Volume In Editor", DisplayPriority = 400))
 	bool bRegenerateInEditor = true;
diff --git a/Engine/Plugins/PCG/Source/PCG/Public/PCGGraph.h b/Engine/Plugins/PCG/Source/PCG/Public/PCGGraph.h
--- a/Engine/Plugins/PCG/Source/PCG/Public/PCGGraph.h
+++ b/Engine/Plugins/PCG/Source/PCG/Public/PCGGraph.h
@@ -14,7 +14,6 @@
 #endif // WITH_EDITOR
 
 #include "ComputeFramework/ComputeGraphInstance.h"
-#include "Editor/PCGGraphToolData.h"
 #include "StructUtils/PropertyBag.h"
 #include "UObject/ObjectPtr.h"
 
@@ -179,9 +178,9 @@ public:
 	UPROPERTY(BlueprintReadWrite, EditAnywhere, Category = AssetInfo, AssetRegistrySearchable, meta = (EditCondition = "IsExportToLibraryEnabled", EditConditionHides, PCGNoHash))
 	bool bExposeToLibrary = false;
 
-	/** Contains the data relevant for PCG Editor Mode usage. */
+	/** Signifies the graph has compatibility with the annotated tags. Multiple tags may be comma separated. Ex. "PCGTool1, PCGTool2, PCGTool3" */
 	UPROPERTY(BlueprintReadWrite, EditAnywhere, Category = AssetInfo, AssetRegistrySearchable, meta=(EditCondition = "bExposeToLibrary", EditConditionHides, PCGNoHash))
-	FPCGGraphToolData ToolData;
+	FString CompatibleToolTags;
 
 	/** When enabled, this graph can be executed outside of the world using an editor execution source. */
 	UPROPERTY(BlueprintReadWrite, EditAnywhere, Category = AssetInfo, AssetRegistrySearchable, meta = (PCGNoHash, EditCondition = "CanToggleStandaloneGraph()", EditConditionHides))
@@ -394,8 +393,6 @@ public:
 
 	/** Returns the node with the given settings in the graph, if any */
 	UE_API UPCGNode* FindNodeWithSettings(const UPCGSettingsInterface* InSettings, bool bRecursive = false) const;
-	
-	UE_API TArray<UPCGNode*> FindNodesWithSettings(const TSubclassOf<UPCGSettingsInterface> InSettingsClass, bool bRecursive = false) const;
 
 	/** Returns the first node that matches the given name in the graph, if any.*/
 	UE_API UPCGNode* FindNodeByTitleName(FName NodeTitle, bool bRecursive = false, TSubclassOf<const UPCGSettings> OptionalClass = {}) const;
diff --git a/Engine/Plugins/PCG/Source/PCG/Public/PCGPointArray.h b/Engine/Plugins/PCG/Source/PCG/Public/PCGPointArray.h
--- a/Engine/Plugins/PCG/Source/PCG/Public/PCGPointArray.h
+++ b/Engine/Plugins/PCG/Source/PCG/Public/PCGPointArray.h
@@ -132,23 +132,6 @@ public:
 		return Values[Index];
 	}
 
-	inline void RemoveAt(int32 Index)
-	{
-		ensure(Index >= 0 && Index < NumValues);
-		const int32 OldNum = Values.Num();
-		
-		if (IsAllocated())
-		{
-			Values.RemoveAt(Index);
-			NumValues -= 1; 
-		}
-
-		if (OldNum > 0 && Values.Num() == 1)
-		{
-			Free();
-		}
-	}
-	
 	/** Move ranges of values(moves them back into the array).Used by the PCGAsync api */
 	inline void MoveRange(int32 RangeStartIndex, int32 MoveToIndex, int32 NumElements)
 	{
@@ -266,7 +249,6 @@ public:
 
 	UE_API void SetNumPoints(int32 InNumPoints, bool bInitializeValues);
 	int32 GetNumPoints() const { return NumPoints; }
-	UE_API void RemoveAt(int32 Index);
 		
 	UE_API EPCGPointNativeProperties GetAllocatedProperties() const;
 	UE_API void Allocate(EPCGPointNativeProperties InProperties);
diff --git a/Engine/Plugins/PCG/Source/PCGEditor/PCGEditor.Build.cs b/Engine/Plugins/PCG/Source/PCGEditor/PCGEditor.Build.cs
--- a/Engine/Plugins/PCG/Source/PCGEditor/PCGEditor.Build.cs
+++ b/Engine/Plugins/PCG/Source/PCGEditor/PCGEditor.Build.cs
@@ -69,15 +69,8 @@ namespace UnrealBuildTool.Rules
 					"TypedElementFramework",
 					"TypedElementRuntime",
 					"UnrealEd",
-					"WidgetRegistration", 
-					"ModelingComponents",
-					"ModelingComponentsEditorOnly"
+					"WidgetRegistration",
 				});
-			
-			// @todo_pcg
-			// Set to 1 for development where each tool graph creates a button
-			// Remove again when we settle on a definitive UI
-			PrivateDefinitions.Add("USE_PCG_MODE_GRAPHS_AS_BUTTONS=0");
 		}
 	}
 }
diff --git a/Engine/Plugins/PCG/Source/PCGEditor/Private/Details/PCGInteractiveToolSettingsDetails.cpp b/Engine/Plugins/PCG/Source/PCGEditor/Private/Details/PCGInteractiveToolSettingsDetails.cpp
--- a/Engine/Plugins/PCG/Source/PCGEditor/Private/Details/PCGInteractiveToolSettingsDetails.cpp
+++ b/Engine/Plugins/PCG/Source/PCGEditor/Private/Details/PCGInteractiveToolSettingsDetails.cpp
@@ -52,8 +52,7 @@ void FPCGInteractiveToolSettingsDetails::CustomizeDetails(IDetailLayoutBuilder&
 	}
 	
 	IDetailCategoryBuilder& CategoryBuilder = DetailBuilder.EditCategory("PCG");
-	UPCGGraphInstance* GraphInstance = ToolSettings.IsValid() ? ToolSettings->GetGraphInstance() : nullptr;
-	if(GraphInstance != nullptr && GraphInstance->GetUserParametersStruct()->GetNumPropertiesInBag() > 0)
+	if(ToolSettings.IsValid() && ToolSettings->GraphInstance != nullptr && ToolSettings->GraphInstance->GetUserParametersStruct()->GetNumPropertiesInBag() > 0)
 	{
 		TArray<TSharedRef<IPropertyHandle>> PropertyHandles;
 		CategoryBuilder.GetDefaultProperties(PropertyHandles);
@@ -63,10 +62,10 @@ void FPCGInteractiveToolSettingsDetails::CustomizeDetails(IDetailLayoutBuilder&
 			CategoryBuilder.AddProperty(PropertyHandle);
 
 			// We add the parameters right behind the graph property
-			if(PropertyHandle->GetProperty()->GetFName() == GET_MEMBER_NAME_CHECKED(UPCGInteractiveToolSettings, ToolGraph))
-			{
+			if(PropertyHandle->GetProperty()->GetFName() == GET_MEMBER_NAME_CHECKED(UPCGInteractiveToolSettings, TargetGraph))
+			{			
 				IDetailPropertyRow* PropertyRow = CategoryBuilder.AddExternalObjectProperty(
-				{GraphInstance}, 
+				{ToolSettings->GraphInstance}, 
 				GET_MEMBER_NAME_CHECKED(UPCGGraphInstance, ParametersOverrides),
 				 EPropertyLocation::Default,
 				 FAddPropertyParams());
@@ -87,12 +86,9 @@ void FPCGInteractiveToolSettingsDetails::Rebuild(UObject* Object, const FPropert
 {
 	if(Object == ToolSettings)
 	{
-		if(PropertyChangedChainEvent.Property && PropertyChangedChainEvent.Property->GetFName() == GET_MEMBER_NAME_CHECKED(UPCGInteractiveToolSettings, ToolGraph))
+		if(PropertyChangedChainEvent.Property->GetFName() == GET_MEMBER_NAME_CHECKED(UPCGInteractiveToolSettings, TargetGraph))
 		{
-			if(DetailLayoutBuilder.IsValid())
-			{
-				DetailLayoutBuilder.Pin()->ForceRefreshDetails();
-			}
+			DetailLayoutBuilder.Pin()->ForceRefreshDetails();
 		}
 	}
 }
diff --git a/Engine/Plugins/PCG/Source/PCGEditor/Private/EditorMode/PCGEdMode.cpp b/Engine/Plugins/PCG/Source/PCGEditor/Private/EditorMode/PCGEdMode.cpp
--- a/Engine/Plugins/PCG/Source/PCGEditor/Private/EditorMode/PCGEdMode.cpp
+++ b/Engine/Plugins/PCG/Source/PCGEditor/Private/EditorMode/PCGEdMode.cpp
@@ -28,14 +28,11 @@
 
 #include "ToolContextInterfaces.h"
 #include "ToolMenus.h"
-#include "ToolTargetManager.h"
 #include "EditorMode/PCGEdModeSettings.h"
 #include "EditorMode/PCGEdModeStyle.h"
 #include "EditorMode/Tools/Paint/PCGPaintTool.h"
 #include "Engine/World.h" // FWorldDelegates
 #include "Framework/Application/SlateApplication.h"
-#include "ToolTargets/PrimitiveComponentToolTarget.h"
-#include "ToolTargets/StaticMeshComponentToolTarget.h"
 #include "ViewportToolbar/UnrealEdViewportToolbar.h"
 
 #include UE_INLINE_GENERATED_CPP_BY_NAME(PCGEdMode)
@@ -49,7 +46,7 @@ namespace PCGEditorMode
 	namespace Constants
 	{
 		const FName PCGEditorModeOwnerName = "PCGEditorModeViewportToolBar";
-		const FText PCGEditorModeDisplayName = LOCTEXT("PCGEditorModeDisplayName", "PCG");
+		const FText PCGEditorModeDisplayName = LOCTEXT("PCGEditorModeDisplayName", "PCG (Experimental)");
 	}
 
 	namespace Helpers
@@ -160,9 +157,6 @@ void UPCGEditorMode::Enter()
 		return;
 	}
 
-	UEditorInteractiveToolsContext* InteractiveToolsContext = GetInteractiveToolsContext();
-	InteractiveToolsContext->TargetManager->AddTargetFactory(NewObject<UPrimitiveComponentToolTargetFactory>(InteractiveToolsContext->TargetManager));
-	
 	ToolManager->OnToolPostBuild.AddUObject(this, &UPCGEditorMode::OnToolPostBuild);
 	ToolManager->OnToolShutdownRequest.BindLambda([this](UInteractiveToolManager*, UInteractiveTool*, const EToolShutdownType ShutdownType)
 	{
@@ -177,54 +171,20 @@ void UPCGEditorMode::Enter()
 		}
 	});
 
-#if USE_PCG_MODE_GRAPHS_AS_BUTTONS
-	for(TSharedPtr<FUICommandInfo> ToolCommandInfo : FPCGEditorModeToolCommands::Get().GetLineContextToolCommands())
-	{
-		UPCGDrawSplineToolBuilder* SplineToolBuilder = NewObject<UPCGDrawSplineToolBuilder>();
-		SplineToolBuilder->SetToolGraph(CastChecked<UPCGGraph>(FSoftObjectPath(ToolCommandInfo->GetCommandName().ToString()).TryLoad()));
-		SplineToolBuilder->SetIsSplineSurface(false);
-		RegisterTool(ToolCommandInfo, ToolCommandInfo->GetCommandName().ToString(), SplineToolBuilder);
-	}
-
-	UPCGDrawSplineToolBuilder* SplineSurfaceToolBuilder = NewObject<UPCGDrawSplineToolBuilder>();
-	SplineSurfaceToolBuilder->SetIsSplineSurface(false);
-
-	// Add spline surface tools here
-	
-	for(TSharedPtr<FUICommandInfo> ToolCommandInfo : FPCGEditorModeToolCommands::Get().GetPaintContextToolCommands())
-	{
-		UPCGPaintToolBuilder* PaintToolBuilder = NewObject<UPCGPaintToolBuilder>();
-		PaintToolBuilder->SetToolGraph(CastChecked<UPCGGraph>(FSoftObjectPath(ToolCommandInfo->GetCommandName().ToString()).TryLoad()));
-		
-		RegisterTool(ToolCommandInfo, ToolCommandInfo->GetCommandName().ToString(), PaintToolBuilder);
-	}
-
-	for(TSharedPtr<FUICommandInfo> ToolCommandInfo : FPCGEditorModeToolCommands::Get().GetVolumeContextToolCommands())
-	{
-		UPCGCreateVolumeToolBuilder* VolumeToolBuilder = NewObject<UPCGCreateVolumeToolBuilder>();
-		VolumeToolBuilder->SetToolGraph(CastChecked<UPCGGraph>(FSoftObjectPath(ToolCommandInfo->GetCommandName().ToString()).TryLoad()));
-		RegisterTool(ToolCommandInfo, ToolCommandInfo->GetCommandName().ToString(), VolumeToolBuilder);
-	}
-#else
 	const FPCGEditorModeToolCommands& ToolManagerCommands = FPCGEditorModeToolCommands::Get();
 
-	UPCGDrawSplineToolBuilder* DrawSplineToolBuilder = NewObject<UPCGDrawSplineToolBuilder>();
-	DrawSplineToolBuilder->SetSplineToolClass(UPCGDrawSplineTool::StaticClass());
-	RegisterTool(ToolManagerCommands.EnableDrawSplineTool, TEXT("DrawSpline"), DrawSplineToolBuilder);
-
-	UPCGDrawSplineToolBuilder* DrawSurfaceToolBuilder = NewObject<UPCGDrawSplineToolBuilder>();
-	DrawSurfaceToolBuilder->SetSplineToolClass(UPCGDrawSplineSurfaceTool::StaticClass());
-	RegisterTool(ToolManagerCommands.EnableDrawSurfaceTool, TEXT("DrawSurface"), DrawSurfaceToolBuilder);
-
-	UPCGPaintToolBuilder* PaintToolBuilder = NewObject<UPCGPaintToolBuilder>();
-	RegisterTool(ToolManagerCommands.EnablePaintTool, TEXT("Paint"), PaintToolBuilder);
+	// @todo_pcg: Could be abstracted out and have the tool identifier come from the label
+	UPCGDrawSplineToolBuilder* DrawSplineBuilder = NewObject<UPCGDrawSplineToolBuilder>();
+	DrawSplineBuilder->SetIsSplineSurface(false);
+	RegisterTool(ToolManagerCommands.EnableDrawSplineTool, TEXT("DrawSpline"), DrawSplineBuilder);
+	UPCGDrawSplineToolBuilder* DrawSurfaceBuilder = NewObject<UPCGDrawSplineToolBuilder>();
+	DrawSurfaceBuilder->SetIsSplineSurface(true);
+	RegisterTool(ToolManagerCommands.EnableDrawSurfaceTool, TEXT("DrawSurface"), DrawSurfaceBuilder);
 
-	// @TODO_PCG Reenable Volume Tool when it makes sense
-	// UPCGCreateVolumeToolBuilder* CreateVolumeToolBuilder = NewObject<UPCGCreateVolumeToolBuilder>();
-	// RegisterTool(ToolManagerCommands.EnableCreateVolumeTool, TEXT("CreateVolume"), CreateVolumeToolBuilder);
+	RegisterTool(ToolManagerCommands.EnablePaintTool, TEXT("Paint"), NewObject<UPCGPaintToolBuilder>());
+	RegisterTool(ToolManagerCommands.EnableCreateVolumeTool, TEXT("CreateVolume"), NewObject<UPCGCreateVolumeToolBuilder>());
 
-	ToolManager->SelectActiveToolType(EToolSide::Left, "DrawSpline");
-#endif
+	ToolManager->SelectActiveToolType(EToolSide::Left, TEXT("Line"));
 }
 
 void UPCGEditorMode::Exit()
diff --git a/Engine/Plugins/PCG/Source/PCGEditor/Private/EditorMode/PCGEdModeCommands.cpp b/Engine/Plugins/PCG/Source/PCGEditor/Private/EditorMode/PCGEdModeCommands.cpp
--- a/Engine/Plugins/PCG/Source/PCGEditor/Private/EditorMode/PCGEdModeCommands.cpp
+++ b/Engine/Plugins/PCG/Source/PCGEditor/Private/EditorMode/PCGEdModeCommands.cpp
@@ -2,16 +2,11 @@
 
 #include "EditorMode/PCGEdModeCommands.h"
 
-#include "PCGEditorModule.h"
-#include "AssetRegistry/AssetRegistryModule.h"
 #include "EditorMode/PCGEdModeSettings.h"
 #include "EditorMode/PCGEdModeStyle.h"
-#include "EditorMode/Tools/Helpers/PCGEdModeEditorUtilities.h"
 #include "EditorMode/Tools/Line/PCGDrawSplineTool.h"
 #include "EditorMode/Tools/Paint/PCGPaintTool.h"
 #include "EditorMode/Tools/Volume/PCGCreateVolumeTool.h"
-#include "Logging/StructuredLog.h"
-#include "Misc/OutputDeviceNull.h"
 
 #define LOCTEXT_NAMESPACE "PCGEditorMode"
 
@@ -90,6 +85,8 @@ void FPCGEditorModeToolCommands::RegisterCommands()
 	PCG_TOOL_COMMAND_INFO(EnableCreateVolumeTool, "CreateVolume", "Create Volume", "Create a PCG volume in the scene.", "PCGEditorMode.Tools.CreateVolume");
 }
 
+#undef PCG_TOOL_COMMAND_INFO
+
 void FPCGEditorModeToolCommands::RegisterAllToolCommands()
 {
 	FPCGDrawSplineCommands::Register();
@@ -109,31 +106,19 @@ void FPCGEditorModeToolCommands::GetToolDefaultObjectList(TArray<UInteractiveToo
 	// @todo_pcg: To be populated with tool commands for bindings.
 }
 
-TArray<TSharedPtr<FUICommandInfo>> FPCGEditorModeToolCommands::GetLineContextToolCommands()
+TArray<TSharedPtr<FUICommandInfo>> FPCGEditorModeToolCommands::GetLineContextToolCommands() const
 {
-#if USE_PCG_MODE_GRAPHS_AS_BUTTONS
-	return PCG::EditorMode::Utility::GetOrCreateUICommandsWithToolTag(AsShared(), "SplineTool");
-#else
 	return {EnableDrawSplineTool, EnableDrawSurfaceTool};
-#endif
 }
 
-TArray<TSharedPtr<FUICommandInfo>> FPCGEditorModeToolCommands::GetPaintContextToolCommands()
+TArray<TSharedPtr<FUICommandInfo>> FPCGEditorModeToolCommands::GetPaintContextToolCommands() const
 {
-#if USE_PCG_MODE_GRAPHS_AS_BUTTONS
-	return PCG::EditorMode::Utility::GetOrCreateUICommandsWithToolTag(AsShared(), "PaintTool");
-#else
 	return {EnablePaintTool};
-#endif
 }
 
-TArray<TSharedPtr<FUICommandInfo>> FPCGEditorModeToolCommands::GetVolumeContextToolCommands()
+TArray<TSharedPtr<FUICommandInfo>> FPCGEditorModeToolCommands::GetVolumeContextToolCommands() const
 {
-#if USE_PCG_MODE_GRAPHS_AS_BUTTONS
-	return PCG::EditorMode::Utility::GetOrCreateUICommandsWithToolTag(AsShared(), "VolumeTool");
-#else
 	return {EnableCreateVolumeTool};
-#endif
 }
 
 /**
@@ -153,13 +138,11 @@ void ToolCommandClassName::GetToolDefaultObjectList(TArray<UInteractiveTool*>& T
 {\
 ToolCDOs.Add(GetMutableDefault<ToolClassName>());\
 }
-PCG_DEFINE_TOOL_COMMANDS(FPCGDrawSplineCommands, "DrawSpline", "Draw Spline", UPCGDrawSplineToolBase);
-PCG_DEFINE_TOOL_COMMANDS(FPCGDrawSurfaceCommands, "DrawSurface", "Draw Surface", UPCGDrawSplineToolBase);
-PCG_DEFINE_TOOL_COMMANDS(FPCGPaintCommands, "Paint", "Paint", UPCGPaintTool);
+PCG_DEFINE_TOOL_COMMANDS(FPCGDrawSplineCommands, "DrawSpline", "Draw Spline", UPCGDrawSplineTool);
+PCG_DEFINE_TOOL_COMMANDS(FPCGDrawSurfaceCommands, "DrawSurface", "Draw Surface", UPCGDrawSplineTool);
+PCG_DEFINE_TOOL_COMMANDS(FPCGPaintCommands, "DrawSurface", "Draw Surface", UPCGPaintTool);
 PCG_DEFINE_TOOL_COMMANDS(FPCGCreateVolumeCommands, "CreateVolume", "Create Volume", UPCGCreateVolumeTool);
 
 #undef PCG_DEFINE_TOOL_COMMANDS
 
-#undef PCG_TOOL_COMMAND_INFO
-
 #undef LOCTEXT_NAMESPACE
diff --git a/Engine/Plugins/PCG/Source/PCGEditor/Private/EditorMode/PCGEdModeToolkit.cpp b/Engine/Plugins/PCG/Source/PCGEditor/Private/EditorMode/PCGEdModeToolkit.cpp
--- a/Engine/Plugins/PCG/Source/PCGEditor/Private/EditorMode/PCGEdModeToolkit.cpp
+++ b/Engine/Plugins/PCG/Source/PCGEditor/Private/EditorMode/PCGEdModeToolkit.cpp
@@ -4,6 +4,7 @@
 
 #include "EditorMode/PCGEdModeCommands.h"
 #include "EditorMode/PCGEdModeSettings.h"
+#include "EditorMode/Tools/PCGInteractiveTool.h"
 
 #include "EditorModeManager.h"
 #include "LevelEditorViewport.h"
@@ -71,16 +72,6 @@ void FPCGEditorModeToolkit::OnToolPaletteChanged(const FName PaletteName)
 	ensure(false);
 }
 
-void FPCGEditorModeToolkit::BuildToolPalette(FName Palette, class FToolBarBuilder& ToolbarBuilder)
-{
-	if (!OwningEditorMode.IsValid())
-	{
-		return;
-	}
-
-	FModeToolkit::BuildToolPalette(Palette, ToolbarBuilder);	
-}
-
 FText FPCGEditorModeToolkit::GetActiveToolDisplayName() const
 {
 	return CachedActiveToolInfo.ToolDisplayName;
@@ -118,7 +109,10 @@ void FPCGEditorModeToolkit::OnToolEnded(UInteractiveToolManager* Manager, UInter
 	ClearToolNotification();
 	ClearToolWarning();
 
-	UnregisterToolCallbacks(Tool);
+	if (UPCGInteractiveTool* CurrentTool = CastChecked<UPCGInteractiveTool>(Tool, ECastCheckedType::NullAllowed))
+	{
+		UnregisterToolCallbacks(CurrentTool);
+	}
 }
 
 void FPCGEditorModeToolkit::PostToolNotification(const FText& Message)
@@ -210,7 +204,7 @@ EVisibility FPCGEditorModeToolkit::GetAcceptButtonVisibility() const
 
 EVisibility FPCGEditorModeToolkit::GetCancelButtonVisibility() const
 {
-	return GetScriptableEditorMode()->GetInteractiveToolsContext()->CanCancelActiveTool() ? EVisibility::Visible : EVisibility::Collapsed;
+	return GetScriptableEditorMode()->GetInteractiveToolsContext()->ActiveToolHasAccept() ? EVisibility::Visible : EVisibility::Collapsed;
 }
 
 FReply FPCGEditorModeToolkit::OnAcceptButtonClicked() const
@@ -250,21 +244,23 @@ void FPCGEditorModeToolkit::UnregisterToolkitCallbacks()
 	PCGEditorModeSettings->OnSettingChanged().RemoveAll(this);
 }
 
-void FPCGEditorModeToolkit::RegisterToolCallbacks(UInteractiveTool* CurrentTool)
+void FPCGEditorModeToolkit::RegisterToolCallbacks(UPCGInteractiveTool* CurrentTool)
 {
 	if (CurrentTool)
 	{
 		CurrentTool->OnPropertySetsModified.AddSP(this, &FPCGEditorModeToolkit::UpdateActiveTool);
 		CurrentTool->OnPropertyModifiedDirectlyByTool.AddSP(this, &FPCGEditorModeToolkit::InvalidateCachedDetailPanelState);
+		CurrentTool->OnToolWarningChanged.AddSP(this, &FPCGEditorModeToolkit::PostToolWarning);
 	}
 }
 
-void FPCGEditorModeToolkit::UnregisterToolCallbacks(UInteractiveTool* CurrentTool) const
+void FPCGEditorModeToolkit::UnregisterToolCallbacks(UPCGInteractiveTool* CurrentTool) const
 {
 	if (CurrentTool)
 	{
 		CurrentTool->OnPropertySetsModified.RemoveAll(this);
 		CurrentTool->OnPropertyModifiedDirectlyByTool.RemoveAll(this);
+		CurrentTool->OnToolWarningChanged.RemoveAll(this);
 	}
 }
 
@@ -280,43 +276,45 @@ void FPCGEditorModeToolkit::OnEditorModeSettingsChanged(UObject* InObject, FProp
 void FPCGEditorModeToolkit::UpdateActiveTool()
 {
 	UInteractiveTool* CurrentTool = GetScriptableEditorMode()->GetToolManager(EToolsContextScope::EdMode)->GetActiveTool(EToolSide::Left);
-	
-	CachedActiveToolInfo = CurrentTool->GetToolInfo();
+	if (UPCGInteractiveTool* CurrentPCGTool = CastChecked<UPCGInteractiveTool>(CurrentTool, ECastCheckedType::NullAllowed))
+	{
+		CachedActiveToolInfo = CurrentPCGTool->GetToolInfo();
 
-	ToolkitBuilder->SetActiveToolDisplayName(CachedActiveToolInfo.ToolDisplayName);
+		ToolkitBuilder->SetActiveToolDisplayName(CachedActiveToolInfo.ToolDisplayName);
 
-	if (const UPCGEditorModeSettings* Settings = GetDefault<UPCGEditorModeSettings>())
-	{
-		if (Settings->bHideToolButtonsDuringActiveTool)
+		if (const UPCGEditorModeSettings* Settings = GetDefault<UPCGEditorModeSettings>())
 		{
-			ToolkitBuilder->SetActivePaletteCommandsVisibility(EVisibility::Collapsed);
+			if (Settings->bHideToolButtonsDuringActiveTool)
+			{
+				ToolkitBuilder->SetActivePaletteCommandsVisibility(EVisibility::Collapsed);
+			}
 		}
-	}
 
-	// Before actually changing the detail panel, we need to see where the current keyboard focus is, because
-	// if it's inside the detail panel, we'll need to reset it to the detail panel as a whole, else we might
-	// lose it entirely when that detail panel element gets destroyed (which would make us unable to receive any
-	// hotkey presses until the user clicks somewhere).
-	const TSharedPtr<SWidget> FocusedWidget = FSlateApplication::Get().GetKeyboardFocusedWidget();
-	if (FocusedWidget != ModeDetailsView)
-	{
-		// Search upward from the currently focused widget
-		TSharedPtr<SWidget> CurrentWidget = FocusedWidget;
-		while (CurrentWidget.IsValid())
+		// Before actually changing the detail panel, we need to see where the current keyboard focus is, because
+		// if it's inside the detail panel, we'll need to reset it to the detail panel as a whole, else we might
+		// lose it entirely when that detail panel element gets destroyed (which would make us unable to receive any
+		// hotkey presses until the user clicks somewhere).
+		const TSharedPtr<SWidget> FocusedWidget = FSlateApplication::Get().GetKeyboardFocusedWidget();
+		if (FocusedWidget != ModeDetailsView)
 		{
-			if (CurrentWidget == ModeDetailsView)
+			// Search upward from the currently focused widget
+			TSharedPtr<SWidget> CurrentWidget = FocusedWidget;
+			while (CurrentWidget.IsValid())
 			{
-				// Reset focus to the detail panel as a whole to avoid losing it when the inner elements change.
-				FSlateApplication::Get().SetKeyboardFocus(ModeDetailsView);
-				break;
+				if (CurrentWidget == ModeDetailsView)
+				{
+					// Reset focus to the detail panel as a whole to avoid losing it when the inner elements change.
+					FSlateApplication::Get().SetKeyboardFocus(ModeDetailsView);
+					break;
+				}
+
+				CurrentWidget = CurrentWidget->GetParentWidget();
 			}
-
-			CurrentWidget = CurrentWidget->GetParentWidget();
 		}
+
+		ModeDetailsView->SetObjects({CurrentPCGTool->GetToolProperties(true)});
+		RegisterToolCallbacks(CurrentPCGTool);
 	}
-	
-	ModeDetailsView->SetObjects(CurrentTool->GetToolProperties(true));
-	RegisterToolCallbacks(CurrentTool);
 
 	// @todo_pcg: If a tool property needs to invalidate the cached details panel, do it here. 
 	// InvalidateCachedDetailPanelState(nullptr);
@@ -330,7 +328,7 @@ void FPCGEditorModeToolkit::InvalidateCachedDetailPanelState(UObject* ChangedObj
 void FPCGEditorModeToolkit::RegisterPalettes()
 {
 	const FPCGEditorModePaletteCommands& PaletteCommands = FPCGEditorModePaletteCommands::Get();
-	FPCGEditorModeToolCommands& ToolCommands = FPCGEditorModeToolCommands::Get();
+	const FPCGEditorModeToolCommands& ToolCommands = FPCGEditorModeToolCommands::Get();
 
 	ToolkitSections = MakeShared<FToolkitSections>();
 	ToolkitSections->DetailsView = ModeDetailsView;
diff --git a/Engine/Plugins/PCG/Source/PCGEditor/Private/EditorMode/Tools/Helpers/PCGEdModeActorHelpers.cpp b/Engine/Plugins/PCG/Source/PCGEditor/Private/EditorMode/Tools/Helpers/PCGEdModeActorHelpers.cpp
--- a/Engine/Plugins/PCG/Source/PCGEditor/Private/EditorMode/Tools/Helpers/PCGEdModeActorHelpers.cpp
+++ b/Engine/Plugins/PCG/Source/PCGEditor/Private/EditorMode/Tools/Helpers/PCGEdModeActorHelpers.cpp
@@ -8,7 +8,7 @@
 
 #define LOCTEXT_NAMESPACE "PCGEditorMode"
 
-namespace UE::PCG::EditorMode
+namespace PCG::EditorMode
 {
 	namespace Actor
 	{
@@ -31,14 +31,14 @@ namespace UE::PCG::EditorMode
 			return UPCGActorHelpers::SpawnDefaultActor(SpawnDefaultActorParams);
 		}
 
-		AActor* SpawnWorking(UWorld* World, const FTransform& Transform, TSubclassOf<AActor> ActorClass)
+		AActor* SpawnWorking(UWorld* World, const FTransform& Transform)
 		{
 			FActorSpawnParameters SpawnParams;
 			SpawnParams.Name = DefaultWorkingActorName;
 			SpawnParams.NameMode = FActorSpawnParameters::ESpawnActorNameMode::Requested;
 			SpawnParams.ObjectFlags = RF_Transient;
 
-			UPCGActorHelpers::FSpawnDefaultActorParams SpawnDefaultActorParams(World, ActorClass, Transform, SpawnParams);
+			UPCGActorHelpers::FSpawnDefaultActorParams SpawnDefaultActorParams(World, AActor::StaticClass(), Transform, SpawnParams);
 
 			return UPCGActorHelpers::SpawnDefaultActor(SpawnDefaultActorParams);
 		}
@@ -53,8 +53,8 @@ namespace UE::PCG::EditorMode
 				return nullptr;
 			}
 
-			UPCGComponent* PCGComponent = Owner->GetComponentByClass<UPCGComponent>();
-			return PCGComponent;
+			TArray<UActorComponent*> ToolComponents = Owner->GetComponentsByTag(UPCGComponent::StaticClass(), Tags::ToolGeneratedTag);
+			return ToolComponents.IsEmpty() ? nullptr : Cast<UPCGComponent>(ToolComponents[0]);
 		}
 
 		UPCGComponent* Create(AActor* Owner, const FName ComponentName)
@@ -75,7 +75,7 @@ namespace UE::PCG::EditorMode
 		{
 			return Create(Owner, Actor::DefaultWorkingActorName);
 		}
-	}
-}
+	} // namespace PCG::EditorMode::PCGComponent
+} // namespace PCG::EditorMode
 
 #undef LOCTEXT_NAMESPACE
diff --git a/Engine/Plugins/PCG/Source/PCGEditor/Private/EditorMode/Tools/Helpers/PCGEdModeEditorUtilities.cpp b/Engine/Plugins/PCG/Source/PCGEditor/Private/EditorMode/Tools/Helpers/PCGEdModeEditorUtilities.cpp
deleted file mode 100644
diff --git a/Engine/Plugins/PCG/Source/PCGEditor/Private/EditorMode/Tools/Helpers/PCGEdModeSplineHelpers.cpp b/Engine/Plugins/PCG/Source/PCGEditor/Private/EditorMode/Tools/Helpers/PCGEdModeSplineHelpers.cpp
--- a/Engine/Plugins/PCG/Source/PCGEditor/Private/EditorMode/Tools/Helpers/PCGEdModeSplineHelpers.cpp
+++ b/Engine/Plugins/PCG/Source/PCGEditor/Private/EditorMode/Tools/Helpers/PCGEdModeSplineHelpers.cpp
@@ -3,12 +3,11 @@
 #include "EditorMode/Tools/Helpers/PCGEdModeSplineHelpers.h"
 
 #include "Components/SplineComponent.h"
-#include "EditorMode/Tools/Helpers/PCGEdModeActorHelpers.h"
 #include "UObject/UObjectGlobals.h"
 
 #define LOCTEXT_NAMESPACE "PCGEditorMode"
 
-namespace UE::PCG::EditorMode::Spline
+namespace PCG::EditorMode::Spline
 {
 	static const FLazyName DefaultWorkingSplineName = "PCG Tool Working Spline";
 
@@ -42,8 +41,6 @@ namespace UE::PCG::EditorMode::Spline
 			Spline = NewObject<USplineComponent>(Actor, USplineComponent::StaticClass(), Name, Flags);
 		}
 
-		Spline->SetFlags(RF_Transient);
-		Spline->ComponentTags.Add(Tags::ToolGeneratedTag);
 		Actor->AddInstanceComponent(Spline);
 		Spline->RegisterComponent();
 		Spline->ResetRelativeTransform();
@@ -57,7 +54,7 @@ namespace UE::PCG::EditorMode::Spline
 		return CreateComponent(Actor, DefaultWorkingSplineName, /*TemplateComponent=*/nullptr, /*bTransactional=*/false);
 	}
 
-	USplineComponent* Find(const AActor* Actor, const int32 TargetIndex)
+	USplineComponent* FindComponent(const AActor* Actor, const int32 TargetIndex)
 	{
 		if (!ensure(Actor))
 		{
diff --git a/Engine/Plugins/PCG/Source/PCGEditor/Private/EditorMode/Tools/Line/PCGDrawSplineTool.cpp b/Engine/Plugins/PCG/Source/PCGEditor/Private/EditorMode/Tools/Line/PCGDrawSplineTool.cpp
--- a/Engine/Plugins/PCG/Source/PCGEditor/Private/EditorMode/Tools/Line/PCGDrawSplineTool.cpp
+++ b/Engine/Plugins/PCG/Source/PCGEditor/Private/EditorMode/Tools/Line/PCGDrawSplineTool.cpp
@@ -2,11 +2,11 @@
 
 #include "EditorMode/Tools/Line/PCGDrawSplineTool.h"
 
-#include "InteractiveToolManager.h"
 #include "PCGComponent.h"
 #include "EditorMode/PCGEdModeSettings.h"
 #include "EditorMode/PCGEdModeStyle.h"
 #include "EditorMode/Tools/Helpers/PCGEdModeSceneVisualizeHelpers.h"
+#include "EditorMode/Tools/Helpers/PCGEdModeSplineHelpers.h"
 
 #include "BaseBehaviors/SingleClickOrDragBehavior.h"
 #include "Components/SplineComponent.h"
@@ -32,7 +32,7 @@ namespace PCGDrawSplineTool
 			// These pass the working spline to the overloads below
 			virtual void Apply(UObject* Object) override
 			{
-				if (!ensure(Object && Object->IsA<UPCGDrawSplineToolBase>() && WorkingSpline))
+				if (!ensure(Object && Object->IsA<UPCGDrawSplineTool>() && WorkingSpline))
 				{
 					return;
 				}
@@ -42,7 +42,7 @@ namespace PCGDrawSplineTool
 
 			virtual void Revert(UObject* Object) override
 			{
-				if (!ensure(Object && Object->IsA<UPCGDrawSplineToolBase>() && WorkingSpline))
+				if (!ensure(Object && Object->IsA<UPCGDrawSplineTool>() && WorkingSpline))
 				{
 					return;
 				}
@@ -193,97 +193,40 @@ namespace PCGDrawSplineTool
 	{
 		const FText ToolName = LOCTEXT("ToolName", "DrawSpline");
 		const FText ToolDescription = LOCTEXT("DrawSplineDisplayMessage", "Draw a spline and distribute spline data directly to a PCG Graph Instance.");
-		const FName SplineToolTag = "SplineTool";
-		const FName SplineSurfaceToolTag = "SplineSurfaceTool";		
 		static const FText AddPointTransactionName = LOCTEXT("AddPointTransactionName", "Add Point");
 	}
 } // namespace PCGDrawSplineTool
 
-UPCGInteractiveToolSettings_SplineBase::UPCGInteractiveToolSettings_SplineBase()
+UPCGDrawSplineToolSettings::UPCGDrawSplineToolSettings()
 	: SplineName(GetDefault<UPCGEditorModeSettings>()->DefaultNewSplineComponentName)
 {
-	
-}
-
-USplineComponent* UPCGInteractiveToolSettings_SplineBase::GetWorkingSplineComponent(FName DataInstanceIdentifier) const
-{
-	if(const FPCGInteractiveToolWorkingData_Spline* ToolData = GetTypedWorkingData<FPCGInteractiveToolWorkingData_Spline>(DataInstanceIdentifier))
-	{
-		return ToolData->GetSplineComponent();
-	}
-	
-	return nullptr;
-}
-
-TArray<USplineComponent*> UPCGInteractiveToolSettings_SplineBase::GetWorkingSplineComponents() const
-{
-	TArray<USplineComponent*> Result;
-
-	TArray<FName> AllWorkingDataInstanceNames = GetValidDataInstanceNamesForGraph();
-	for(const FName& DataInstanceIdentifier : AllWorkingDataInstanceNames)
-	{
-		if(USplineComponent* WorkingSplineComponent = GetWorkingSplineComponent(DataInstanceIdentifier))
-		{
-			Result.Add(WorkingSplineComponent);
-		}
-	}
-
-	return Result;
-}
-
-const UScriptStruct* UPCGInteractiveToolSettings_Spline::GetWorkingDataType() const
-{
-	return FPCGInteractiveToolWorkingData_Spline::StaticStruct();
-}
-
-FName UPCGInteractiveToolSettings_Spline::GetToolTag() const
-{
-	return PCGDrawSplineTool::Constants::SplineToolTag;
-}
-
-void UPCGInteractiveToolSettings_Spline::OnWorkingDataInitialized(FPCGInteractiveToolWorkingData* WorkingData) const
-{
-	if(USplineComponent* SplineComponent = GetWorkingSplineComponent(DataInstance))
+	if (!HasAnyFlags(RF_ClassDefaultObject))
 	{
-		SplineComponent->SetClosedLoop(bClosedSpline);
+		UPCGDrawSplineToolSettings::TryLoadDefaultGraph();
 	}
 }
 
-const UScriptStruct* UPCGInteractiveToolSettings_SplineSurface::GetWorkingDataType() const
+FText UPCGDrawSplineToolSettings::GetToolTag() const
 {
-	return FPCGInteractiveToolWorkingData_SplineSurface::StaticStruct();
+	return PCGDrawSplineTool::Constants::ToolName;
 }
 
-FName UPCGInteractiveToolSettings_SplineSurface::GetToolTag() const
+void UPCGDrawSplineTool::Setup()
 {
-	return PCGDrawSplineTool::Constants::SplineSurfaceToolTag;
-}
+	using namespace PCGDrawSplineTool;
 
-void UPCGInteractiveToolSettings_Spline::OnPropertyModified(UObject* Object, FProperty* Property)
-{
-	Super::OnPropertyModified(Object, Property);
-	
-	if(Object == this && Property->GetFName() == GET_MEMBER_NAME_CHECKED(UPCGInteractiveToolSettings_Spline, bClosedSpline))
-	{
-		if(USplineComponent* SplineComponent = GetWorkingSplineComponent(DataInstance))
-		{
-			SplineComponent->SetClosedLoop(bClosedSpline, true);
+	Super::Setup();
 
-			FPropertyChangedEvent Event(nullptr);
-			FCoreUObjectDelegates::OnObjectPropertyChanged.Broadcast(SplineComponent, Event);
-		}
-	}
-}
+	UPCGDrawSplineToolSettings* Settings = CastChecked<UPCGDrawSplineToolSettings>(ToolSettings);
+	CachedMinPointSpacingSquared = Settings->MinPointSpacing * Settings->MinPointSpacing;
 
-void UPCGDrawSplineToolBase::Setup()
-{
-	using namespace PCGDrawSplineTool;
+	WorkingSpline = Spline::CreateWorkingComponent(WorkingActor.Get());
+	check(WorkingSpline);
 
-	Super::Setup();
+	WorkingSpline->ClearSplinePoints();
+	WorkingSpline->SetClosedLoop(bIsSplineSurface);
+	WorkingSpline->bSplineHasBeenEdited = true;
 
-	UPCGInteractiveToolSettings_SplineBase* ToolSettings = GetSplineSettings();
-	CachedMinPointSpacingSquared = ToolSettings->MinPointSpacing * ToolSettings->MinPointSpacing;
-	
 	ClickOrDragBehavior = NewObject<USingleClickOrDragInputBehavior>();
 	ClickOrDragBehavior->Initialize(this, this);
 	AddInputBehavior(ClickOrDragBehavior);
@@ -293,23 +236,52 @@ void UPCGDrawSplineToolBase::Setup()
 	{
 		.World = GetWorld()
 	};
+
+	Settings->WatchProperty(Settings->TargetActor, [this](TWeakObjectPtr<AActor> TargetActor)
+	{
+		if (!TargetActor.IsValid())
+		{
+			SetWarning(LOCTEXT("TargetActorSetToInvalid", "Target actor is invalid!"));
+		}
+	});
+
+	Settings->SilentUpdateWatched();
 }
 
-void UPCGDrawSplineToolBase::Shutdown(const EToolShutdownType ShutdownType)
+void UPCGDrawSplineTool::Shutdown(const EToolShutdownType ShutdownType)
 {
-	// We have to call this before the super function as it would remove the property sets before we can route the shutdown function to the pcg settings
-	UE::PCG::EditorMode::Tool::Shutdown(this, ShutdownType);
-	
+	using namespace PCGDrawSplineTool;
+
+	if (ShutdownType == EToolShutdownType::Accept && WorkingSpline && WorkingSpline->GetNumberOfSplinePoints() > 0)
+	{
+		ApplyTool();
+	}
+
+	if (WorkingComponent.IsValid())
+	{
+		WorkingComponent->CleanupLocalImmediate(/*bRemoveAllComponents=*/true, /*bCleanUpLocalComponents=*/true);
+		WorkingComponent->bActivated = false;
+		WorkingComponent->DestroyComponent();
+	}
+
+	if (WorkingSpline)
+	{
+		WorkingSpline->DestroyComponent();
+		WorkingSpline = nullptr;
+	}
+
+	if (WorkingActor.IsValid())
+	{
+		WorkingActor->MarkAsGarbage();
+		WorkingActor = nullptr;
+	}
+
 	Super::Shutdown(ShutdownType);
 }
 
-void UPCGDrawSplineToolBase::AddSplinePoint(const FVector& HitLocation, const FVector& HitNormal) const
+void UPCGDrawSplineTool::AddSplinePoint(const FVector& HitLocation, const FVector& HitNormal) const
 {
 	using namespace PCGDrawSplineTool;
-
-	UPCGInteractiveToolSettings_SplineBase* ToolSettings = GetSplineSettings();
-	
-	USplineComponent* WorkingSpline = ToolSettings->GetWorkingSplineComponent(ToolSettings->DataInstance);
 	if (!WorkingSpline)
 	{
 		return;
@@ -319,14 +291,15 @@ void UPCGDrawSplineToolBase::AddSplinePoint(const FVector& HitLocation, const FV
 	const FVector NewNormal = GetNewPointNormal(HitNormal);
 
 	FVector FinalLocation = HitLocation;
-	switch (ToolSettings->Target.OffsetMode)
+	const UPCGDrawSplineToolSettings* Settings = CastChecked<UPCGDrawSplineToolSettings>(ToolSettings);
+	switch (Settings->Target.OffsetMode)
 	{
 		case EPCGToolDrawTargetOffset::DistanceOffset:
 			check(HitNormal.IsNormalized());
-			FinalLocation += ToolSettings->Target.OffsetDistance * NewNormal;
+			FinalLocation += Settings->Target.OffsetDistance * NewNormal;
 			break;
 		case EPCGToolDrawTargetOffset::ExplicitOffset:
-			FinalLocation += ToolSettings->Target.OffsetVector;
+			FinalLocation += Settings->Target.OffsetVector;
 			break;
 		default: // fall-through
 		case EPCGToolDrawTargetOffset::HitLocation:
@@ -335,25 +308,13 @@ void UPCGDrawSplineToolBase::AddSplinePoint(const FVector& HitLocation, const FV
 
 	WorkingSpline->AddSplinePoint(FinalLocation, ESplineCoordinateSpace::World, /*bUpdate=*/false);
 	WorkingSpline->SetUpVectorAtSplinePoint(NumSplinePoints, NewNormal, ESplineCoordinateSpace::World, /*bUpdate=*/true);
-
-	FPropertyChangedEvent EmptyPropertyChangedEvent(nullptr);
-	if(UE::PCG::EditorMode::Tool::CVarCommitPropertyChangedEventOnTimer.GetValueOnAnyThread())
-	{
-		ToolSettings->PropertyChangedEventQueue.Add(WorkingSpline, EmptyPropertyChangedEvent);
-	}
-	else
-	{
-		FCoreUObjectDelegates::OnObjectPropertyChanged.Broadcast(WorkingSpline, EmptyPropertyChangedEvent);
-	}
 }
 
-FVector UPCGDrawSplineToolBase::GetNewPointNormal(const FVector& HitNormal) const
+FVector UPCGDrawSplineTool::GetNewPointNormal(const FVector& HitNormal) const
 {
-	UPCGInteractiveToolSettings_SplineBase* ToolSettings = GetSplineSettings();
-	
-	USplineComponent* WorkingSpline = ToolSettings->GetWorkingSplineComponent(ToolSettings->DataInstance);
+	const UPCGDrawSplineToolSettings* Settings = CastChecked<UPCGDrawSplineToolSettings>(ToolSettings);
 	const int32 NumSplinePoints = WorkingSpline->GetNumberOfSplinePoints();
-	switch (ToolSettings->Target.NormalMode)
+	switch (Settings->Target.NormalMode)
 	{
 		case EPCGToolDrawTargetNormal::AlignToPrevious:
 		{
@@ -380,14 +341,105 @@ FVector UPCGDrawSplineToolBase::GetNewPointNormal(const FVector& HitNormal) cons
 		case EPCGToolDrawTargetNormal::WorldUp:
 			return FVector::UnitZ();
 		case EPCGToolDrawTargetNormal::Explicit:
-			return ToolSettings->Target.Normal;
+			return Settings->Target.Normal;
 		default: // fall-through
 		case EPCGToolDrawTargetNormal::HitNormal:
 			return HitNormal;
 	}
 }
 
-FInputRayHit UPCGDrawSplineToolBase::IsHitByClick(const FInputDeviceRay& ClickPosition)
+void UPCGDrawSplineTool::ApplyToolInternal() const
+{
+	GetToolManager()->BeginUndoTransaction(LOCTEXT("DrawSplineTransactionName", "[PCG Tool] Draw Spline"));
+	const UPCGDrawSplineToolSettings* Settings = CastChecked<UPCGDrawSplineToolSettings>(ToolSettings);
+
+	TOptional<FText> OutWarning;
+
+	UWorld* World = GetContextualWorldChecked();
+	USplineComponent* ToolSplineComponent = WorkingSpline.Get();
+
+	{
+		using namespace PCG::EditorMode;
+
+		check(World);
+		check(ToolSplineComponent);
+		
+		if (!Settings->TargetGraph)
+		{
+			OutWarning = LOCTEXT("InvalidGraphInstance", "[Spline Tool] Failed: Invalid Graph Instance");
+		}
+
+		AActor* Actor = nullptr;
+		UPCGComponent* PCGComponent = nullptr;
+
+		switch (Settings->Interaction)
+		{
+		case EPCGToolInteractionMode::CreateNewActor:
+			{
+				Actor = Actor::Spawn(World, Settings->NewActorName, FTransform::Identity);
+			}
+			break;
+		case EPCGToolInteractionMode::AddToSelectedActor:
+			{
+				if (!Settings->TargetActor.IsValid())
+				{
+					OutWarning = LOCTEXT("InvalidActorWarning", "Invalid Actor!");
+				}
+
+				Actor = Settings->TargetActor.Get();
+				PCGComponent = Actor ? PCGComponent::Find(Actor) : nullptr;
+			}
+			break;
+		default:
+			checkNoEntry();
+		}
+
+		// @todo_pcg: Will overwrite the current spline. Could be an option.
+		if (Actor)
+		{
+			USplineComponent* Spline = Spline::FindComponent(Actor);
+			if (!Spline)
+			{
+				Spline = Spline::CreateComponent(Actor, Settings->SplineName, ToolSplineComponent, /*bTransactional=*/true);
+				check(Spline);
+			}
+			else
+			{
+				Spline::CopySplinePoints(*ToolSplineComponent, *Spline, /*bTransactional=*/true);
+
+				// The copy points doesn't trigger a object changed event; we need to trigger one here manually so it gets picked up
+				FPropertyChangedEvent EmptyEvent{ nullptr };
+				FCoreUObjectDelegates::OnObjectPropertyChanged.Broadcast(Spline, EmptyEvent);
+			}
+
+			if(Settings->TargetGraph)
+			{
+				if (!PCGComponent)
+				{
+					PCGComponent = PCGComponent::Create(Actor, Settings->NewPCGComponentName);
+					check(PCGComponent);
+				}
+
+				PCGComponent->SetGraph(Settings->TargetGraph);
+				PCGComponent->GetGraphInstance()->CopyParameterOverrides(WorkingComponent->GetGraphInstance());
+
+				if(Settings->bAutoGenerate)
+				{
+					PCGComponent->Generate();
+				}
+			}
+		}
+
+		if(OutWarning.IsSet())
+		{
+			NotifyError(OutWarning.GetValue());
+		}
+	}
+
+	GetToolManager()->EndUndoTransaction();
+}
+
+FInputRayHit UPCGDrawSplineTool::IsHitByClick(const FInputDeviceRay& ClickPosition)
 {
 	using namespace PCG::EditorMode;
 	if (TOptional<Scene::FViewHitResult> Result = ViewportRay(ClickPosition, CachedRayParams))
@@ -398,16 +450,14 @@ FInputRayHit UPCGDrawSplineToolBase::IsHitByClick(const FInputDeviceRay& ClickPo
 	return {};
 }
 
-void UPCGDrawSplineToolBase::OnClicked(const FInputDeviceRay& ClickPosition)
+void UPCGDrawSplineTool::OnClicked(const FInputDeviceRay& ClickPosition)
 {
 	using namespace PCGDrawSplineTool;
 
 	if (TOptional<Scene::FViewHitResult> Result = ViewportRay(ClickPosition, CachedRayParams))
 	{
-	UPCGInteractiveToolSettings_SplineBase* ToolSettings = GetSplineSettings();
-		
-		USplineComponent* WorkingSpline = ToolSettings->GetWorkingSplineComponent(ToolSettings->DataInstance);
-		switch (ToolSettings->DrawMode)
+		const UPCGDrawSplineToolSettings* Settings = CastChecked<UPCGDrawSplineToolSettings>(ToolSettings);
+		switch (Settings->DrawMode)
 		{
 			case EPCGDrawSplineDrawMode::ClickAutoTangent:
 			case EPCGDrawSplineDrawMode::FreeDraw:
@@ -417,7 +467,7 @@ void UPCGDrawSplineToolBase::OnClicked(const FInputDeviceRay& ClickPosition)
 				const int32 PointIndex = WorkingSpline->GetNumberOfSplinePoints() - 1;
 				GetToolManager()->EmitObjectChange(this,
 												   MakeUnique<CommandChange::FSimplePointInsertionChange>(
-													   WorkingSpline,
+													   WorkingSpline.Get(),
 													   Result->ImpactPosition,
 													   WorkingSpline->GetUpVectorAtSplinePoint(PointIndex, ESplineCoordinateSpace::World)),
 												   Constants::AddPointTransactionName);
@@ -435,7 +485,7 @@ void UPCGDrawSplineToolBase::OnClicked(const FInputDeviceRay& ClickPosition)
 				GetToolManager()->EmitObjectChange(
 					this,
 					MakeUnique<CommandChange::FTangentPointInsertionChange>(
-						WorkingSpline,
+						WorkingSpline.Get(),
 						Result->ImpactPosition,
 						WorkingSpline->GetUpVectorAtSplinePoint(PointIndex, ESplineCoordinateSpace::World),
 						FVector::Zero()),
@@ -446,7 +496,7 @@ void UPCGDrawSplineToolBase::OnClicked(const FInputDeviceRay& ClickPosition)
 	}
 }
 
-FInputRayHit UPCGDrawSplineToolBase::CanBeginClickDragSequence(const FInputDeviceRay& DragPosition)
+FInputRayHit UPCGDrawSplineTool::CanBeginClickDragSequence(const FInputDeviceRay& DragPosition)
 {
 	using namespace PCG::EditorMode::Scene;
 	if (TOptional<FViewHitResult> Result = ViewportRay(DragPosition, CachedRayParams))
@@ -457,30 +507,26 @@ FInputRayHit UPCGDrawSplineToolBase::CanBeginClickDragSequence(const FInputDevic
 	return {};
 }
 
-void UPCGDrawSplineToolBase::OnClickPress(const FInputDeviceRay& ClickPosition)
+void UPCGDrawSplineTool::OnClickPress(const FInputDeviceRay& ClickPosition)
 {
 	using namespace PCG::EditorMode::Scene;
 	if (TOptional<FViewHitResult> Result = ViewportRay(ClickPosition, CachedRayParams))
 	{
-		UPCGInteractiveToolSettings_SplineBase* ToolSettings = GetSplineSettings();
-		
 		AddSplinePoint(Result->ImpactPosition, Result->ImpactNormal);
-		
-		if (ToolSettings->DrawMode == EPCGDrawSplineDrawMode::FreeDraw)
+
+		const UPCGDrawSplineToolSettings* Settings = CastChecked<UPCGDrawSplineToolSettings>(ToolSettings);
+		if (Settings->DrawMode == EPCGDrawSplineDrawMode::FreeDraw)
 		{
 			// Cache the index for where this stroke began
-			FreeDrawStrokeStartIndex = ToolSettings->GetWorkingSplineComponent(ToolSettings->DataInstance)->GetNumberOfSplinePoints() - 1;
+			FreeDrawStrokeStartIndex = WorkingSpline->GetNumberOfSplinePoints() - 1;
 		}
 	}
 }
 
-void UPCGDrawSplineToolBase::OnClickDrag(const FInputDeviceRay& DragPosition)
+void UPCGDrawSplineTool::OnClickDrag(const FInputDeviceRay& DragPosition)
 {
 	using namespace PCG::EditorMode::Scene;
 
-	UPCGInteractiveToolSettings_SplineBase* ToolSettings = GetSplineSettings();
-	
-	USplineComponent* WorkingSpline = ToolSettings->GetWorkingSplineComponent(ToolSettings->DataInstance);
 	const int32 NumSplinePoints = WorkingSpline->GetNumberOfSplinePoints();
 	if (!ensure(NumSplinePoints > 0))
 	{
@@ -490,7 +536,8 @@ void UPCGDrawSplineToolBase::OnClickDrag(const FInputDeviceRay& DragPosition)
 	if (TOptional<FViewHitResult> Result = ViewportRay(DragPosition, CachedRayParams))
 	{
 		const int32 FinalSplinePoint = NumSplinePoints - 1;
-		switch (ToolSettings->DrawMode)
+		const UPCGDrawSplineToolSettings* Settings = CastChecked<UPCGDrawSplineToolSettings>(ToolSettings);
+		switch (Settings->DrawMode)
 		{
 			case EPCGDrawSplineDrawMode::ClickAutoTangent:
 			{
@@ -527,14 +574,14 @@ void UPCGDrawSplineToolBase::OnClickDrag(const FInputDeviceRay& DragPosition)
 					FVector FinalLocation = Result->ImpactPosition;
 					const FVector Normal = GetNewPointNormal(Result->ImpactNormal);
 
-					switch (ToolSettings->Target.OffsetMode)
+					switch (Settings->Target.OffsetMode)
 					{
 						case EPCGToolDrawTargetOffset::DistanceOffset:
 							check(Result->ImpactNormal.IsNormalized());
-							FinalLocation += ToolSettings->Target.OffsetDistance * Normal;
+							FinalLocation += Settings->Target.OffsetDistance * Normal;
 							break;
 						case EPCGToolDrawTargetOffset::ExplicitOffset:
-							FinalLocation += ToolSettings->Target.OffsetVector;
+							FinalLocation += Settings->Target.OffsetVector;
 							break;
 						default: // fall-through
 						case EPCGToolDrawTargetOffset::HitLocation:
@@ -555,20 +602,16 @@ void UPCGDrawSplineToolBase::OnClickDrag(const FInputDeviceRay& DragPosition)
 	}
 }
 
-void UPCGDrawSplineToolBase::OnClickRelease(const FInputDeviceRay& ReleasePosition)
+void UPCGDrawSplineTool::OnClickRelease(const FInputDeviceRay& ReleasePosition)
 {
 	// Activate the final position
 	OnClickDrag(ReleasePosition);
 	OnTerminateDragSequence();
 }
 
-void UPCGDrawSplineToolBase::OnTerminateDragSequence()
+void UPCGDrawSplineTool::OnTerminateDragSequence()
 {
 	using namespace PCGDrawSplineTool;
-
-	UPCGInteractiveToolSettings_SplineBase* ToolSettings = GetSplineSettings();
-	USplineComponent* WorkingSpline = ToolSettings->GetWorkingSplineComponent(ToolSettings->DataInstance);
-
 	check(WorkingSpline);
 
 	bDrawTangentForPreviousPoint = false;
@@ -578,16 +621,17 @@ void UPCGDrawSplineToolBase::OnTerminateDragSequence()
 	const int32 PreviousSplinePoint = NumSplinePoints - 1;
 
 	check(PreviousSplinePoint >= 0);
-	
+
+	const UPCGDrawSplineToolSettings* Settings = CastChecked<UPCGDrawSplineToolSettings>(ToolSettings);
 	// Emit the appropriate undo transaction
-	switch (ToolSettings->DrawMode)
+	switch (Settings->DrawMode)
 	{
 		case EPCGDrawSplineDrawMode::ClickAutoTangent:
 		{
 			GetToolManager()->EmitObjectChange(
 				this,
 				MakeUnique<CommandChange::FSimplePointInsertionChange>(
-					WorkingSpline,
+					WorkingSpline.Get(),
 					WorkingSpline->GetLocationAtSplinePoint(PreviousSplinePoint, ESplineCoordinateSpace::World),
 					WorkingSpline->GetUpVectorAtSplinePoint(PreviousSplinePoint, ESplineCoordinateSpace::World)),
 				Constants::AddPointTransactionName);
@@ -598,7 +642,7 @@ void UPCGDrawSplineToolBase::OnTerminateDragSequence()
 			GetToolManager()->EmitObjectChange(
 				this,
 				MakeUnique<CommandChange::FTangentPointInsertionChange>(
-					WorkingSpline,
+					WorkingSpline.Get(),
 					WorkingSpline->GetLocationAtSplinePoint(PreviousSplinePoint, ESplineCoordinateSpace::World),
 					WorkingSpline->GetUpVectorAtSplinePoint(PreviousSplinePoint, ESplineCoordinateSpace::World),
 					WorkingSpline->GetTangentAtSplinePoint(PreviousSplinePoint, ESplineCoordinateSpace::World)),
@@ -618,35 +662,15 @@ void UPCGDrawSplineToolBase::OnTerminateDragSequence()
 
 			GetToolManager()->EmitObjectChange(
 				this,
-				MakeUnique<CommandChange::FStrokeInsertionChange>(WorkingSpline, HitLocations, UpVectors),
+				MakeUnique<CommandChange::FStrokeInsertionChange>(WorkingSpline.Get(), HitLocations, UpVectors),
 				Constants::AddPointTransactionName);
 			break;
 		}
 	}
 }
 
-UPCGInteractiveToolSettings_SplineBase* UPCGDrawSplineToolBase::GetSplineSettings() const
-{
-	TArray<UObject*> AllToolProperties = GetToolProperties(false);
-
-	UObject** FoundProperties = AllToolProperties.FindByPredicate([](const UObject* Candidate)
-	{
-		return Candidate->IsA<UPCGInteractiveToolSettings_SplineBase>();
-	});
-
-	if(FoundProperties)
-	{
-		return Cast<UPCGInteractiveToolSettings_SplineBase>(*FoundProperties);
-	}
-
-	return nullptr;
-}
-
-void UPCGDrawSplineToolBase::OnTick(const float DeltaTime)
+void UPCGDrawSplineTool::OnTick(const float DeltaTime)
 {
-	UPCGInteractiveToolSettings_SplineBase* ToolSettings = GetSplineSettings();
-	
-	USplineComponent* WorkingSpline = ToolSettings->GetWorkingSplineComponent(ToolSettings->DataInstance);
 	if (!WorkingSpline)
 	{
 		GetToolManager()->PostActiveToolShutdownRequest(
@@ -659,13 +683,10 @@ void UPCGDrawSplineToolBase::OnTick(const float DeltaTime)
 	Super::OnTick(DeltaTime);
 }
 
-void UPCGDrawSplineToolBase::Render(IToolsContextRenderAPI* RenderAPI)
+void UPCGDrawSplineTool::Render(IToolsContextRenderAPI* RenderAPI)
 {
 	using namespace PCG::EditorMode::Scene;
 
-	UPCGInteractiveToolSettings_SplineBase* ToolSettings = GetSplineSettings();
-	
-	USplineComponent* WorkingSpline = ToolSettings->GetWorkingSplineComponent(ToolSettings->DataInstance);
 	if (WorkingSpline)
 	{
 		const int32 LastPointIndex = WorkingSpline->GetNumberOfSplinePoints() - 1;
@@ -679,36 +700,31 @@ void UPCGDrawSplineToolBase::Render(IToolsContextRenderAPI* RenderAPI)
 	}
 }
 
-bool UPCGDrawSplineToolBase::CanAccept() const
+void UPCGDrawSplineTool::OnRefresh()
 {
-	UPCGInteractiveToolSettings_SplineBase* ToolSettings = GetSplineSettings();
-
-	USplineComponent* WorkingSpline = ToolSettings->GetWorkingSplineComponent(ToolSettings->DataInstance);
-	return WorkingSpline && WorkingSpline->GetNumberOfSplinePoints() > 0;
-}
+	if (WorkingSpline)
+	{
+		WorkingSpline->PostEditChange();
+	}
 
-void UPCGDrawSplineTool::Setup()
-{
-	AddToolPropertySource(ToolSettings);
-	
-	Super::Setup();
+	Super::OnRefresh();
 }
 
-void UPCGDrawSplineSurfaceTool::Setup()
+bool UPCGDrawSplineTool::CanAccept() const
 {
-	AddToolPropertySource(ToolSettings);
-	
-	Super::Setup();
+	return WorkingSpline && WorkingSpline->GetNumberOfSplinePoints() > 0;
 }
 
 bool UPCGDrawSplineToolBuilder::CanBuildTool(const FToolBuilderState& SceneState) const
 {
-	return SplineToolClass != nullptr;
+	return true;
 }
 
 UInteractiveTool* UPCGDrawSplineToolBuilder::BuildTool(const FToolBuilderState& SceneState) const
 {
-	UPCGDrawSplineToolBase* NewTool = NewObject<UPCGDrawSplineToolBase>(SceneState.ToolManager, SplineToolClass);
+	UPCGDrawSplineTool* NewTool = NewObject<UPCGDrawSplineTool>(SceneState.ToolManager);
+
+	NewTool->SetIsSplineSurface(bIsSplineSurface);
 
 	NewTool->SetToolInfo(
 	{
@@ -716,18 +732,8 @@ UInteractiveTool* UPCGDrawSplineToolBuilder::BuildTool(const FToolBuilderState&
 		.ToolDisplayMessage = PCGDrawSplineTool::Constants::ToolDescription,
 		.ToolIcon = FPCGEditorModeStyle::Get().GetBrush("PCGEditorMode.Tools.DrawSpline")
 	});
-	
-	if(UE::PCG::EditorMode::Tool::BuildTool(NewTool) == false)
-	{
-		return nullptr;
-	}
 
 	return NewTool;
 }
 
-void UPCGDrawSplineToolBuilder::SetSplineToolClass(TSubclassOf<UPCGDrawSplineToolBase> InClass)
-{
-	SplineToolClass = InClass;
-}
-
 #undef LOCTEXT_NAMESPACE
diff --git a/Engine/Plugins/PCG/Source/PCGEditor/Private/EditorMode/Tools/PCGInteractiveTool.cpp b/Engine/Plugins/PCG/Source/PCGEditor/Private/EditorMode/Tools/PCGInteractiveTool.cpp
new file mode 100644
--- /dev/null
+++ b/Engine/Plugins/PCG/Source/PCGEditor/Private/EditorMode/Tools/PCGInteractiveTool.cpp
@@ -0,0 +1,152 @@
+// Copyright Epic Games, Inc. All Rights Reserved.
+
+#include "EditorMode/Tools/PCGInteractiveTool.h"
+
+#include "PCGComponent.h"
+#include "PCGEditorModule.h"
+#include "EditorMode/PCGEdModeHelpers.h"
+#include "EditorMode/PCGEdModeSettings.h"
+#include "EditorMode/Tools/PCGInteractiveToolSettings.h"
+
+#include UE_INLINE_GENERATED_CPP_BY_NAME(PCGInteractiveTool)
+
+#define LOCTEXT_NAMESPACE "PCGInteractiveTool"
+
+void UPCGInteractiveTool::Setup()
+{
+	using namespace PCG::EditorMode;
+	// `GetSettingsClass` must be overridden with a valid settings class for this tool.
+	if (const UClass* SettingsClass = GetSettingsClass())
+	{
+		ToolSettings = NewObject<UPCGInteractiveToolSettings>(this, SettingsClass, SettingsClass->GetDefaultObjectName());
+	}
+
+	// Validate the classes were generated.
+	check(ToolSettings)
+
+	ToolSettings->RestoreProperties(this);
+	AddToolPropertySource(ToolSettings);
+
+	// Do some housecleaning checks. Must be defined prior. Can define in the `PostBuildTool` override.
+	check(!GetToolInfo().ToolDisplayName.IsEmpty())
+	check(!GetToolInfo().ToolDisplayMessage.IsEmpty())
+	check(GetToolInfo().ToolIcon)
+
+	WorkingActor = Actor::SpawnWorking(GetContextualWorld());
+	check(WorkingActor.IsValid());
+	
+	WorkingComponent = PCGComponent::CreateWorking(WorkingActor.Get());
+
+	// @todo_pcg: Maybe a helper to validate the graph?
+	if (!ToolSettings->TargetGraph)
+	{
+		NotifyError(LOCTEXT("InvalidGraph", "Invalid or non-existent graph. Shutting down."));
+		GetToolManager()->ShutdownAfterError();
+		return;
+	}
+
+	check(WorkingComponent.IsValid());
+	WorkingComponent->SetGraph(ToolSettings->TargetGraph);
+	
+	// We set the graph instance of the tool settings so that we can display its parameters
+	ToolSettings->GraphInstance = WorkingComponent->GetGraphInstance();
+	
+	Super::Setup();
+}
+
+void UPCGInteractiveTool::Shutdown(const EToolShutdownType ShutdownType)
+{
+	Super::Shutdown(ShutdownType);
+
+	// Save the updates to the tool settings
+	ToolSettings->SaveProperties(this);
+
+	RemoveToolPropertySource(ToolSettings);
+}
+
+void UPCGInteractiveTool::OnTick(const float DeltaTime)
+{
+	GraphRefreshTimer -= DeltaTime;
+	if (GraphRefreshTimer < 0.f)
+	{
+		OnRefresh();
+	}
+}
+
+void UPCGInteractiveTool::OnPropertyModified(UObject* PropertySet, FProperty* Property)
+{
+	Super::OnPropertyModified(PropertySet, Property);
+
+	if(PropertySet->IsA<UPCGInteractiveToolSettings>() && Property->GetFName() == GET_MEMBER_NAME_CHECKED(UPCGInteractiveToolSettings, TargetGraph))
+	{
+		UPCGInteractiveToolSettings* Settings = Cast<UPCGInteractiveToolSettings>(PropertySet);
+		WorkingComponent->SetGraph(Settings->TargetGraph);
+	}
+}
+
+void UPCGInteractiveTool::OnRefresh()
+{
+	if (WorkingComponent.IsValid() && !WorkingComponent->IsCleaningUp() && !WorkingComponent->IsGenerating())
+	{
+		WorkingComponent->GenerateLocal(true);
+	}
+
+	// Reset the refresh timer
+	GraphRefreshTimer = GetDefault<UPCGEditorModeSettings>()->GraphRefreshRate;
+}
+
+void UPCGInteractiveTool::SetTargetActor(AActor* InTargetActor)
+{
+	if (ToolSettings)
+	{
+		ToolSettings->TargetActor = InTargetActor;
+	}
+}
+
+TWeakObjectPtr<AActor> UPCGInteractiveTool::GetTargetActor() const
+{
+	return ToolSettings ? ToolSettings->TargetActor : nullptr;
+}
+
+UWorld* UPCGInteractiveTool::GetContextualWorld() const
+{
+	return GetToolManager()->GetContextQueriesAPI()->GetCurrentEditingWorld();
+}
+
+UWorld* UPCGInteractiveTool::GetContextualWorldChecked() const
+{
+	UWorld* World = GetContextualWorld();
+	check(World);
+	return World;
+}
+
+void UPCGInteractiveTool::ApplyTool() const
+{
+	ApplyToolInternal();
+}
+
+void UPCGInteractiveTool::SetWarning(const FText& Warning) const
+{
+	WarningMessage = Warning;
+	OnToolWarningChanged.Broadcast(WarningMessage);
+}
+
+void UPCGInteractiveTool::ClearWarning() const
+{
+	WarningMessage = FText::GetEmpty();
+	OnToolWarningChanged.Broadcast(WarningMessage);
+}
+
+void UPCGInteractiveTool::NotifyError(const FText& ErrorMessage) const
+{
+	SetWarning(ErrorMessage);
+
+	if (GetDefault<UPCGEditorModeSettings>()->bShowEditorToastOnToolErrors)
+	{
+		PCG::EditorMode::DispatchEditorNotification(ErrorMessage);
+	}
+
+	UE_LOG(LogPCGEditor, Error, TEXT("[%sTool]: %s"), *GetToolInfo().ToolDisplayName.ToString(), *ErrorMessage.ToString());
+}
+
+#undef LOCTEXT_NAMESPACE
diff --git a/Engine/Plugins/PCG/Source/PCGEditor/Private/EditorMode/Tools/PCGInteractiveToolBuilder.cpp b/Engine/Plugins/PCG/Source/PCGEditor/Private/EditorMode/Tools/PCGInteractiveToolBuilder.cpp
--- a/Engine/Plugins/PCG/Source/PCGEditor/Private/EditorMode/Tools/PCGInteractiveToolBuilder.cpp
+++ b/Engine/Plugins/PCG/Source/PCGEditor/Private/EditorMode/Tools/PCGInteractiveToolBuilder.cpp
@@ -3,59 +3,9 @@
 #include "EditorMode/Tools/PCGInteractiveToolBuilder.h"
 
 #include "InteractiveToolManager.h"
-#include "EditorMode/Tools/PCGInteractiveToolSettings.h"
 
 #include UE_INLINE_GENERATED_CPP_BY_NAME(PCGInteractiveToolBuilder)
 
-#define LOCTEXT_NAMESPACE "PCGInteractiveToolBuilder"
-
-bool UE::PCG::EditorMode::Tool::BuildTool(UInteractiveTool* InTool)
-{
-	// This function uses reflection to instantiate and initialize the tool settings object.
-	// This allows BP to only create a settings property for it to work, in theory
-	
-	using namespace PCG::EditorMode;
-	
-	const FObjectProperty* SettingsProperty = nullptr;
-	const void* PropertyAddress = nullptr;
-	UClass* SettingsClass = nullptr;
-	for(TPropertyValueIterator<FObjectProperty> PropertyIt(InTool->GetClass(), InTool); PropertyIt; ++PropertyIt)
-	{
-		if(PropertyIt.Key()->IsA(FObjectProperty::StaticClass()))
-		{
-			const FObjectProperty* ObjectProperty = CastFieldChecked<FObjectProperty>(PropertyIt.Key());
-			if(ObjectProperty->PropertyClass->IsChildOf(UPCGInteractiveToolSettings::StaticClass()) && ObjectProperty->PropertyClass->HasAnyClassFlags(EClassFlags::CLASS_Abstract) == false)
-			{
-				SettingsProperty = ObjectProperty;
-				PropertyAddress = PropertyIt.Value();
-				SettingsClass = ObjectProperty->PropertyClass;
-				break;
-			}
-		}
-	}
-	
-	if(SettingsProperty == nullptr || PropertyAddress == nullptr || SettingsClass == nullptr)
-	{
-		FText Message = FText::FormatOrdered(LOCTEXT("PCGSettingsPropertyMissingMessage", "Settings {0} of tool {1} does not have a valid PCGInteractiveToolSettings property. Did you forget to mark it as UPROPERTY?"), SettingsClass ? SettingsClass->GetClass()->GetDisplayNameText() : FText::GetEmpty(), InTool->GetClass()->GetDisplayNameText());
-		InTool->GetToolManager()->DisplayMessage(Message, EToolMessageLevel::UserError);
-		return false;
-	}
-	
-	UPCGInteractiveToolSettings* Settings = NewObject<UPCGInteractiveToolSettings>(InTool, SettingsClass);
-	if(Settings->Initialize(InTool) == false)
-	{
-		Settings->Cancel(InTool);
-		Settings->Shutdown();
-		
-		FText Message = FText::FormatOrdered(LOCTEXT("PCGSettingsInitializationFailedMessage", "Tool {0} failed initialization on settings {1}"), InTool->GetClass()->GetDisplayNameText(), Settings->GetClass()->GetDisplayNameText());
-		InTool->GetToolManager()->DisplayMessage(Message, EToolMessageLevel::UserError);
-		return false;
-	}
-	
-	SettingsProperty->SetObjectPropertyValue(const_cast<void*>(PropertyAddress), Settings);
-	return true;
-}
-
 bool UPCGInteractiveToolBuilder::CanBuildTool(const FToolBuilderState& SceneState) const
 {
 	return Super::CanBuildTool(SceneState);
@@ -66,42 +16,8 @@ UInteractiveTool* UPCGInteractiveToolBuilder::BuildTool(const FToolBuilderState&
 	return Super::BuildTool(SceneState);
 }
 
-void UPCGInteractiveToolBuilder::PostSetupTool(UInteractiveTool* Tool, const FToolBuilderState& SceneState) const
-{
-	Super::PostBuildTool(Tool, SceneState);
-	SceneState.ToolManager->DisplayMessage(Tool->GetToolInfo().ToolDisplayMessage, EToolMessageLevel::UserNotification);
-
-	TArray<UObject*> ToolProperties = Tool->GetToolProperties();
-
-	if(ToolGraph)
-	{
-		for(UObject* ToolPropertySet : ToolProperties)
-		{
-			if(UPCGInteractiveToolSettings* Settings = Cast<UPCGInteractiveToolSettings>(ToolPropertySet))
-			{
-				Settings->SetToolGraph(ToolGraph);
-			}
-		}
-	}
-}
-
-void UPCGInteractiveToolWithToolTargetsBuilder::PostSetupTool(UInteractiveTool* Tool, const FToolBuilderState& SceneState) const
+void UPCGInteractiveToolBuilder::PostBuildTool(UInteractiveTool* Tool, const FToolBuilderState& SceneState) const
 {
 	Super::PostBuildTool(Tool, SceneState);
 	SceneState.ToolManager->DisplayMessage(Tool->GetToolInfo().ToolDisplayMessage, EToolMessageLevel::UserNotification);
-
-	TArray<UObject*> ToolProperties = Tool->GetToolProperties();
-
-	if(ToolGraph)
-	{
-		for(UObject* ToolPropertySet : ToolProperties)
-		{
-			if(UPCGInteractiveToolSettings* Settings = Cast<UPCGInteractiveToolSettings>(ToolPropertySet))
-			{
-				Settings->SetToolGraph(ToolGraph);
-			}
-		}
-	}
 }
-
-#undef LOCTEXT_NAMESPACE
diff --git a/Engine/Plugins/PCG/Source/PCGEditor/Private/EditorMode/Tools/PCGInteractiveToolSettings.cpp b/Engine/Plugins/PCG/Source/PCGEditor/Private/EditorMode/Tools/PCGInteractiveToolSettings.cpp
--- a/Engine/Plugins/PCG/Source/PCGEditor/Private/EditorMode/Tools/PCGInteractiveToolSettings.cpp
+++ b/Engine/Plugins/PCG/Source/PCGEditor/Private/EditorMode/Tools/PCGInteractiveToolSettings.cpp
@@ -1,522 +1,25 @@
 // Copyright Epic Games, Inc. All Rights Reserved.
 
 #include "EditorMode/Tools/PCGInteractiveToolSettings.h"
-#include "PCGComponent.h"
+
 #include "PCGEditorModule.h"
 #include "PCGEditorSettings.h"
 #include "EditorMode/PCGEdModeHelpers.h"
 #include "EditorMode/PCGEdModeSettings.h"
-#include "EditorMode/Tools/Helpers/PCGEdModeEditorUtilities.h"
-#include "Elements/PCGDataFromTool.h"
-#include "Subsystems/PCGSubsystem.h"
-
-#include "InteractiveToolManager.h"
-#include "BaseBehaviors/KeyInputBehavior.h"
+#include "EditorMode/Tools/PCGInteractiveTool.h"
+#include "Helpers/PCGHelpers.h"
 #include "Logging/StructuredLog.h"
 
 #include UE_INLINE_GENERATED_CPP_BY_NAME(PCGInteractiveToolSettings)
 
-#define LOCTEXT_NAMESPACE "PCGInteractiveToolSettings"
-
-namespace UE::PCG::EditorMode::Tool
-{
-	TAutoConsoleVariable<bool> CVarCommitPropertyChangedEventOnTimer(TEXT("pcg.Editor.Tools.CommitPropertyChangedEventOnTimer"), true,
-		TEXT("If true, will queue up and automatically commit property changed events on a timer instead of committing directly."));
-}
-
-void UE::PCG::EditorMode::Tool::Shutdown(UInteractiveTool* InTool, EToolShutdownType ShutdownType)
-{
-	for(UObject* ToolPropertyObject : InTool->GetToolProperties(false))
-	{
-		if(UPCGInteractiveToolSettings* PCGSettings = Cast<UPCGInteractiveToolSettings>(ToolPropertyObject))
-		{
-			if(ShutdownType == EToolShutdownType::Cancel)
-			{
-				PCGSettings->Cancel(InTool);
-			}
-			else if(ShutdownType == EToolShutdownType::Accept)
-			{
-				PCGSettings->Apply(InTool);
-			}
-		}
-	}
-
-	for(UObject* ToolPropertyObject : InTool->GetToolProperties(false))
-	{
-		if(UPCGInteractiveToolSettings* PCGSettings = Cast<UPCGInteractiveToolSettings>(ToolPropertyObject))
-		{
-			PCGSettings->Shutdown();
-		}
-	}	
-}
-
 UPCGInteractiveToolSettings::UPCGInteractiveToolSettings()
 	: NewActorName(GetDefault<UPCGEditorModeSettings>()->DefaultNewActorName)
-	, NewPCGComponentName(GetDefault<UPCGEditorModeSettings>()->DefaultNewPCGComponentName)
-{
-	
-}
-
-bool UPCGInteractiveToolSettings::Initialize(UInteractiveTool* InTool)
-{
-	using namespace UE::PCG::EditorMode;
-
-	UWorld* World = InTool->GetToolManager()->GetContextQueriesAPI()->GetCurrentEditingWorld();
-
-	FToolBuilderState State;
-	InTool->GetToolManager()->GetContextQueriesAPI()->GetCurrentSelectionState(State);
-	TArray<AActor*> SelectedActors = State.SelectedActors;
-
-	// If we have selected one actor, that's the actor we want to edit
-	// This actor might:
-	// - Not have a PCG Component
-	// - Have a PCG Component but no valid working data
-	// - Have a PCG Component with valid working data we want to edit
-	if(SelectedActors.Num() == 1)
-	{		
-		EditActor = SelectedActors[0];
-	}
-	// If we have no selected actor, we spawn a working actor that needs to fulfill the same actor requirements. This actor is fully instantiated when Applied.
-	else
-	{
-		AActor* SpawnedActor = Actor::SpawnWorking(World, FTransform::Identity, GetWorkingActorClass());
-		// our spawned actor is also the actor we are editing
-		EditActor = SpawnedActor;
-
-		GeneratedResources.SpawnedActor = SpawnedActor;
-	}
-
-	if(!ensureMsgf(EditActor.IsValid(), TEXT("The Edit Actor has to be valid at this point. Either as a previously existing actor, or a new spawned actor.")))
-	{
-		return false;
-	}
-	
-	TWeakObjectPtr<UPCGComponent> PCGComponent = FindOrGeneratePCGComponent(*EditActor.Get());
-	if(!ensureMsgf(PCGComponent.IsValid(), TEXT("The PCG Component has to be valid at this point. Either as a previously existing PCG Component, or a transient generated PCG Component.")))
-	{
-		return false;
-	}
-
-	OnModified.AddUObject(this, &UPCGInteractiveToolSettings::OnPropertyModified);
-	
-	// If we don't have a valid graph already, we try to load a default graph
-	if(PCGComponent->GetGraph() == nullptr)
-	{
-		TryLoadDefaultGraph();
-	}
-	// Otherwise, we use the graph we already have
-	else
-	{
-		SetToolGraph(PCGComponent->GetGraph());
-	}
-	
-	// Keyboard setup for switching between data instance identifiers
-	UKeyInputBehavior* KeyInputBehavior = NewObject<UKeyInputBehavior>();
-	KeyInputBehavior->Initialize(this,
-		{
-			EKeys::One, EKeys::Two, EKeys::Three, EKeys::Four, EKeys::Five,
-			EKeys::Six, EKeys::Seven, EKeys::Eight, EKeys::Nine, EKeys::Zero
-		});
-	KeyInputBehavior->bRequireAllKeys = false;
-	InTool->AddInputBehavior(KeyInputBehavior);
-
-	return true;
-}
-
-void UPCGInteractiveToolSettings::Apply(UInteractiveTool* OwningTool)
-{
-	using namespace UE::PCG::EditorMode;
-
-	OwningTool->GetToolManager()->BeginUndoTransaction(LOCTEXT("ApplyToolTransactionName", "Applied PCG Tool"));
-	
-	if (GeneratedResources.SpawnedActor.IsValid())
-	{
-		GeneratedResources.SpawnedActor->ClearFlags(RF_Transient);
-	}
-	
-	if (GeneratedResources.GeneratedPCGComponent.IsValid())
-	{
-		// On Apply, if we have a PCG Component that was generated by the tool, we remove the transient flag
-		GeneratedResources.GeneratedPCGComponent->ClearFlags(RF_Transient);
-	}
-
-	/** Now each working data struct gets to Apply its changes. */
-	for (const auto& WorkingDataMap : GetMutableTypedWorkingDataMap<FPCGInteractiveToolWorkingData>())
-	{
-		FPCGInteractiveToolWorkingDataContext Context;
-		Context.OwningPCGComponent = GetWorkingPCGComponent();
-		Context.DataInstanceIdentifier = WorkingDataMap.Key;
-		Context.OwningActor = Context.OwningPCGComponent->GetOwner();
-		Context.InteractiveTool = OwningTool;
-		Context.PCGSettings = this;
-		
-		WorkingDataMap.Value->OnToolApply(Context);
-	}
-
-   OwningTool->GetToolManager()->EndUndoTransaction();
-}
-
-void UPCGInteractiveToolSettings::Cancel(UInteractiveTool* OwningTool)
-{
-	/** Each working data struct is responsible for wiping or restoring data to the previous state. */
-	for(const auto& WorkingDataMap : GetMutableTypedWorkingDataMap<FPCGInteractiveToolWorkingData>())
-	{
-		FPCGInteractiveToolWorkingDataContext Context;
-		Context.OwningPCGComponent = GetWorkingPCGComponent();
-		Context.DataInstanceIdentifier = WorkingDataMap.Key;
-		Context.OwningActor = Context.OwningPCGComponent->GetOwner();
-		Context.InteractiveTool = OwningTool;
-		Context.PCGSettings = this;
-
-		// OnToolCancel will delete resources generated by the tool data itself
-		WorkingDataMap.Value->OnToolCancel(Context);
-	}
-
-	DeleteGeneratedResources();
-}
-
-void UPCGInteractiveToolSettings::Shutdown()
-{
-	using namespace UE::PCG::EditorMode;
-	
-	OnModified.RemoveAll(this);
-	
-	// We mark it as garbage so it stops ticking
-	MarkAsGarbage();
-	
-	// On Shutdown, we flush the cache and regenerate to make sure our results are up to date
-	if(UPCGComponent* Component = GetWorkingPCGComponent())
-	{
-		if (UPCGSubsystem* PCGSubsystem = Component->GetSubsystem())
-		{
-			PCGSubsystem->FlushCache();
-		}
-		
-		Component->Generate(true);
-	}
-}
-
-void UPCGInteractiveToolSettings::InitializeWorkingData(FName DataInstanceIdentifier)
-{
-	UPCGComponent* PCGComponent = GetWorkingPCGComponent();
-	const UScriptStruct* Struct = GetWorkingDataType();
-	ensureMsgf(Struct->IsChildOf<FPCGInteractiveToolWorkingData>(), TEXT("Working Data Type must inherit from FPCGInteractiveToolWorkingData."));
-
-	FPCGInteractiveToolWorkingDataContext Context;
-	Context.OwningPCGComponent = PCGComponent;
-	Context.DataInstanceIdentifier = DataInstanceIdentifier;
-	Context.OwningActor = Context.OwningPCGComponent->GetOwner();
-	Context.InteractiveTool = GetTypedOuter<UInteractiveTool>();
-	Context.PCGSettings = this;
-	
-	FName WorkingDataIdentifier = GetWorkingDataIdentifier(DataInstanceIdentifier);
-	
-	if (EditActor.IsValid())
-	{
-		// If we don't find the requested working data, we create it
-		if(PCGComponent->ToolDataContainer.ToolData.Contains(WorkingDataIdentifier) == false)
-		{
-			TInstancedStruct<FPCGInteractiveToolWorkingData> NewToolWorkingData;
-			NewToolWorkingData.InitializeAsScriptStruct(Struct);
-			
-			PCGComponent->ToolDataContainer.ToolData.Add(WorkingDataIdentifier, MoveTemp(NewToolWorkingData));
-			GeneratedResources.AddedWorkingDataIdentifiers.Add(WorkingDataIdentifier);
-		}
-	}
-
-	FPCGInteractiveToolWorkingData* WorkingData = GetMutableTypedWorkingData(DataInstanceIdentifier);
-	
-	// If it was freshly constructed, we initialize it. This will allow the structs to create or find components to work with.
-	if(ensureMsgf(WorkingData, TEXT("Working Data has to be allocated at this point!.")))
-	{
-		if(WorkingData->IsInitialized() == false)
-		{
-			WorkingData->OnInitialized().BindUObject(this, &UPCGInteractiveToolSettings::OnWorkingDataInitialized);
-			WorkingData->Initialize(Context);
-			WorkingData->OnInitialized().Unbind();
-		}
-	}
-	else
-	{
-		return;
-	}
-	
-	// The working data needs to be initialized at this point
-	ensureMsgf(WorkingData->IsInitialized(), TEXT("Working Data isn't initialized!. Have you forgotten to call the parent Initialize function?"));
-	
-	// If the data is valid, we call start. If not, we delete it.
-	if(WorkingData->IsValid())
-	{
-		// Notify the working data that we have begun accessing it. This allows to make a 'last state' that the tool can revert to on Cancel.
-		WorkingData->OnToolStart(Context);
-	}
-	else
-	{
-		// If the working data is invalid despite being initialized, we delete and request it again.
-		// This can happen when data dependencies have been deleted outside of PCG, such as deleting a spline component on the actor
-		WorkingData->OnToolCancel(Context);
-		PCGComponent->ToolDataContainer.ToolData.Remove(WorkingDataIdentifier);
-
-		// In which case, requesting new working data will recreate a new spline
-		InitializeWorkingData(DataInstanceIdentifier);
-	}	
-}
-
-void UPCGInteractiveToolSettings::SetToolGraph(UPCGGraphInterface* InGraph)
-{
-	ToolGraph = InGraph;
-
-	// This will call OnPropertyModified, forwarding the change
-	FPropertyChangedEvent Event(FindFProperty<FProperty>(UPCGInteractiveToolSettings::StaticClass(), GET_MEMBER_NAME_CHECKED(UPCGInteractiveToolSettings, ToolGraph)));
-	PostEditChangeProperty(Event);
-}
-
-void UPCGInteractiveToolSettings::SetDataInstance(FName DataInstanceIdentifier)
-{
-	DataInstance = DataInstanceIdentifier;
-	
-	// This will call OnPropertyModified, forwarding the change
-	FPropertyChangedEvent Event(FindFProperty<FProperty>(UPCGInteractiveToolSettings::StaticClass(), GET_MEMBER_NAME_CHECKED(UPCGInteractiveToolSettings, DataInstance)));
-	PostEditChangeProperty(Event);
-}
-
-void UPCGInteractiveToolSettings::SetActorClassToSpawn(TSubclassOf<AActor> Class)
-{
-	ActorClassToSpawn = Class;
-
-	// This will call OnPropertyModified, forwarding the change
-	FPropertyChangedEvent Event(FindFProperty<FProperty>(UPCGInteractiveToolSettings::StaticClass(), GET_MEMBER_NAME_CHECKED(UPCGInteractiveToolSettings, ActorClassToSpawn)));
-	PostEditChangeProperty(Event);
-}
-
-UPCGComponent* UPCGInteractiveToolSettings::FindOrGeneratePCGComponent(AActor& OwningActor) const
-{
-	using namespace UE::PCG::EditorMode;
-	
-	UPCGComponent* PCGComponent = PCGComponent::Find(&OwningActor);
-
-	if(PCGComponent == nullptr)
-	{
-		PCGComponent = PCGComponent::Create(&OwningActor, "ToolPCGComponent");
-		PCGComponent->SetFlags(RF_Transient);
-		GeneratedResources.GeneratedPCGComponent = PCGComponent;
-	}
-
-	return PCGComponent;
-}
-
-void UPCGInteractiveToolSettings::DeleteGeneratedResources() const
-{
-	if(GeneratedResources.SpawnedActor.IsValid())
-	{
-		GeneratedResources.SpawnedActor->Destroy();
-		GeneratedResources.SpawnedActor.Reset();
-	}
-	
-	if(GeneratedResources.GeneratedPCGComponent.IsValid())
-	{
-		GeneratedResources.GeneratedPCGComponent->Cleanup(true);
-		GeneratedResources.GeneratedPCGComponent->DestroyComponent();
-		GeneratedResources.GeneratedPCGComponent.Reset();
-	}
-
-	// If we still have a PCG Component, this means it was not a generated one. We have to remove the working data.
-	WorkingComponentCache.Reset();
-	if(UPCGComponent* PCGComponent = GetWorkingPCGComponent())
-	{
-		for(const FName& AddedWorkingDataIdentifier : GeneratedResources.AddedWorkingDataIdentifiers)
-		{
-			PCGComponent->ToolDataContainer.ToolData.Remove(AddedWorkingDataIdentifier);
-		}
-	}
-}
+	, NewPCGComponentName(GetDefault<UPCGEditorModeSettings>()->DefaultNewPCGComponentName) {}
 
 void UPCGInteractiveToolSettings::TryLoadDefaultGraph()
 {
 	check(!HasAnyFlags(RF_ClassDefaultObject));
-	SetToolGraph(CastChecked<UPCGGraphInterface>(GetDefaultGraph().TryLoad(), ECastCheckedType::NullAllowed));
-}
-
-void UPCGInteractiveToolSettings::Tick(float DeltaTime)
-{
-	GraphRefreshTimer -= DeltaTime;
-	if (GraphRefreshTimer < 0.f)
-	{
-		OnRefresh();
-	}
-}
-
-bool UPCGInteractiveToolSettings::IsTickable() const
-{
-	return HasAnyFlags(RF_ClassDefaultObject | RF_MirroredGarbage) == false;
-}
-
-TStatId UPCGInteractiveToolSettings::GetStatId() const
-{
-	RETURN_QUICK_DECLARE_CYCLE_STAT(FPCGInteractiveToolSettings, STATGROUP_Tickables);
-}
-
-void UPCGInteractiveToolSettings::OnRefresh()
-{
-	if(HasAnyFlags(RF_MirroredGarbage) == false)
-	{
-		if(bAutoGenerate)
-		{
-			UPCGComponent* PCGComponent = GetWorkingPCGComponent();
-			if(PCGComponent && !PCGComponent->IsCleaningUp() && !PCGComponent->IsGenerating())
-			{
-				PCGComponent->GenerateLocal(true);
-			}
-		}
-
-		// We commit the property changed events now to let the graph regenerate, if the object is tracked
-		CommitPropertyChangedEvents();
-		
-		// Reset the refresh timer
-		GraphRefreshTimer = GetDefault<UPCGEditorModeSettings>()->GraphRefreshRate;
-
-		OnRefreshInternal();
-	}
-}
-
-bool UPCGInteractiveToolSettings::CanEditChange(const FProperty* InProperty) const
-{
-	bool bSuperResult =  Super::CanEditChange(InProperty);
-
-	if(bSuperResult == true)
-	{
-		if(InProperty->GetFName() == GET_MEMBER_NAME_CHECKED(UPCGInteractiveToolSettings, ActorClassToSpawn))
-		{
-			return EditActor.IsValid();
-		}
-	}
-
-	return bSuperResult;
-}
-
-void UPCGInteractiveToolSettings::CommitPropertyChangedEvents()
-{
-	for(auto& Pair : PropertyChangedEventQueue)
-	{
-		if(Pair.Key.IsValid())
-		{
-			FCoreUObjectDelegates::OnObjectPropertyChanged.Broadcast(Pair.Key.Get(), Pair.Value);
-		}
-	}
-
-	PropertyChangedEventQueue.Empty();
-}
-
-void UPCGInteractiveToolSettings::OnKeyPressed(const FKey& InKeyID)
-{
-	TArray<FName> ValidDataInstanceNames = GetValidDataInstanceNamesForGraph();
-
-	int32 Index = INDEX_NONE;
-	
-	if(InKeyID == EKeys::One)
-	{
-		Index = 1;
-	}
-	else if(InKeyID == EKeys::Two)
-	{
-		Index = 2;
-	}
-	else if(InKeyID == EKeys::Three)
-	{
-		Index = 3;
-	}
-	else if(InKeyID == EKeys::Four)
-	{
-		Index = 4;
-	}
-	else if(InKeyID == EKeys::Five)
-	{
-		Index = 5;
-	}
-	else if(InKeyID == EKeys::Six)
-	{
-		Index = 6;
-	}
-	else if(InKeyID == EKeys::Seven)
-	{
-		Index = 7;
-	}
-	else if(InKeyID == EKeys::Eight)
-	{
-		Index = 8;
-	}
-	else if(InKeyID == EKeys::Nine)
-	{
-		Index = 9;
-	}
-	else if(InKeyID == EKeys::Zero)
-	{
-		Index = 10;
-	}
-
-	Index--;
-	
-	if (ValidDataInstanceNames.IsValidIndex(Index))
-	{
-		SetDataInstance(ValidDataInstanceNames[Index]);
-	}
-}
-
-void UPCGInteractiveToolSettings::OnWorkingDataInitialized(FPCGInteractiveToolWorkingData* WorkingData) const
-{
-}
-
-void UPCGInteractiveToolSettings::OnPropertyModified(UObject* Object, FProperty* Property)
-{
-	if(Object == this && Property->GetFName() == GET_MEMBER_NAME_CHECKED(UPCGInteractiveToolSettings, ToolGraph))
-	{
-		GetWorkingPCGComponent()->SetGraph(ToolGraph);
-
-		if(ToolGraph)
-		{
-			// If the graph changes, we also set the default data instance
-			SetActorClassToSpawn(ToolGraph->ToolData.InitialActorClassToSpawn);
-		}
-
-		// We set the data instance to the first one provided by the graph. Uses the default as fallback.
-		FName DefaultDataInstanceName = GetDefaultDataInstanceName();
-		SetDataInstance(DefaultDataInstanceName);
-		
-		// If we have a graph, we update the working data to the available data instances ones from the graph
-		if(ToolGraph)
-		{
-			// We update the working data
-			TArray<FName> ValidDataInstanceNames = GetValidDataInstanceNamesForGraph();
-
-			if(ValidDataInstanceNames.Num() > 0)
-			{
-				for(const FName& DataInstanceIdentifier : ValidDataInstanceNames)
-				{
-					InitializeWorkingData(DataInstanceIdentifier);
-				}
-			}
-			// We make sure at least the default is available so that the tool data exists and can be used even if the graph doesn't make use of it
-			else
-			{
-				InitializeWorkingData(DefaultDataInstanceName);
-			}
-			
-		}
-		// If not, we use the default
-		else
-		{
-			InitializeWorkingData(DefaultDataInstanceName);
-		}
-	}
-	if(Object == this && Property->GetFName() == GET_MEMBER_NAME_CHECKED(UPCGInteractiveToolSettings, ActorClassToSpawn))
-	{
-		// @todo_pcg
-		// Delete previously spawned actor and set up the new one. Ideally recycling the existing PCG Component.
-	}
-}
-
-void UPCGInteractiveToolSettings::OnRefreshInternal()
-{
+	TargetGraph = CastChecked<UPCGGraphInterface>(GetDefaultGraph().TryLoad(), ECastCheckedType::NullAllowed);
 }
 
 FSoftObjectPath UPCGInteractiveToolSettings::GetDefaultGraph() const
@@ -532,109 +35,31 @@ FSoftObjectPath UPCGInteractiveToolSettings::GetDefaultGraph() const
 		{
 			return Settings->DefaultGraph.ToSoftObjectPath();
 		}
-		
+	
+		UE_LOGFMT(LogPCGEditor, Warning, "Could not get default graph for settings class {0} in tool {1}", this->GetClass()->GetDisplayNameText().ToString(), GetTypedOuter<UPCGInteractiveTool>()->GetClass()->GetDisplayNameText().ToString());
 		return FSoftObjectPath();
 	}
-	
-	return FSoftObjectPath();
-}
 
-TSubclassOf<AActor> UPCGInteractiveToolSettings::GetWorkingActorClass() const
-{
-	return ActorClassToSpawn;
+	UE_LOGFMT(LogPCGEditor, Warning, "Could not find tool settings for settings class {0} in tool {1}. Please add them in the PCG Editor Settings.", this->GetClass()->GetDisplayNameText().ToString(), GetTypedOuter<UPCGInteractiveTool>()->GetClass()->GetDisplayNameText().ToString());
+	return FSoftObjectPath();
 }
 
 bool UPCGInteractiveToolSettings::GraphAssetFilter(const FAssetData& AssetData) const
 {
-	using namespace UE::PCG::EditorMode;
-	
-	if(AssetData.GetClass()->IsChildOf<UPCGGraphInterface>() == false)
-	{
-		return true;
-	}
-	
-	return !Utility::DoesPCGGraphHaveToolTag(GetToolTag(), AssetData);
-}
-
-UPCGComponent* UPCGInteractiveToolSettings::GetWorkingPCGComponent() const
-{
-	if(WorkingComponentCache.IsValid() == false)
+	// This can be checked by subclasses by either finding a tag or more specifically checking for values.
+	if (AssetData.FindTag(GET_MEMBER_NAME_CHECKED(UPCGGraph, CompatibleToolTags)))
 	{
-		WorkingComponentCache = UE::PCG::EditorMode::PCGComponent::Find(EditActor.Get());
-	}
-	
-	return WorkingComponentCache.Get();
-}
-
-UPCGGraphInstance* UPCGInteractiveToolSettings::GetGraphInstance() const
-{
-	if(UPCGComponent* PCGComponent = GetWorkingPCGComponent())
-	{
-		return PCGComponent->GetGraphInstance();
-	}
-	
-	return nullptr;
-}
-
-TArray<FName> UPCGInteractiveToolSettings::GetValidDataInstanceNamesForGraph() const
-{
-	TArray<FName> Result;
-
-	if (ToolGraph == nullptr)
-	{
-		return Result;
-	}
-
-	TArray<UPCGNode*> Nodes = ToolGraph->GetGraph()->FindNodesWithSettings(UPCGDataFromTool::StaticClass(), false);
-
-	FName ToolTag = GetToolTag();
-	for (UPCGNode* Node : Nodes)
-	{
-		if (UPCGDataFromTool* DataFromTool = Cast<UPCGDataFromTool>(Node->GetSettings()))
+		FString ToolTagsString;
+		TArray<FString> Tags;
+		if (AssetData.GetTagValue<FString>(GET_MEMBER_NAME_CHECKED(UPCGGraph, CompatibleToolTags), ToolTagsString))
 		{
-			if (DataFromTool->ToolTag != ToolTag)
-			{
-				continue;
-			}
-			
-			Result.AddUnique(DataFromTool->DataInstance);
+			Tags = PCGHelpers::GetStringArrayFromCommaSeparatedList(ToolTagsString);
 		}
-	}
-	
-	return Result;
-}
-
-FName UPCGInteractiveToolSettings::GetDefaultDataInstanceName() const
-{
-	TArray<FName> ValidDataInstanceNames = GetValidDataInstanceNamesForGraph();
-	
-	// By default, we use the regular tool tag as the working data identifier, if available
-	// This means a NONE instance name.
-	if (ValidDataInstanceNames.Contains(NAME_None))
-	{
-		return NAME_None;
-	}
-
-	// If we have multiple available and they don't include, we choose the first available one
-	if (ValidDataInstanceNames.Num() > 0)
-	{
-		return ValidDataInstanceNames[0];
-	}
-	
-	return NAME_None;
-}
 
-FName UPCGInteractiveToolSettings::GetWorkingDataIdentifier(FName DataInstanceIdentifier) const
-{
-	FNameBuilder DataIdentifierBuilder(GetToolTag());
-
-	if (DataInstanceIdentifier.IsNone() == false)
-	{
-		DataIdentifierBuilder.AppendChar('_');
-		DataIdentifierBuilder.Append(DataInstanceIdentifier.ToString());
+		return !Tags.FindByPredicate([ToolTag = GetToolTag().ToString()](const FString& InTag) { return InTag.Equals(ToolTag, ESearchCase::IgnoreCase); });
 	}
 
-	return DataIdentifierBuilder.ToString();
+	// True == filter the asset
+	return true;
 }
 
-#undef LOCTEXT_NAMESPACE
diff --git a/Engine/Plugins/PCG/Source/PCGEditor/Private/EditorMode/Tools/Paint/PCGPaintTool.cpp b/Engine/Plugins/PCG/Source/PCGEditor/Private/EditorMode/Tools/Paint/PCGPaintTool.cpp
--- a/Engine/Plugins/PCG/Source/PCGEditor/Private/EditorMode/Tools/Paint/PCGPaintTool.cpp
+++ b/Engine/Plugins/PCG/Source/PCGEditor/Private/EditorMode/Tools/Paint/PCGPaintTool.cpp
@@ -2,19 +2,7 @@
 
 #include "EditorMode/Tools/Paint/PCGPaintTool.h"
 
-#include "InteractiveGizmoManager.h"
-#include "InteractiveToolManager.h"
-#include "MaterialEditingLibrary.h"
-#include "PreviewMesh.h"
-#include "ToolTargetManager.h"
-#include "BaseBehaviors/TwoAxisPropertyEditBehavior.h"
-#include "BaseGizmos/VolumetricBrushStampIndicator.h"
-#include "Components/BoxComponent.h"
 #include "EditorMode/PCGEdModeStyle.h"
-#include "Helpers/PCGHelpers.h"
-#include "TargetInterfaces/PrimitiveComponentBackedTarget.h"
-#include "ToolTargets/PrimitiveComponentToolTarget.h"
-#include "Materials/MaterialInstanceDynamic.h"
 
 #include UE_INLINE_GENERATED_CPP_BY_NAME(PCGPaintTool)
 
@@ -22,358 +10,73 @@
 
 namespace PCGPaintTool::Constants
 {
-	const FText ToolDisplayName = LOCTEXT("ToolName", "PaintTool");
-	const FName ToolTag = "PaintTool";
-	const FText ToolDescription = LOCTEXT("DisplayMessage", "Draw points onto a surface and let the PCG graph process them.");
+	const FText ToolName = LOCTEXT("ToolName", "PaintTool");
+	const FText ToolDescription = LOCTEXT("DisplayMessage", "Tool not yet implemented.");
 }
 
-UPCGInteractiveToolSettings_PaintTool::UPCGInteractiveToolSettings_PaintTool()
-{
-	// We turn auto generation off, as we'll make use of the tracking system to regenerate instead
-	bAutoGenerate = false;
-} 
-
-FName UPCGInteractiveToolSettings_PaintTool::GetToolTag() const
-{
-	return PCGPaintTool::Constants::ToolTag;
-}
-
-const UScriptStruct* UPCGInteractiveToolSettings_PaintTool::GetWorkingDataType() const
-{
-	return FPCGInteractiveToolWorkingData_PointArrayData::StaticStruct();
-}
-
-void UPCGPaintTool::Setup()
-{
-	PropertyClass = UPCGEdModeBrushBase::StaticClass();
-	
-	AddToolPropertySource(ToolSettings);
-
-	AActor* TargetActor = Cast<UPrimitiveComponentToolTarget>(GetTarget())->GetOwnerActor();
-
-	if(ToolSettings->GeneratedResources.SpawnedActor.IsValid())
-	{
-		ToolSettings->GeneratedResources.SpawnedActor->SetActorLocation(TargetActor->GetActorLocation());
-	}
-	
-	Super::Setup();	
-}
-
-void UPCGPaintTool::Shutdown(EToolShutdownType ShutdownType)
-{
-	UE::PCG::EditorMode::Tool::Shutdown(this, ShutdownType);
-	
-	Super::Shutdown(ShutdownType);
-}
-
-bool UPCGPaintTool::HasAccept() const
-{
-	return true;
-}
+UPCGPaintSettings::UPCGPaintSettings(const FObjectInitializer& ObjectInitializer)
+	: Brush(Cast<UPCGEdModeBrushBase>(
+		ObjectInitializer.CreateEditorOnlyDefaultSubobject(
+			this,
+			GET_MEMBER_NAME_CHECKED(UPCGPaintSettings, Brush),
+			UPCGEdModeBrushBase::StaticClass(),
+			/*bIsTransient=*/true)))
+{} // @todo_pcg: The brush will inevitably be serialized somewhere as the user modifies it.
 
-bool UPCGPaintTool::CanAccept() const
+FText UPCGPaintSettings::GetToolTag() const
 {
-	return true;
+	return PCGPaintTool::Constants::ToolName;
 }
 
-bool UPCGPaintTool::HasCancel() const
+FSoftObjectPath UPCGPaintSettings::GetDefaultGraph() const
 {
-	return true;
+	return FSoftObjectPath{};
 }
 
-void UPCGPaintTool::OnTick(float DeltaTime)
+bool UPCGPaintSettings::ApplyTool(UWorld* World, FText& OutWarning) const
 {
-	Super::OnTick(DeltaTime);
-
-	bool bShiftPressedThisTick = GetShiftToggle();
-	bool bShiftPressedChanged = bShiftPressedThisTick != bShiftPressedLastTick;
-	
-	if(bPaintOnTick)
-	{
-		if(bShiftPressedThisTick)
-		{
-			TryDeletePointsFromBrush();	
-		}
-		else
-		{
-			AddSinglePointFromBrush();
-		}
-
-		bPaintOnTick = false;
-	}
-
-	if(bShiftPressedChanged)
-	{
-		if(UVolumetricBrushStampIndicator* VolumetricStampIndicator = Cast<UVolumetricBrushStampIndicator>(BrushStampIndicator))
-		{
-			if(UMaterialInstanceDynamic* Instance = Cast<UMaterialInstanceDynamic>(VolumetricStampIndicator->GetPreviewMesh()->GetActiveMaterial()))
-			{
-				if(bShiftPressedThisTick)
-				{
-					Instance->SetVectorParameterValue("IndicatorColor", FLinearColor::Red);
-				}
-				else
-				{
-					FLinearColor DefaultColor;
-					Instance->GetVectorParameterDefaultValue(FName("IndicatorColor"), DefaultColor);
-					Instance->SetVectorParameterValue("IndicatorColor", DefaultColor);
-				}
-			}
-		}
-	}
-
-	bShiftPressedLastTick = bShiftPressedThisTick;
-}
-
-void UPCGPaintTool::OnClickPress(const FInputDeviceRay& PressPos)
-{
-	Super::OnClickPress(PressPos);
-
-	bPaintOnTick = true;
-}
-
-void UPCGPaintTool::OnUpdateDrag(const FRay& Ray)
-{
-	Super::OnUpdateDrag(Ray);
-
-	bPaintOnTick = true;
+	// @todo_pcg: Apply the output of the tool here.
+	return false;
 }
 
-void UPCGPaintTool::SetupBrushStampIndicator()
+void UPCGPaintTool::Setup()
 {
-	if (!BrushStampIndicator)
-	{
-		// See USkinWeightsPaintTool::SetupBrushStampIndicator
-		FString VolumetricBrushIndicatorGizmoType = TEXT("VolumetricBrushIndicatorGizmoType");
-		// register and spawn brush indicator gizmo
-		GetToolManager()->GetPairedGizmoManager()->RegisterGizmoType(VolumetricBrushIndicatorGizmoType, NewObject<UVolumetricBrushStampIndicatorBuilder>());
-		BrushStampIndicator = GetToolManager()->GetPairedGizmoManager()->CreateGizmo<UVolumetricBrushStampIndicator>(VolumetricBrushIndicatorGizmoType, FString(), this);
-		if (UVolumetricBrushStampIndicator* VolumetricBrushIndicator = Cast<UVolumetricBrushStampIndicator>(BrushStampIndicator))
-		{
-			VolumetricBrushIndicator->MakeBrushIndicatorMesh(this, GetTargetWorld());
-		}
-	}
-}
+	// @todo_pcg: Set up the tool here.
+	Super::Setup();
 
-void UPCGPaintTool::UpdateBrushStampIndicator()
-{
-	if (BrushStampIndicator)
-	{
-		if (BrushEditBehavior.IsValid())
-		{
-			FLinearColor NewColor;
-			if(BrushEditBehavior->IsEditing())
-			{
-				NewColor = FLinearColor::White;
-			}
-			else
-			{
-				if(GetShiftToggle())
-				{
-					NewColor = FLinearColor::Red;
-				}
-				else
-				{
-					NewColor = FLinearColor::Green;
-				}
-			}
-			
-			BrushStampIndicator->LineColor = NewColor;	
-		}
-		
-		BrushStampIndicator->Update(BrushProperties->BrushRadius, LastBrushStamp.WorldPosition, LastBrushStamp.WorldNormal, BrushProperties->BrushFalloffAmount, BrushProperties->BrushStrength);
-	}
-}
+	UPCGPaintSettings* Settings = CastChecked<UPCGPaintSettings>(ToolSettings);
 
-void UPCGPaintTool::ShutdownBrushStampIndicator()
-{
-	if (BrushStampIndicator)
-	{
-		FString VolumetricBrushIndicatorGizmoType = TEXT("VolumetricBrushIndicatorGizmoType");
-		
-		GetToolManager()->GetPairedGizmoManager()->DestroyGizmo(BrushStampIndicator);
-		BrushStampIndicator = nullptr;
-		GetToolManager()->GetPairedGizmoManager()->DeregisterGizmoType(VolumetricBrushIndicatorGizmoType);
-	}
+	NotifyError(LOCTEXT("ToolNotYetImplemented", "The tool is not yet implemented."));
 }
 
-double UPCGPaintTool::EstimateMaximumTargetDimension()
+TSubclassOf<UPCGInteractiveToolSettings> UPCGPaintTool::GetSettingsClass() const
 {
-	// We return the target actor's max dimension
-	return ToolSettings->EditActor->GetComponentsBoundingBox().GetExtent().GetAbsMax();
+	return UPCGPaintSettings::StaticClass();
 }
 
-void UPCGPaintTool::AddSinglePointFromBrush()
-{
-	UPCGPointArrayData* PointData = ToolSettings->GetTypedWorkingData<FPCGInteractiveToolWorkingData_PointArrayData>(ToolSettings->DataInstance)->GetPointArrayData();
-	
-	const int32 PointIndex = PointData->GetNumPoints();
-	PointData->SetNumPoints(PointIndex + 1);
-	
-	// Values to assign
-	const float PointSteepness = 1.0f;
-	const FTransform PointTransform = FTransform(LastBrushStamp.WorldPosition);
-	const FVector& Position = PointTransform.GetLocation();
-	const int32 PointSeed = PCGHelpers::ComputeSeed((int)Position.X, (int)Position.Y, (int)Position.Z);
-	const FBox LocalBounds = FBox(FVector(-CurrentBrushRadius), FVector(CurrentBrushRadius));
-	const FVector PointBoundsMin = LocalBounds.Min;
-	const FVector PointBoundsMax = LocalBounds.Max;
-	const int64 PointMetadataEntry = PointData->Metadata->AddEntry();
-
-	// SteepnessRange - Initialize and Allocate if needed 
-	const TConstPCGValueRange<float> ConstSteepnessRange = PointData->GetConstSteepnessValueRange();
-	TOptional<const float> SteepnessSingleValue = ConstSteepnessRange.GetSingleValue();
-	const bool bAllocateSteepness = SteepnessSingleValue.IsSet() && SteepnessSingleValue.GetValue() != PointSteepness;
-	TPCGValueRange<float> SteepnessRange = PointData->GetSteepnessValueRange(bAllocateSteepness);
-	
-	// TransformRange - Initialize and Allocate if needed 
-	const TConstPCGValueRange<FTransform> ConstTransformRange = PointData->GetConstTransformValueRange();
-	TOptional<const FTransform> TransformSingleValue = ConstTransformRange.GetSingleValue();
-	const bool bAllocateTransform = TransformSingleValue.IsSet() && !TransformSingleValue.GetValue().Equals(PointTransform);
-	TPCGValueRange<FTransform> TransformRange = PointData->GetTransformValueRange(bAllocateTransform);
-
-	// SeedRange - Initialize and Allocate if needed 
-	const TConstPCGValueRange<int32> ConstSeedRange = PointData->GetConstSeedValueRange();
-	TOptional<const int32> SeedSingleValue = ConstSeedRange.GetSingleValue();
-	const bool bAllocateSeed = SeedSingleValue.IsSet() && SeedSingleValue.GetValue() != PointSeed;
-	TPCGValueRange<int32> SeedRange = PointData->GetSeedValueRange(bAllocateSeed);
-
-	// BoundsMinRange - Initialize and Allocate if needed 
-	const TConstPCGValueRange<FVector> ConstBoundsMinRange = PointData->GetConstBoundsMinValueRange();
-	TOptional<const FVector> BoundsMinSingleValue = ConstBoundsMinRange.GetSingleValue();
-	const bool bAllocateBoundsMin = BoundsMinSingleValue.IsSet() && BoundsMinSingleValue.GetValue() != PointBoundsMin;
-	TPCGValueRange<FVector> BoundsMinRange = PointData->GetBoundsMinValueRange(bAllocateBoundsMin);
-		
-	// BoundsMaxRange - Initialize and Allocate if needed 
-	const TConstPCGValueRange<FVector> ConstBoundsMaxRange = PointData->GetConstBoundsMaxValueRange();
-	TOptional<const FVector> BoundsMaxSingleValue = ConstBoundsMaxRange.GetSingleValue();
-	const bool bAllocateBoundsMax = BoundsMaxSingleValue.IsSet() && BoundsMaxSingleValue.GetValue() != PointBoundsMax;
-	TPCGValueRange<FVector> BoundsMaxRange = PointData->GetBoundsMaxValueRange(bAllocateBoundsMax);
-
-	// MetadataEntryRange - Initialize and Allocate if needed 
-	const TConstPCGValueRange<int64> ConstMetadataEntryRange = PointData->GetConstMetadataEntryValueRange();
-	TOptional<const int64> MetadataEntrySingleValue = ConstMetadataEntryRange.GetSingleValue();
-	const bool bAllocateMetadataEntry = MetadataEntrySingleValue.IsSet() && MetadataEntrySingleValue.GetValue() != PointMetadataEntry;
-	TPCGValueRange<int64> MetadataEntryRange = PointData->GetMetadataEntryValueRange(bAllocateMetadataEntry);
-
-	// Assign values
-	SteepnessRange[PointIndex] = PointSteepness;
-	TransformRange[PointIndex] = PointTransform;
-	SeedRange[PointIndex] = PointSeed;
-	BoundsMinRange[PointIndex] = PointBoundsMin;
-	BoundsMaxRange[PointIndex] = PointBoundsMax;
-	MetadataEntryRange[PointIndex] = PointMetadataEntry;
-
-	FPCGAttributeIdentifier BrushSizeAttributeIdentifier;
-	BrushSizeAttributeIdentifier.Name = "BrushSize";
-	BrushSizeAttributeIdentifier.MetadataDomain = PointData->Metadata->GetConstDefaultMetadataDomain()->GetDomainID();
-	
-	FPCGMetadataAttribute<float>* BrushSizeAttribute = PointData->Metadata->FindOrCreateAttribute(BrushSizeAttributeIdentifier, 100.f, /*bAllowsInterpolation=*/false, /*bOverrideParent=*/false, /*bOverwriteIfTypeMismatch=*/false);
-	if (BrushSizeAttribute)
-	{
-		BrushSizeAttribute->SetValue(PointMetadataEntry, CurrentBrushRadius);
-	}
-
-	FPCGAttributeIdentifier FalloffAttributeIdentifier;
-	FalloffAttributeIdentifier.Name = "BrushFalloff";
-	FalloffAttributeIdentifier.MetadataDomain = PointData->Metadata->GetConstDefaultMetadataDomain()->GetDomainID();
-	FPCGMetadataAttribute<float>* FalloffAttribute = PointData->Metadata->FindOrCreateAttribute(FalloffAttributeIdentifier, 100.f, /*bAllowsInterpolation=*/false, /*bOverrideParent=*/false, /*bOverwriteIfTypeMismatch=*/false);
-	if (FalloffAttribute)
-	{
-		FalloffAttribute->SetValue(PointMetadataEntry, BrushProperties->BrushFalloffAmount);
-	}
-	
-
-	// @todo_pcg
-	// Instead of using a queue on timer, can we put this on another thread? Direct broadcasting takes too much perf on tick, timed queue results in delays
-
-	// GetToolData dynamically tracks the pointdata object and will rerun the graph when it changes
-	// To avoid doing this on tick, we queue up the changes
-	
-	FPropertyChangedEvent EmptyPropertyChangedEvent(nullptr);
-	if(UE::PCG::EditorMode::Tool::CVarCommitPropertyChangedEventOnTimer.GetValueOnAnyThread())
-	{
-		ToolSettings->PropertyChangedEventQueue.Add(PointData, EmptyPropertyChangedEvent);
-	}
-	else
-	{
-		FCoreUObjectDelegates::OnObjectPropertyChanged.Broadcast(PointData, EmptyPropertyChangedEvent);
-	}
-}
-
-void UPCGPaintTool::TryDeletePointsFromBrush()
+bool UPCGPaintTool::CanAccept() const
 {
-	auto DeletePointsFromDataInstance = [this](FName DataInstanceName)
-	{
-		UPCGPointArrayData* PointData = ToolSettings->GetTypedWorkingData<FPCGInteractiveToolWorkingData_PointArrayData>(DataInstanceName)->GetPointArrayData();
-		FSphere Sphere(LastBrushStamp.WorldPosition, CurrentBrushRadius);
-
-		TArray<FTransform> Transforms = PointData->GetTransformsCopy();
-		for(int32 Index = PointData->GetNumPoints() - 1; Index >= 0 ; --Index)
-		{
-			if(Sphere.IsInside(Transforms[Index].GetLocation()))
-			{
-				PointData->RemoveAt(Index);
-			}
-		}
-
-		FPropertyChangedEvent EmptyPropertyChangedEvent(nullptr);
-		if(UE::PCG::EditorMode::Tool::CVarCommitPropertyChangedEventOnTimer.GetValueOnAnyThread())
-		{
-			ToolSettings->PropertyChangedEventQueue.Add(PointData, EmptyPropertyChangedEvent);
-		}
-		else
-		{
-			FCoreUObjectDelegates::OnObjectPropertyChanged.Broadcast(PointData, EmptyPropertyChangedEvent);
-		}
-	};
-	
-	if(ToolSettings->bErasePointsOfSelectedDataInstance)
-	{
-		DeletePointsFromDataInstance(ToolSettings->DataInstance);
-	}
-	else
-	{
-		for(const FName& DataInstanceName : ToolSettings->GetValidDataInstanceNamesForGraph())
-		{
-			DeletePointsFromDataInstance(DataInstanceName);
-		}
-	}
-	
+	// @todo_pcg: true once implemented
+	return false;
 }
 
 bool UPCGPaintToolBuilder::CanBuildTool(const FToolBuilderState& SceneState) const
 {
-	return SceneState.TargetManager->CountSelectedAndTargetable(SceneState, GetTargetRequirements()) > 0;
+	return true;
 }
 
 UInteractiveTool* UPCGPaintToolBuilder::BuildTool(const FToolBuilderState& SceneState) const
 {
 	UPCGPaintTool* NewTool = NewObject<UPCGPaintTool>(SceneState.ToolManager);
-	
+
 	NewTool->SetToolInfo(
 	{
-		.ToolDisplayName = PCGPaintTool::Constants::ToolDisplayName,
+		.ToolDisplayName = PCGPaintTool::Constants::ToolName,
 		.ToolDisplayMessage = PCGPaintTool::Constants::ToolDescription,
 		.ToolIcon = FPCGEditorModeStyle::Get().GetBrush("PCGEditorMode.Tools.Paint")
 	});
 
-	UToolTarget* const Target = SceneState.TargetManager->BuildFirstSelectedTargetable(SceneState, GetTargetRequirements());
-	NewTool->SetTarget(Target);
-	NewTool->SetWorld(SceneState.World);
-
-	if(UE::PCG::EditorMode::Tool::BuildTool(NewTool) == false)
-	{
-		return nullptr;
-	}
-	
 	return NewTool;
 }
 
-const FToolTargetTypeRequirements& UPCGPaintToolBuilder::GetTargetRequirements() const
-{
-	static FToolTargetTypeRequirements TypeRequirements(UPrimitiveComponentBackedTarget::StaticClass());
-	return TypeRequirements;
-}
-
 #undef LOCTEXT_NAMESPACE
diff --git a/Engine/Plugins/PCG/Source/PCGEditor/Private/EditorMode/Tools/Volume/PCGCreateVolumeTool.cpp b/Engine/Plugins/PCG/Source/PCGEditor/Private/EditorMode/Tools/Volume/PCGCreateVolumeTool.cpp
--- a/Engine/Plugins/PCG/Source/PCGEditor/Private/EditorMode/Tools/Volume/PCGCreateVolumeTool.cpp
+++ b/Engine/Plugins/PCG/Source/PCGEditor/Private/EditorMode/Tools/Volume/PCGCreateVolumeTool.cpp
@@ -2,7 +2,6 @@
 
 #include "EditorMode/Tools/Volume/PCGCreateVolumeTool.h"
 
-#include "InteractiveToolManager.h"
 #include "EditorMode/PCGEdModeStyle.h"
 
 #include UE_INLINE_GENERATED_CPP_BY_NAME(PCGCreateVolumeTool)
@@ -11,27 +10,20 @@
 
 namespace PCGCreateVolumeTool::Constants
 {
-	const FText ToolDisplayName = LOCTEXT("ToolDisplayName", "Create Volume");
-	const FName ToolName = "CreateVolumeTool";
+	const FText ToolName = LOCTEXT("ToolName", "CreateVolume");
 }
 
-FName UPCGInteractiveToolSettings_CreateVolume::GetToolTag() const
+FText UPCGCreateVolumeToolSettings::GetToolTag() const
 {
 	return PCGCreateVolumeTool::Constants::ToolName;
 }
 
-const UScriptStruct* UPCGInteractiveToolSettings_CreateVolume::GetWorkingDataType() const
-{
-	// @todo_pcg Add createvolume tool data struct here 
-	return nullptr;
-}
-
-FSoftObjectPath UPCGInteractiveToolSettings_CreateVolume::GetDefaultGraph() const
+FSoftObjectPath UPCGCreateVolumeToolSettings::GetDefaultGraph() const
 {
 	return FSoftObjectPath{};
 }
 
-bool UPCGInteractiveToolSettings_CreateVolume::ApplyTool(UWorld* World, FText& OutWarning) const
+bool UPCGCreateVolumeToolSettings::ApplyTool(UWorld* World, FText& OutWarning) const
 {
 	// @todo_pcg: Apply the output of the tool here.
 	return false;
@@ -41,15 +33,15 @@ void UPCGCreateVolumeTool::Setup()
 {
 	// @todo_pcg: Set up the tool here.
 	Super::Setup();
-	
-	AddToolPropertySource(ToolSettings);
+
+	UPCGCreateVolumeToolSettings* Settings = CastChecked<UPCGCreateVolumeToolSettings>(ToolSettings);
+
+	NotifyError(LOCTEXT("ToolNotYetImplemented", "The tool is not yet implemented."));
 }
 
 void UPCGCreateVolumeTool::Shutdown(EToolShutdownType ShutdownType)
 {
-	UE::PCG::EditorMode::Tool::Shutdown(this, ShutdownType);
-
-	Super::Shutdown(ShutdownType);
+	// @todo_pcg: Clean up the tool here.
 }
 
 void UPCGCreateVolumeTool::OnTick(const float DeltaTime)
@@ -68,7 +60,7 @@ void UPCGCreateVolumeTool::Render(IToolsContextRenderAPI* RenderAPI)
 	FPrimitiveDrawInterface* PDI = RenderAPI->GetPrimitiveDrawInterface();
 	check(PDI);
 
-	const UPCGInteractiveToolSettings_CreateVolume* Settings = CastChecked<UPCGInteractiveToolSettings_CreateVolume>(ToolSettings);
+	const UPCGCreateVolumeToolSettings* Settings = CastChecked<UPCGCreateVolumeToolSettings>(ToolSettings);
 
 	/** @todo_pcg: Three click process -
 	 *	First click - Point (Location)
@@ -77,6 +69,11 @@ void UPCGCreateVolumeTool::Render(IToolsContextRenderAPI* RenderAPI)
 	 */
 }
 
+TSubclassOf<UPCGInteractiveToolSettings> UPCGCreateVolumeTool::GetSettingsClass() const
+{
+	return UPCGCreateVolumeToolSettings::StaticClass();
+}
+
 FInputRayHit UPCGCreateVolumeTool::IsHitByClick(const FInputDeviceRay& ClickPosition)
 {
 	using namespace PCG::EditorMode::Scene;
@@ -90,7 +87,7 @@ FInputRayHit UPCGCreateVolumeTool::IsHitByClick(const FInputDeviceRay& ClickPosi
 
 void UPCGCreateVolumeTool::OnClicked(const FInputDeviceRay& ClickPosition)
 {
-	UPCGInteractiveToolSettings_CreateVolume* Settings = CastChecked<UPCGInteractiveToolSettings_CreateVolume>(ToolSettings);
+	UPCGCreateVolumeToolSettings* Settings = CastChecked<UPCGCreateVolumeToolSettings>(ToolSettings);
 
 	/** @todo_pcg: Three click process -
 	 *	First click - Point (Location)
@@ -101,10 +98,6 @@ void UPCGCreateVolumeTool::OnClicked(const FInputDeviceRay& ClickPosition)
 	// @todo_pcg: Different chords to translate, rotate, scale at any step of the process?
 }
 
-void UPCGCreateVolumeTool::ApplyTool() const
-{
-}
-
 bool UPCGCreateVolumeToolBuilder::CanBuildTool(const FToolBuilderState& SceneState) const
 {
 	return true;
@@ -116,15 +109,10 @@ UInteractiveTool* UPCGCreateVolumeToolBuilder::BuildTool(const FToolBuilderState
 
 	NewTool->SetToolInfo(
 	{
-		.ToolDisplayName = PCGCreateVolumeTool::Constants::ToolDisplayName,
+		.ToolDisplayName = PCGCreateVolumeTool::Constants::ToolName,
 		.ToolDisplayMessage = LOCTEXT("CreateVolumeToolDisplayMessage", "Tool not yet implemented."),
 		.ToolIcon = FPCGEditorModeStyle::Get().GetBrush("PCGEditorMode.Tools.CreateVolume")
 	});
-	
-	if(UE::PCG::EditorMode::Tool::BuildTool(NewTool) == false)
-	{
-		return nullptr;
-	}	
 
 	return NewTool;
 }
diff --git a/Engine/Plugins/PCG/Source/PCGEditor/Public/EditorMode/PCGEdModeCommands.h b/Engine/Plugins/PCG/Source/PCGEditor/Public/EditorMode/PCGEdModeCommands.h
--- a/Engine/Plugins/PCG/Source/PCGEditor/Public/EditorMode/PCGEdModeCommands.h
+++ b/Engine/Plugins/PCG/Source/PCGEditor/Public/EditorMode/PCGEdModeCommands.h
@@ -62,15 +62,15 @@ public:
 	/*** PCG EDITOR MODE TOOLS - Line Context **/
 	TSharedPtr<FUICommandInfo> EnableDrawSplineTool;
 	TSharedPtr<FUICommandInfo> EnableDrawSurfaceTool;
-	TArray<TSharedPtr<FUICommandInfo>> GetLineContextToolCommands();
+	TArray<TSharedPtr<FUICommandInfo>> GetLineContextToolCommands() const;
 
 	/*** PCG EDITOR MODE TOOLS - Paint Context **/
 	TSharedPtr<FUICommandInfo> EnablePaintTool;
-	TArray<TSharedPtr<FUICommandInfo>> GetPaintContextToolCommands();
+	TArray<TSharedPtr<FUICommandInfo>> GetPaintContextToolCommands() const;
 
 	/*** PCG EDITOR MODE TOOLS - Volume Context **/
 	TSharedPtr<FUICommandInfo> EnableCreateVolumeTool;
-	TArray<TSharedPtr<FUICommandInfo>> GetVolumeContextToolCommands();
+	TArray<TSharedPtr<FUICommandInfo>> GetVolumeContextToolCommands() const;
 };
 
 /**
diff --git a/Engine/Plugins/PCG/Source/PCGEditor/Public/EditorMode/PCGEdModeToolkit.h b/Engine/Plugins/PCG/Source/PCGEditor/Public/EditorMode/PCGEdModeToolkit.h
--- a/Engine/Plugins/PCG/Source/PCGEditor/Public/EditorMode/PCGEdModeToolkit.h
+++ b/Engine/Plugins/PCG/Source/PCGEditor/Public/EditorMode/PCGEdModeToolkit.h
@@ -24,7 +24,6 @@ public:
 	virtual FText GetToolPaletteDisplayName(FName PaletteName) const override;
 	virtual void OnToolPaletteChanged(FName PaletteName) override;
 	virtual bool HasIntegratedToolPalettes() const override { return false; }
-	virtual void BuildToolPalette(FName Palette, class FToolBarBuilder& ToolbarBuilder) override;
 
 	virtual FText GetActiveToolDisplayName() const override;
 	virtual FText GetActiveToolMessage() const override;
@@ -60,8 +59,8 @@ private:
 	void RegisterToolkitCallbacks();
 	void UnregisterToolkitCallbacks();
 
-	void RegisterToolCallbacks(UInteractiveTool* CurrentTool);
-	void UnregisterToolCallbacks(UInteractiveTool* CurrentTool) const;
+	void RegisterToolCallbacks(UPCGInteractiveTool* CurrentTool);
+	void UnregisterToolCallbacks(UPCGInteractiveTool* CurrentTool) const;
 
 	void OnEditorModeSettingsChanged(UObject* InObject, FPropertyChangedEvent& InChangedEvent);
 
diff --git a/Engine/Plugins/PCG/Source/PCGEditor/Public/EditorMode/Tools/Helpers/PCGEdModeActorHelpers.h b/Engine/Plugins/PCG/Source/PCGEditor/Public/EditorMode/Tools/Helpers/PCGEdModeActorHelpers.h
--- a/Engine/Plugins/PCG/Source/PCGEditor/Public/EditorMode/Tools/Helpers/PCGEdModeActorHelpers.h
+++ b/Engine/Plugins/PCG/Source/PCGEditor/Public/EditorMode/Tools/Helpers/PCGEdModeActorHelpers.h
@@ -4,7 +4,7 @@
 
 #include "Helpers/PCGActorHelpers.h"
 
-namespace UE::PCG::EditorMode
+namespace PCG::EditorMode
 {
 	namespace Tags
 	{
@@ -18,7 +18,7 @@ namespace UE::PCG::EditorMode
 		/** Spawn an actor for the toolkit with explicit parameters. */
 		AActor* Spawn(UWorld* World, FName ActorName, const FTransform& Transform = FTransform::Identity);
 		/** Helper to spawn a temporary working actor for a tool. */
-		AActor* SpawnWorking(UWorld* World, const FTransform& Transform = FTransform::Identity, TSubclassOf<AActor> ActorClass = AActor::StaticClass());
+		AActor* SpawnWorking(UWorld* World, const FTransform& Transform = FTransform::Identity);
 	}
 
 	namespace PCGComponent
diff --git a/Engine/Plugins/PCG/Source/PCGEditor/Public/EditorMode/Tools/Helpers/PCGEdModeEditorUtilities.h b/Engine/Plugins/PCG/Source/PCGEditor/Public/EditorMode/Tools/Helpers/PCGEdModeEditorUtilities.h
deleted file mode 100644
diff --git a/Engine/Plugins/PCG/Source/PCGEditor/Public/EditorMode/Tools/Helpers/PCGEdModeSplineHelpers.h b/Engine/Plugins/PCG/Source/PCGEditor/Public/EditorMode/Tools/Helpers/PCGEdModeSplineHelpers.h
--- a/Engine/Plugins/PCG/Source/PCGEditor/Public/EditorMode/Tools/Helpers/PCGEdModeSplineHelpers.h
+++ b/Engine/Plugins/PCG/Source/PCGEditor/Public/EditorMode/Tools/Helpers/PCGEdModeSplineHelpers.h
@@ -7,8 +7,8 @@
 class USplineComponent;
 class AActor;
 
-namespace UE::PCG::EditorMode::Spline
-{	
+namespace PCG::EditorMode::Spline
+{
 	USplineComponent* CreateComponent(
 		AActor* Actor,
 		FName ComponentName,
@@ -17,7 +17,7 @@ namespace UE::PCG::EditorMode::Spline
 
 	USplineComponent* CreateWorkingComponent(AActor* Actor);
 
-	USplineComponent* Find(const AActor* Actor, const int32 TargetIndex = INDEX_NONE);
+	USplineComponent* FindComponent(const AActor* Actor, const int32 TargetIndex = INDEX_NONE);
 
 	void CopySplinePoints(const USplineComponent& Source, USplineComponent& Destination, const bool bTransactional = false);
 }
diff --git a/Engine/Plugins/PCG/Source/PCGEditor/Public/EditorMode/Tools/Line/PCGDrawSplineTool.h b/Engine/Plugins/PCG/Source/PCGEditor/Public/EditorMode/Tools/Line/PCGDrawSplineTool.h
--- a/Engine/Plugins/PCG/Source/PCGEditor/Public/EditorMode/Tools/Line/PCGDrawSplineTool.h
+++ b/Engine/Plugins/PCG/Source/PCGEditor/Public/EditorMode/Tools/Line/PCGDrawSplineTool.h
@@ -2,13 +2,13 @@
 
 #pragma once
 
+#include "EditorMode/Tools/PCGInteractiveTool.h"
 #include "EditorMode/Tools/PCGInteractiveToolBuilder.h"
 #include "EditorMode/Tools/PCGInteractiveToolCommon.h"
 #include "EditorMode/Tools/PCGInteractiveToolSettings.h"
 #include "EditorMode/Tools/Helpers/PCGEdModeSceneQueryHelpers.h"
 
 #include "BaseBehaviors/BehaviorTargetInterfaces.h"
-#include "Components/SplineComponent.h"
 
 #include "PCGDrawSplineTool.generated.h"
 
@@ -23,20 +23,21 @@ enum class EPCGDrawSplineDrawMode : uint8
 	FreeDraw UMETA(ToolTip = "Drag to place multiple points, with spacing controlled by Min Point Spacing.")
 };
 
-UCLASS(Abstract, BlueprintType)
-class PCGEDITOR_API UPCGInteractiveToolSettings_SplineBase : public UPCGInteractiveToolSettings
+UCLASS(BlueprintType)
+class PCGEDITOR_API UPCGDrawSplineToolSettings : public UPCGInteractiveToolSettings
 {
 	GENERATED_BODY()
 
 public:
-	UPCGInteractiveToolSettings_SplineBase();
-	
-	USplineComponent* GetWorkingSplineComponent(FName DataInstanceIdentifier) const;
-	TArray<USplineComponent*> GetWorkingSplineComponents() const;
-	
+	UPCGDrawSplineToolSettings();
+
+	// ~Begin UPCGInteractiveToolSettings interface
+	virtual FText GetToolTag() const override;
+	// ~End UPCGInteractiveToolSettings interface
+
 	UPROPERTY(EditAnywhere, Category = "Spline")
-	EPCGDrawSplineDrawMode DrawMode = EPCGDrawSplineDrawMode::FreeDraw;
-	
+	EPCGDrawSplineDrawMode DrawMode = EPCGDrawSplineDrawMode::TangentDrag;
+
 	/** Minimum world distance to cover while free drawing before a point will be generated. */
 	UPROPERTY(EditAnywhere, Category = "Spline", meta = (
 			ClampMin = 0,
@@ -53,35 +54,8 @@ public:
 	FName SplineName;
 };
 
-UCLASS(BlueprintType)
-class PCGEDITOR_API UPCGInteractiveToolSettings_Spline : public UPCGInteractiveToolSettings_SplineBase
-{
-	GENERATED_BODY()
-
-	
-protected:
-	virtual const UScriptStruct* GetWorkingDataType() const override;
-	virtual FName GetToolTag() const override;
-	virtual void OnWorkingDataInitialized(FPCGInteractiveToolWorkingData* WorkingData) const override;
-	virtual void OnPropertyModified(UObject* Object, FProperty* Property) override;
-public:
-	UPROPERTY(EditAnywhere, Category = "Spline")
-	bool bClosedSpline = false;
-};
-
-
-UCLASS(BlueprintType)
-class PCGEDITOR_API UPCGInteractiveToolSettings_SplineSurface : public UPCGInteractiveToolSettings_SplineBase
-{
-	GENERATED_BODY()
-
-protected:
-	virtual const UScriptStruct* GetWorkingDataType() const override;
-	virtual FName GetToolTag() const override;
-};
-
-UCLASS(Abstract, Experimental)
-class PCGEDITOR_API UPCGDrawSplineToolBase : public UInteractiveTool, public IClickBehaviorTarget, public IClickDragBehaviorTarget
+UCLASS(Experimental)
+class PCGEDITOR_API UPCGDrawSplineTool : public UPCGInteractiveTool, public IClickBehaviorTarget, public IClickDragBehaviorTarget
 {
 	GENERATED_BODY()
 
@@ -97,6 +71,11 @@ public:
 	virtual void Render(IToolsContextRenderAPI* RenderAPI) override;
 	// ~End UInteractiveTool interface
 
+	// ~Begin UPCGInteractiveTool interface
+	virtual void OnRefresh() override;
+	virtual TSubclassOf<UPCGInteractiveToolSettings> GetSettingsClass() const override { return UPCGDrawSplineToolSettings::StaticClass(); }
+	// ~End UPCGInteractiveTool interface
+
 	// ~Begin IClickBehaviorTarget interface
 	virtual FInputRayHit IsHitByClick(const FInputDeviceRay& ClickPosition) override;
 	virtual void OnClicked(const FInputDeviceRay& ClickPosition) override;
@@ -110,47 +89,28 @@ public:
 	virtual void OnTerminateDragSequence() override;
 	// ~End IClickDragBehaviorTarget interface
 
-	UPCGInteractiveToolSettings_SplineBase* GetSplineSettings() const;
+	void SetIsSplineSurface(const bool bInIsSplineSurface) { bIsSplineSurface = bInIsSplineSurface; }
+
+protected:
+	virtual void ApplyToolInternal() const override;
 private:
 	void AddSplinePoint(const FVector& HitLocation, const FVector& HitNormal) const;
 	FVector GetNewPointNormal(const FVector& HitNormal) const;
 
 	UPROPERTY(Instanced)
 	TObjectPtr<USingleClickOrDragInputBehavior> ClickOrDragBehavior = nullptr;
-	
+
+	UPROPERTY(Transient)
+	TObjectPtr<USplineComponent> WorkingSpline = nullptr;
+
 	PCG::EditorMode::Scene::FViewRayParams CachedRayParams{};
+	bool bIsSplineSurface = false;
 	bool bDrawTangentForPreviousPoint = false;
 	bool bFreeDrawPlacedPreviewPoint = false;
 	int32 FreeDrawStrokeStartIndex = 0;
-
-protected:
 	double CachedMinPointSpacingSquared = 0;
 };
 
-UCLASS(Experimental)
-class PCGEDITOR_API UPCGDrawSplineTool : public UPCGDrawSplineToolBase
-{
-	GENERATED_BODY()
-
-public:
-	virtual void Setup() override;
-public:
-	UPROPERTY(Transient)
-	TObjectPtr<UPCGInteractiveToolSettings_Spline> ToolSettings;
-};
-
-UCLASS(Experimental)
-class PCGEDITOR_API UPCGDrawSplineSurfaceTool : public UPCGDrawSplineToolBase
-{
-	GENERATED_BODY()
-
-public:
-	virtual void Setup() override;
-public:
-	UPROPERTY(Transient)
-	TObjectPtr<UPCGInteractiveToolSettings_SplineSurface> ToolSettings;
-};
-
 UCLASS(Transient)
 class PCGEDITOR_API UPCGDrawSplineToolBuilder : public UPCGInteractiveToolBuilder
 {
@@ -160,8 +120,8 @@ public:
 	virtual bool CanBuildTool(const FToolBuilderState& SceneState) const override;
 	virtual UInteractiveTool* BuildTool(const FToolBuilderState& SceneState) const override;
 
-	void SetSplineToolClass(TSubclassOf<UPCGDrawSplineToolBase> InClass);
+	void SetIsSplineSurface(const bool bInIsSplineSurface) { bIsSplineSurface = bInIsSplineSurface; }
+
 private:
-	UPROPERTY()
-	TSubclassOf<UPCGDrawSplineToolBase> SplineToolClass;
+	bool bIsSplineSurface = false;
 };
diff --git a/Engine/Plugins/PCG/Source/PCGEditor/Public/EditorMode/Tools/PCGInteractiveTool.h b/Engine/Plugins/PCG/Source/PCGEditor/Public/EditorMode/Tools/PCGInteractiveTool.h
new file mode 100644
--- /dev/null
+++ b/Engine/Plugins/PCG/Source/PCGEditor/Public/EditorMode/Tools/PCGInteractiveTool.h
@@ -0,0 +1,69 @@
+// Copyright Epic Games, Inc. All Rights Reserved.
+
+#pragma once
+
+#include "InteractiveTool.h"
+#include "InteractiveToolManager.h"
+
+#include "PCGInteractiveTool.generated.h"
+
+DECLARE_MULTICAST_DELEGATE_OneParam(OnPCGInteractiveToolWarningChanged, const FText& Warning);
+
+class UPCGComponent;
+class UPCGInteractiveToolSettings;
+
+UCLASS(Abstract, Transient)
+class PCGEDITOR_API UPCGInteractiveTool : public UInteractiveTool
+{
+	GENERATED_BODY()
+
+public:
+	// ~Begin UInteractiveTool interface
+	virtual void Setup() override;
+	virtual void Shutdown(EToolShutdownType ShutdownType) override;
+	virtual void OnTick(float DeltaTime) override;
+	virtual void OnPropertyModified(UObject* PropertySet, FProperty* Property) override;
+	// ~End UInteractiveTool interface
+
+	virtual void OnRefresh();
+	virtual void SetTargetActor(AActor* InTargetActor);
+	virtual TWeakObjectPtr<AActor> GetTargetActor() const;
+	
+	UWorld* GetContextualWorld() const;
+	UWorld* GetContextualWorldChecked() const;
+
+	OnPCGInteractiveToolWarningChanged OnToolWarningChanged;
+
+	void ApplyTool() const;
+
+protected:
+	/** Functionality to execute when the tool is applied. This must be defined in the subclass. */
+	virtual void ApplyToolInternal() const PURE_VIRTUAL(UPCGInteractiveTool::ApplyToolInternal, return;)
+	
+	/** The concrete class of the tool. This must be defined in the subclass. */
+	virtual TSubclassOf<UPCGInteractiveToolSettings> GetSettingsClass() const PURE_VIRTUAL(UPCGInteractiveTool::GetSettingsClass, return nullptr;)
+
+	/** Set the warning for the tool. */
+	virtual void SetWarning(const FText& Warning) const;
+	/** Clear the warning for the tool. */
+	virtual void ClearWarning() const;
+	/** Notify the tool of an error to log, display, and optionally toast. */
+	virtual void NotifyError(const FText& ErrorMessage) const;
+
+	// @todo_pcg: Convert to a target object, rather than actor and component, to ease transition into SceneGraph
+	/** The tool can work directly on this temporary actor until effects are applied. */
+	TWeakObjectPtr<AActor> WorkingActor = nullptr;
+	/** The tool can work directly on this temporary PCG component until effects are applied. */
+	TWeakObjectPtr<UPCGComponent> WorkingComponent = nullptr;
+
+public:
+	UPROPERTY(EditAnywhere, Category = "Tool Settings", meta = (ShowOnlyInnerProperties))
+	TObjectPtr<UPCGInteractiveToolSettings> ToolSettings = nullptr;
+
+private:
+	TWeakObjectPtr<UWorld> TargetWorld = nullptr;
+	mutable FText WarningMessage;
+
+	/** The current time left before a graph refresh (in seconds). */
+	float GraphRefreshTimer = 0.f;
+};
diff --git a/Engine/Plugins/PCG/Source/PCGEditor/Public/EditorMode/Tools/PCGInteractiveToolBuilder.h b/Engine/Plugins/PCG/Source/PCGEditor/Public/EditorMode/Tools/PCGInteractiveToolBuilder.h
--- a/Engine/Plugins/PCG/Source/PCGEditor/Public/EditorMode/Tools/PCGInteractiveToolBuilder.h
+++ b/Engine/Plugins/PCG/Source/PCGEditor/Public/EditorMode/Tools/PCGInteractiveToolBuilder.h
@@ -3,15 +3,9 @@
 #pragma once
 
 #include "InteractiveToolBuilder.h"
-#include "PCGGraph.h"
 
 #include "PCGInteractiveToolBuilder.generated.h"
 
-namespace UE::PCG::EditorMode::Tool
-{
-	bool BuildTool(UInteractiveTool* InTool);
-}
-
 UCLASS()
 class PCGEDITOR_API UPCGInteractiveToolBuilder : public UInteractiveToolBuilder
 {
@@ -21,29 +15,6 @@ public:
 	// ~Begin UInteractiveToolBuilder interface
 	virtual bool CanBuildTool(const FToolBuilderState& SceneState) const override;
 	virtual UInteractiveTool* BuildTool(const FToolBuilderState& SceneState) const override;
-	virtual void PostSetupTool(UInteractiveTool* Tool, const FToolBuilderState& SceneState) const override;
+	virtual void PostBuildTool(UInteractiveTool* Tool, const FToolBuilderState& SceneState) const override;
 	// ~End UInteractiveToolBuilder interface
-
-	void SetToolGraph(UPCGGraph* InGraph) { ToolGraph = InGraph; }
-	UPCGGraph* GetToolGraph() const { return ToolGraph; }
-private:
-	UPROPERTY()
-	TObjectPtr<UPCGGraph> ToolGraph;
 };
-
-UCLASS()
-class PCGEDITOR_API UPCGInteractiveToolWithToolTargetsBuilder : public UInteractiveToolWithToolTargetsBuilder
-{
-	GENERATED_BODY()
-
-public:
-	// ~Begin UInteractiveToolBuilder interface
-	virtual void PostSetupTool(UInteractiveTool* Tool, const FToolBuilderState& SceneState) const override;
-	// ~End UInteractiveToolBuilder interface
-
-	void SetToolGraph(UPCGGraph* InGraph) { ToolGraph = InGraph; }
-	UPCGGraph* GetToolGraph() const { return ToolGraph; }
-private:
-	UPROPERTY()
-	TObjectPtr<UPCGGraph> ToolGraph;
-};
\ No newline at end of file
diff --git a/Engine/Plugins/PCG/Source/PCGEditor/Public/EditorMode/Tools/PCGInteractiveToolCommon.h b/Engine/Plugins/PCG/Source/PCGEditor/Public/EditorMode/Tools/PCGInteractiveToolCommon.h
--- a/Engine/Plugins/PCG/Source/PCGEditor/Public/EditorMode/Tools/PCGInteractiveToolCommon.h
+++ b/Engine/Plugins/PCG/Source/PCGEditor/Public/EditorMode/Tools/PCGInteractiveToolCommon.h
@@ -6,6 +6,13 @@
 
 #include "PCGInteractiveToolCommon.generated.h"
 
+UENUM()
+enum class EPCGToolInteractionMode : uint8
+{
+	CreateNewActor,
+	AddToSelectedActor
+};
+
 UENUM()
 enum class EPCGToolDrawTargetOffset : uint8
 {
diff --git a/Engine/Plugins/PCG/Source/PCGEditor/Public/EditorMode/Tools/PCGInteractiveToolSettings.h b/Engine/Plugins/PCG/Source/PCGEditor/Public/EditorMode/Tools/PCGInteractiveToolSettings.h
--- a/Engine/Plugins/PCG/Source/PCGEditor/Public/EditorMode/Tools/PCGInteractiveToolSettings.h
+++ b/Engine/Plugins/PCG/Source/PCGEditor/Public/EditorMode/Tools/PCGInteractiveToolSettings.h
@@ -8,226 +8,59 @@
 
 #include "InteractiveTool.h"
 #include "NativeGameplayTags.h"
-#include "PCGComponent.h"
-#include "TickableEditorObject.h"
-#include "BaseBehaviors/BehaviorTargetInterfaces.h"
-#include "PCGInteractiveToolSettings.generated.h"
-
-struct FPCGInteractiveToolWorkingData;
 
-namespace UE::PCG::EditorMode::Tool
-{
-	PCGEDITOR_API void Shutdown(UInteractiveTool* InTool, EToolShutdownType ShutdownType);
-	
-	PCGEDITOR_API extern TAutoConsoleVariable<bool> CVarCommitPropertyChangedEventOnTimer;
-}
+#include "PCGInteractiveToolSettings.generated.h"
 
-UCLASS(Abstract, Blueprintable)
-class PCGEDITOR_API UPCGInteractiveToolSettings : public UInteractiveToolPropertySet, public FTickableEditorObject, public IKeyInputBehaviorTarget
+UCLASS(Abstract)
+class PCGEDITOR_API UPCGInteractiveToolSettings : public UInteractiveToolPropertySet
 {
 	GENERATED_BODY()
 
 public:
 	UPCGInteractiveToolSettings();
-	virtual ~UPCGInteractiveToolSettings() override {}
-
-	virtual bool Initialize(UInteractiveTool* InTool);
-	virtual void Apply(UInteractiveTool* OwningTool);
-	virtual void Cancel(UInteractiveTool* OwningTool);
-	virtual void Shutdown();
-
-	/** Initializes or finds the working data associated with the data instance identifier.
-	 * If NAME_None, the ToolTag is used as identifier. */
-	void InitializeWorkingData(FName DataInstanceIdentifier = NAME_None);
-	
-	UPCGGraphInterface* GetToolGraph() const { return ToolGraph; }
-	void SetToolGraph(UPCGGraphInterface* InGraph);
 
-	void SetDataInstance(FName DataInstanceIdentifier);
-
-	void SetActorClassToSpawn(TSubclassOf<AActor> Class);
+	void Initialize(UPCGGraphInstance* InGraphInstance)
+	{
+		GraphInstance = InGraphInstance;
+	}
 	
+	/** Tool tag for use with filtering compatible graphs. */
+	virtual FText GetToolTag() const PURE_VIRTUAL(UPCGInteractiveToolSettings::GetToolTag, return FText::GetEmpty(););
 	/** A soft object path to the graph asset to load by default for this tool. */
 	virtual FSoftObjectPath GetDefaultGraph() const;
-	
-	/** Tool tag for use with filtering compatible graphs. */
-	virtual FName GetToolTag() const PURE_VIRTUAL(UPCGInteractiveToolSettings::GetToolTag, return NAME_None;);
-
-	/** The working data type for this settings class. */
-	virtual const UScriptStruct* GetWorkingDataType() const PURE_VIRTUAL(UPCGInteractiveToolSettings::GetWorkingDataType, return FPCGInteractiveToolWorkingData::StaticStruct(););
-	
-	/** The actor to instantiate for edit-purposes when we aren't editing an existing actor. 
-	 * If your target actor has any special requirements, create an actor for that purpose and return it in your settings. */
-	virtual TSubclassOf<AActor> GetWorkingActorClass() const;
 
 	/** Filter details view graph selection. */
 	UFUNCTION()
 	virtual bool GraphAssetFilter(const FAssetData& AssetData) const;
-	
-	UPCGComponent* GetWorkingPCGComponent() const;
-	UPCGGraphInstance* GetGraphInstance() const;
-
-	UFUNCTION()
-	TArray<FName> GetValidDataInstanceNamesForGraph() const;
-	FName GetDefaultDataInstanceName() const;
-	
-	template<typename T = AActor>
-	T* GetTypedWorkingActor() const
-	{
-		static_assert(TIsDerivedFrom<T, AActor>::IsDerived, "T needs to be derived from AActor.");
-		
-		return Cast<T>(EditActor.Get());
-	}
-
-	FName GetWorkingDataIdentifier(FName DataInstanceIdentifier = NAME_None) const;
-	
-	template<typename T = FPCGInteractiveToolWorkingData>
-	T* GetMutableTypedWorkingData(FName DataInstanceIdentifier = NAME_None)
-	{
-		static_assert(TIsDerivedFrom<T, FPCGInteractiveToolWorkingData>::IsDerived, "T needs to be derived from FPCGInteractiveToolWorkingData.");
-
-		FName DataIdentifier = GetWorkingDataIdentifier(DataInstanceIdentifier);
-		
-		UPCGComponent* PCGComponent = GetWorkingPCGComponent();
-		if (PCGComponent->ToolDataContainer.ToolData.Contains(DataIdentifier))
-		{
-			return PCGComponent->ToolDataContainer.ToolData[DataIdentifier].GetMutablePtr<T>(); 
-		}
-
-		checkNoEntry();
-		return nullptr;
-	}
-	
-	template<typename T = FPCGInteractiveToolWorkingData>
-	const T* GetTypedWorkingData(FName DataInstanceIdentifier = NAME_None) const
-	{
-		static_assert(TIsDerivedFrom<T, FPCGInteractiveToolWorkingData>::IsDerived, "T needs to be derived from FPCGInteractiveToolWorkingData.");
-
-		FName WorkingDataIdentifier = GetWorkingDataIdentifier(DataInstanceIdentifier);
-
-		UPCGComponent* PCGComponent = GetWorkingPCGComponent();
-		if(PCGComponent->ToolDataContainer.ToolData.Contains(WorkingDataIdentifier))
-		{
-			return PCGComponent->ToolDataContainer.ToolData[WorkingDataIdentifier].GetPtr<T>(); 
-		}
-		
-		return nullptr;
-	}
-
-	template<typename T = FPCGInteractiveToolWorkingData>
-	TMap<FName, T*> GetMutableTypedWorkingDataMap() const
-	{
-		static_assert(TIsDerivedFrom<T, FPCGInteractiveToolWorkingData>::IsDerived, "T needs to be derived from FPCGInteractiveToolWorkingData.");
-
-		TMap<FName, T*> Result;
-
-		UPCGComponent* PCGComponent = GetWorkingPCGComponent();
-		
-		for(const FName& DataInstanceName : GetValidDataInstanceNamesForGraph())
-		{
-			FName WorkingDataIdentifier = GetWorkingDataIdentifier(DataInstanceName);
-			
-			if(PCGComponent->ToolDataContainer.ToolData.Contains(WorkingDataIdentifier))
-			{
-				Result.Add(WorkingDataIdentifier, PCGComponent->ToolDataContainer.ToolData[WorkingDataIdentifier].GetMutablePtr<T>()); 
-			}
-		}
-
-		return Result;
-	}
-
-	template<typename T = FPCGInteractiveToolWorkingData>
-	TMap<FName, const T*> GetTypedWorkingDataMap() const
-	{
-		static_assert(TIsDerivedFrom<T, FPCGInteractiveToolWorkingData>::IsDerived, "T needs to be derived from FPCGInteractiveToolWorkingData.");
-
-		TMap<FName, const T*> Result;
-
-		UPCGComponent* PCGComponent = GetWorkingPCGComponent();
-		
-		for(const FName& DataInstanceName : GetValidDataInstanceNamesForGraph())
-		{
-			FName WorkingDataIdentifier = GetWorkingDataIdentifier(DataInstanceName);
-			
-			if(PCGComponent->ToolDataContainer.ToolData.Contains(WorkingDataIdentifier))
-			{
-				Result.Add(WorkingDataIdentifier, PCGComponent->ToolDataContainer.ToolData[WorkingDataIdentifier].GetPtr<T>()); 
-			}
-		}
-
-		return Result;
-	}
 protected:
-	/** Try to load the default graph into ToolGraph. */
-	virtual void TryLoadDefaultGraph();
 	
-	virtual void Tick(float DeltaTime) override;
-	virtual bool IsTickable() const override;
-	virtual TStatId GetStatId() const override;
-
-	void CommitPropertyChangedEvents();
-
-	/** IKeyInputBehaviorTarget */
-	virtual void OnKeyPressed(const FKey& KeyID) override;
-
-	/** Allows to modify the just initialized working data from the editor. This can be used if the initialization of working data requires editor data access. */
-	virtual void OnWorkingDataInitialized(FPCGInteractiveToolWorkingData* WorkingData) const;
-	virtual void OnPropertyModified(UObject* Object, FProperty* Property);
-	virtual void OnRefreshInternal();
-private:
-	UPCGComponent* FindOrGeneratePCGComponent(AActor& OwningActor) const;
-	
-	void DeleteGeneratedResources() const;
-	
-	void OnRefresh();
-
-	virtual bool CanEditChange(const FProperty* InProperty) const override;
+	/** Try to load the default graph into TargetGraph. */
+	virtual void TryLoadDefaultGraph();
 public:
-	UPROPERTY(EditAnywhere, BlueprintReadOnly, Category = "PCG", meta = (NoResetToDefault, GetAssetFilter = "GraphAssetFilter"))
-	TObjectPtr<UPCGGraphInterface> ToolGraph = nullptr;
+	UPROPERTY(BlueprintReadOnly, EditAnywhere, Category = "PCG", meta = (NoResetToDefault, GetAssetFilter = "GraphAssetFilter"))
+	TObjectPtr<UPCGGraphInterface> TargetGraph = nullptr;
 
-	UPROPERTY(EditAnywhere, BlueprintReadOnly, Category = "PCG", meta = (NoResetToDefault, EditCondition="ToolGraph != nullptr", EditConditionHides, GetOptions="GetValidDataInstanceNamesForGraph"))
-	FName DataInstance = NAME_None;
+	UPROPERTY(Transient)
+	TObjectPtr<UPCGGraphInstance> GraphInstance;
 	
 	/** Automatically generate the PCG component while using the tool. */
 	UPROPERTY(EditAnywhere, Category = "PCG")
 	bool bAutoGenerate = true;
 
+	/** The tool can create a new actor or modify an existing one. */
+	UPROPERTY(EditAnywhere, Category = "PCG")
+	EPCGToolInteractionMode Interaction = EPCGToolInteractionMode::CreateNewActor;
+
 	// @todo_pcg: Convert to a target object, rather than actor, to ease transition into SceneGraph
 	/** The selected actor to modify with the tool. */
-	UPROPERTY(EditAnywhere, Category = "PCG", meta = (TransientToolProperty))
-	TWeakObjectPtr<AActor> EditActor = nullptr;
+	UPROPERTY(EditAnywhere, Category = "PCG", meta = (EditCondition = "Interaction == EPCGToolInteractionMode::AddToSelectedActor", EditConditionHides, TransientToolProperty))
+	TWeakObjectPtr<AActor> TargetActor = nullptr;
 
 	/** The name to be issued to a newly created actor when the tool is applied. */
-	UPROPERTY(EditAnywhere, DisplayName="Actor Label", Category = "PCG", meta = (TransientToolProperty, EditConditionHides))
+	UPROPERTY(EditAnywhere, DisplayName="Actor Label", Category = "PCG", meta = (EditCondition = "Interaction == EPCGToolInteractionMode::CreateNewActor", EditConditionHides, TransientToolProperty))
 	FName NewActorName;
 
 	/** The name to be issued to a newly created PCG component when the tool is applied. */
 	UPROPERTY(EditAnywhere, DisplayName="Component Name", Category = "PCG", meta = (TransientToolProperty))
 	FName NewPCGComponentName;
-
-	UPROPERTY(EditAnywhere, Category = "PCG", meta=(TransientToolProperty, EditConditionHides))
-	TSubclassOf<AActor> ActorClassToSpawn;
-
-	/** The current time left before a graph refresh (in seconds). */
-	UPROPERTY(meta = (TransientToolProperty))
-	float GraphRefreshTimer = 0.f;
-	
-	UPROPERTY(meta = (TransientToolProperty))
-	mutable TWeakObjectPtr<UPCGComponent> WorkingComponentCache;
-	
-	UPROPERTY(meta = (TransientToolProperty))
-	mutable bool bLockActorToSpawn = false;
-	
-	TMap<TWeakObjectPtr<UObject>, FPropertyChangedEvent> PropertyChangedEventQueue;
-
-	struct FGeneratedResources
-	{
-		TWeakObjectPtr<AActor> SpawnedActor;
-		TWeakObjectPtr<UPCGComponent> GeneratedPCGComponent;
-		/** A list of working data identifiers added this tool instantiation. Used to delete data on Cancel. */
-		TArray<FName> AddedWorkingDataIdentifiers;
-	};
-	
-	mutable FGeneratedResources GeneratedResources;
 };
diff --git a/Engine/Plugins/PCG/Source/PCGEditor/Public/EditorMode/Tools/Paint/PCGPaintTool.h b/Engine/Plugins/PCG/Source/PCGEditor/Public/EditorMode/Tools/Paint/PCGPaintTool.h
--- a/Engine/Plugins/PCG/Source/PCGEditor/Public/EditorMode/Tools/Paint/PCGPaintTool.h
+++ b/Engine/Plugins/PCG/Source/PCGEditor/Public/EditorMode/Tools/Paint/PCGPaintTool.h
@@ -3,79 +3,53 @@
 #pragma once
 
 #include "PCGEdModeBrushBase.h"
-#include "BaseBehaviors/KeyAsModifierInputBehavior.h"
-#include "BaseBehaviors/BehaviorTargetInterfaces.h"
-#include "Components/BoxComponent.h"
-#include "Data/PCGPointArrayData.h"
+#include "EditorMode/Tools/PCGInteractiveTool.h"
 #include "EditorMode/Tools/PCGInteractiveToolBuilder.h"
 #include "EditorMode/Tools/PCGInteractiveToolSettings.h"
 
 #include "PCGPaintTool.generated.h"
 
 UCLASS()
-class PCGEDITOR_API UPCGInteractiveToolSettings_PaintTool : public UPCGInteractiveToolSettings
+class PCGEDITOR_API UPCGPaintSettings : public UPCGInteractiveToolSettings
 {
 	GENERATED_BODY()
 
-public:
-	UPCGInteractiveToolSettings_PaintTool();
+	UPCGPaintSettings(const FObjectInitializer& ObjectInitializer);
 
 	// ~Begin UPCGInteractiveToolSettings interface
-	virtual FName GetToolTag() const override;
-	virtual const UScriptStruct* GetWorkingDataType() const override;
+	virtual FText GetToolTag() const override;
+	virtual FSoftObjectPath GetDefaultGraph() const override;
 	// ~End UPCGInteractiveToolSettings interface
 
-	/** If false, will erase points on all paint data instances. */
-	UPROPERTY(BlueprintReadOnly, EditAnywhere, Category = "Paint")
-	bool bErasePointsOfSelectedDataInstance = true;
+	bool ApplyTool(UWorld* World, FText& OutWarning) const;
+
+public:
+	UPROPERTY(EditAnywhere, Category = "Brush", Instanced, meta = (ShowOnlyInnerProperties))
+	TObjectPtr<UPCGEdModeBrushBase> Brush = nullptr;
 };
 
 UCLASS(Transient)
-class UPCGPaintTool : public UBaseBrushTool
+class UPCGPaintTool : public UPCGInteractiveTool
 {
 	GENERATED_BODY()
 
 public:
 	// ~Begin UPCGInteractiveToolSettings interface
 	virtual void Setup() override;
-	virtual void Shutdown(EToolShutdownType ShutdownType) override;
-	virtual bool HasAccept() const override;
+	virtual TSubclassOf<UPCGInteractiveToolSettings> GetSettingsClass() const override;
 	virtual bool CanAccept() const override;
-	virtual bool HasCancel() const override;
-	virtual void OnTick(float DeltaTime) override;
 	// ~End UPCGInteractiveToolSettings interface
 
 protected:
-	virtual void OnClickPress(const FInputDeviceRay& PressPos) override;
-	virtual void OnUpdateDrag(const FRay& Ray) override;
-	
-	virtual void SetupBrushStampIndicator() override;
-	virtual void UpdateBrushStampIndicator() override;
-	virtual void ShutdownBrushStampIndicator() override;
-	
-	/** This will be used as the brush base, if not specifying a radius directly */
-	virtual double EstimateMaximumTargetDimension() override;
-private:
-	void AddSinglePointFromBrush();
-	void TryDeletePointsFromBrush();
-private:
-	UPROPERTY(Transient)
-	TObjectPtr<UPCGInteractiveToolSettings_PaintTool> ToolSettings;
-
-	UPROPERTY(Transient)
-	bool bPaintOnTick = false;
-
-	UPROPERTY(Transient)
-	bool bShiftPressedLastTick = false;
+	virtual void ApplyToolInternal() const override {}
 };
 
 UCLASS(Transient)
-class UPCGPaintToolBuilder : public UPCGInteractiveToolWithToolTargetsBuilder
+class UPCGPaintToolBuilder : public UPCGInteractiveToolBuilder
 {
 	GENERATED_BODY()
 
 public:
 	virtual bool CanBuildTool(const FToolBuilderState& SceneState) const override;
 	virtual UInteractiveTool* BuildTool(const FToolBuilderState& SceneState) const override;
-	virtual const FToolTargetTypeRequirements& GetTargetRequirements() const override;
 };
diff --git a/Engine/Plugins/PCG/Source/PCGEditor/Public/EditorMode/Tools/Volume/PCGCreateVolumeTool.h b/Engine/Plugins/PCG/Source/PCGEditor/Public/EditorMode/Tools/Volume/PCGCreateVolumeTool.h
--- a/Engine/Plugins/PCG/Source/PCGEditor/Public/EditorMode/Tools/Volume/PCGCreateVolumeTool.h
+++ b/Engine/Plugins/PCG/Source/PCGEditor/Public/EditorMode/Tools/Volume/PCGCreateVolumeTool.h
@@ -2,6 +2,7 @@
 
 #pragma once
 
+#include "EditorMode/Tools/PCGInteractiveTool.h"
 #include "EditorMode/Tools/PCGInteractiveToolBuilder.h"
 #include "EditorMode/Tools/PCGInteractiveToolSettings.h"
 #include "EditorMode/Tools/Helpers/PCGEdModeSceneQueryHelpers.h"
@@ -10,14 +11,13 @@
 #include "PCGCreateVolumeTool.generated.h"
 
 UCLASS()
-class PCGEDITOR_API UPCGInteractiveToolSettings_CreateVolume : public UPCGInteractiveToolSettings
+class PCGEDITOR_API UPCGCreateVolumeToolSettings : public UPCGInteractiveToolSettings
 {
 	GENERATED_BODY()
 
 public:
 	// ~Begin UPCGInteractiveToolSettings interface
-	virtual FName GetToolTag() const override;
-	virtual const UScriptStruct* GetWorkingDataType() const override;
+	virtual FText GetToolTag() const override;
 	virtual FSoftObjectPath GetDefaultGraph() const override;
 	// ~End UPCGInteractiveToolSettings interface
 
@@ -41,7 +41,7 @@ public:
 };
 
 UCLASS(Transient, MinimalAPI)
-class UPCGCreateVolumeTool : public UInteractiveTool, public IClickBehaviorTarget
+class UPCGCreateVolumeTool : public UPCGInteractiveTool, public IClickBehaviorTarget
 {
 	GENERATED_BODY()
 
@@ -57,18 +57,19 @@ public:
 	virtual void Render(IToolsContextRenderAPI* RenderAPI) override;
 	// ~End UInteractiveTool interface
 
+	// ~Begin UPCGInteractiveTool interface
+	virtual TSubclassOf<UPCGInteractiveToolSettings> GetSettingsClass() const override;
+	// ~End UPCGInteractiveTool interface
+
 	// ~Begin IClickBehaviorTarget interface
 	virtual FInputRayHit IsHitByClick(const FInputDeviceRay& ClickPosition) override;
 	virtual void OnClicked(const FInputDeviceRay& ClickPosition) override;
 	// ~End IClickBehaviorTarget interface
 
 protected:
-	void ApplyTool() const;
+	virtual void ApplyToolInternal() const override {}
 private:
 	PCG::EditorMode::Scene::FViewRayParams CachedRayParams;
-
-	UPROPERTY(Transient)
-	TObjectPtr<UPCGInteractiveToolSettings_CreateVolume> ToolSettings;
 };
 
 UCLASS(Transient, MinimalAPI)
diff --git a/Engine/Plugins/Runtime/MeshModelingToolset/Source/ModelingComponents/Private/BaseGizmos/VolumetricBrushStampIndicator.cpp b/Engine/Plugins/Runtime/MeshModelingToolset/Source/ModelingComponents/Private/BaseGizmos/VolumetricBrushStampIndicator.cpp
--- a/Engine/Plugins/Runtime/MeshModelingToolset/Source/ModelingComponents/Private/BaseGizmos/VolumetricBrushStampIndicator.cpp
+++ b/Engine/Plugins/Runtime/MeshModelingToolset/Source/ModelingComponents/Private/BaseGizmos/VolumetricBrushStampIndicator.cpp
@@ -51,16 +51,6 @@ void UVolumetricBrushStampIndicator::MakeBrushIndicatorMesh(UInteractiveTool* Pa
 	}
 }
 
-UPreviewMesh* UVolumetricBrushStampIndicator::GetPreviewMesh() const
-{
-	if (SphereMesh)
-	{
-		return SphereMesh;
-	}
-
-	return nullptr;
-}
-
 void UVolumetricBrushStampIndicator::Shutdown()
 {
 	if (SphereMesh)
diff --git a/Engine/Plugins/Runtime/MeshModelingToolset/Source/ModelingComponents/Public/BaseGizmos/VolumetricBrushStampIndicator.h b/Engine/Plugins/Runtime/MeshModelingToolset/Source/ModelingComponents/Public/BaseGizmos/VolumetricBrushStampIndicator.h
--- a/Engine/Plugins/Runtime/MeshModelingToolset/Source/ModelingComponents/Public/BaseGizmos/VolumetricBrushStampIndicator.h
+++ b/Engine/Plugins/Runtime/MeshModelingToolset/Source/ModelingComponents/Public/BaseGizmos/VolumetricBrushStampIndicator.h
@@ -39,7 +39,6 @@ public:
 
 	MODELINGCOMPONENTS_API void MakeBrushIndicatorMesh(UInteractiveTool* ParentTool, UWorld* World);
 
-	MODELINGCOMPONENTS_API UPreviewMesh* GetPreviewMesh() const;
 private:
 	UPROPERTY()
 	TObjectPtr<UPreviewMesh> SphereMesh;
