[PCG] Fix typo in FPCGTextureBindingInfo::operator== (bughawk report)

#rnx

diff --git a/Engine/Plugins/PCG/Source/PCG/Private/Compute/DataInterfaces/PCGTextureDataInterface.cpp b/Engine/Plugins/PCG/Source/PCG/Private/Compute/DataInterfaces/PCGTextureDataInterface.cpp
--- a/Engine/Plugins/PCG/Source/PCG/Private/Compute/DataInterfaces/PCGTextureDataInterface.cpp
+++ b/Engine/Plugins/PCG/Source/PCG/Private/Compute/DataInterfaces/PCGTextureDataInterface.cpp
@@ -112,7 +112,7 @@ bool FPCGTextureBindingInfo::operator==(const FPCGTextureBindingInfo& Other) con
 		&& Size == Other.Size
 		&& Dimension == Other.Dimension
 		&& bPointSample == Other.bPointSample
-		&& Transform.Equals(Transform)
+		&& Transform.Equals(Other.Transform)
 		&& TextureBounds == Other.TextureBounds;
 }
 
