Additional fixes for PCG Editor Mode

diff --git a/Engine/Plugins/PCG/Source/PCG/Private/Editor/PCGGraphToolData.cpp b/Engine/Plugins/PCG/Source/PCG/Private/Editor/PCGGraphToolData.cpp
--- a/Engine/Plugins/PCG/Source/PCG/Private/Editor/PCGGraphToolData.cpp
+++ b/Engine/Plugins/PCG/Source/PCG/Private/Editor/PCGGraphToolData.cpp
@@ -55,6 +55,8 @@ void FPCGInteractiveToolWorkingDataGeneratedResources::Finalize()
 			GeneratedComponent->ClearFlags(RF_Transient);
 		}
 	}
+
+	GeneratedComponents.Empty();
 }
 
 void FPCGInteractiveToolWorkingDataGeneratedResources::Cleanup()
@@ -71,6 +73,8 @@ void FPCGInteractiveToolWorkingDataGeneratedResources::Cleanup()
 			GeneratedComponent->DestroyComponent();
 		}
 	}
+
+	GeneratedComponents.Empty();
 }
 
 FPCGInteractiveToolWorkingDataContext::FPCGInteractiveToolWorkingDataContext(UPCGComponent* InOwningComponent, UInteractiveTool* InInteractiveTool,	UInteractiveToolPropertySet* InPropertySet, FName InDataInstanceIdentifier)
@@ -193,6 +197,7 @@ USplineComponent* FPCGInteractiveToolWorkingData_Spline::FindOrGenerateSplineCom
 	if (MatchingSplineComponent == nullptr)
 	{
 		MatchingSplineComponent = GenerateMatchingSplineComponent(Context);
+		GeneratedResources.GeneratedComponents.Add(MatchingSplineComponent);
 	}	
 	
 	ensureMsgf(MatchingSplineComponent != nullptr, TEXT("Matching Spline Component should always be found after initialization as it was either found or created."));
@@ -251,9 +256,7 @@ USplineComponent* FPCGInteractiveToolWorkingData_Spline::GenerateMatchingSplineC
 	MatchingSplineComponent->ClearSplinePoints();
 	MatchingSplineComponent->bSplineHasBeenEdited = true;
 
-	// @todo_pcg
-	// When creating a new spline, we don't know if it should be closed or not
-	// We have access to the tool & property set; if we move them into runtime, we can read the spline tool's bClosedSpline setting
+	// We set it to open or closed using OnInitialized based on the user property
 	MatchingSplineComponent->SetClosedLoop(false);
 		
 	MatchingSplineComponent->RegisterComponent();
diff --git a/Engine/Plugins/PCG/Source/PCGEditor/Private/EditorMode/Tools/Helpers/PCGEdModeSplineHelpers.cpp b/Engine/Plugins/PCG/Source/PCGEditor/Private/EditorMode/Tools/Helpers/PCGEdModeSplineHelpers.cpp
--- a/Engine/Plugins/PCG/Source/PCGEditor/Private/EditorMode/Tools/Helpers/PCGEdModeSplineHelpers.cpp
+++ b/Engine/Plugins/PCG/Source/PCGEditor/Private/EditorMode/Tools/Helpers/PCGEdModeSplineHelpers.cpp
@@ -109,6 +109,6 @@ namespace UE::PCG::EditorMode::Spline
 		Destination.SetClosedLoop(Source.IsClosedLoop());
 		Destination.UpdateSpline();
 	}
-} // namespace PCG::EditorMode::Spline
+}
 
 #undef LOCTEXT_NAMESPACE
