Removed module API macro on default functions

diff --git a/Engine/Plugins/PCG/Source/PCG/Private/Editor/PCGGraphToolData.cpp b/Engine/Plugins/PCG/Source/PCG/Private/Editor/PCGGraphToolData.cpp
--- a/Engine/Plugins/PCG/Source/PCG/Private/Editor/PCGGraphToolData.cpp
+++ b/Engine/Plugins/PCG/Source/PCG/Private/Editor/PCGGraphToolData.cpp
@@ -76,7 +76,9 @@ void FPCGInteractiveToolWorkingDataGeneratedResources::Cleanup()
 
 	GeneratedComponents.Empty();
 }
+#endif
 
+#if WITH_EDITORONLY_DATA
 FPCGInteractiveToolWorkingDataContext::FPCGInteractiveToolWorkingDataContext(UPCGComponent* InOwningComponent, UInteractiveTool* InInteractiveTool,	UInteractiveToolPropertySet* InPropertySet, FName InDataInstanceIdentifier)
 {
 	OwningPCGComponent = InOwningComponent;
@@ -87,11 +89,20 @@ FPCGInteractiveToolWorkingDataContext::FPCGInteractiveToolWorkingDataContext(UPC
 }
 #endif
 
+bool FPCGInteractiveToolWorkingData::IsInitialized() const
+{
+	return bInitialized;
+}
+
 bool FPCGInteractiveToolWorkingData::IsValid() const
 {
 	return true;
 }
 
+void FPCGInteractiveToolWorkingData::InitializeRuntimeElementData(FPCGContext* InContext) const
+{
+}
+
 #if WITH_EDITORONLY_DATA
 void FPCGInteractiveToolWorkingData_PointArrayData::InitializeInternal(const FPCGInteractiveToolWorkingDataContext& Context)
 {
@@ -141,6 +152,11 @@ bool FPCGInteractiveToolWorkingData_PointArrayData::IsValid() const
 	return GeneratedPointData != nullptr;
 }
 
+UPCGPointArrayData* FPCGInteractiveToolWorkingData_PointArrayData::GetPointArrayData() const
+{
+	return GeneratedPointData;
+}
+
 USplineComponent* FPCGInteractiveToolWorkingData_Spline::GetSplineComponent() const
 {
 	return SplineComponent.Get();
diff --git a/Engine/Plugins/PCG/Source/PCG/Public/Editor/PCGGraphToolData.h b/Engine/Plugins/PCG/Source/PCG/Public/Editor/PCGGraphToolData.h
--- a/Engine/Plugins/PCG/Source/PCG/Public/Editor/PCGGraphToolData.h
+++ b/Engine/Plugins/PCG/Source/PCG/Public/Editor/PCGGraphToolData.h
@@ -21,7 +21,8 @@ struct FPCGInteractiveToolWorkingDataGeneratedResources
 {
 	GENERATED_BODY()
 
-	PCG_API FPCGInteractiveToolWorkingDataGeneratedResources() = default;
+	FPCGInteractiveToolWorkingDataGeneratedResources() = default;
+	virtual ~FPCGInteractiveToolWorkingDataGeneratedResources() = default;
 
 #if WITH_EDITORONLY_DATA
 	/** Finalizes generated resources (e.g. remove Transient flag). */
@@ -41,7 +42,8 @@ struct FPCGInteractiveToolWorkingDataContext
 {
 	GENERATED_BODY()
 
-	PCG_API FPCGInteractiveToolWorkingDataContext() = default;
+	FPCGInteractiveToolWorkingDataContext() = default;
+	virtual ~FPCGInteractiveToolWorkingDataContext() = default;
 
 #if WITH_EDITORONLY_DATA
 	PCG_API FPCGInteractiveToolWorkingDataContext(UPCGComponent* OwningComponent, UInteractiveTool* InteractiveTool, UInteractiveToolPropertySet* PropertySet, FName DataInstanceIdentifier);
@@ -75,17 +77,17 @@ struct FPCGInteractiveToolWorkingData
 {
 	GENERATED_BODY()
 
-	PCG_API FPCGInteractiveToolWorkingData() = default;
-	PCG_API virtual ~FPCGInteractiveToolWorkingData() = default;
+	FPCGInteractiveToolWorkingData() = default;
+	virtual ~FPCGInteractiveToolWorkingData() = default;
+
+	PCG_API bool IsInitialized() const;
 
-	PCG_API bool IsInitialized() const { return bInitialized; }
-	
 	PCG_API virtual bool IsValid() const;
 
 	/** Construct and add pcg data to the runtime context.
 	 * This will translate data from this struct into UPCGData to be processed by the graph.
 	 * Add dynamic tracking keys in here if you want to refresh when a change to your source data happens. */
-	PCG_API virtual void InitializeRuntimeElementData(FPCGContext* InContext) const {}
+	PCG_API virtual void InitializeRuntimeElementData(FPCGContext* InContext) const;
 	
 #if WITH_EDITORONLY_DATA
 public:
@@ -129,13 +131,13 @@ struct FPCGInteractiveToolWorkingData_PointArrayData : public FPCGInteractiveToo
 	GENERATED_BODY()
 	
 public:
-	PCG_API FPCGInteractiveToolWorkingData_PointArrayData() = default;
+	FPCGInteractiveToolWorkingData_PointArrayData() = default;
 
 	PCG_API virtual void InitializeRuntimeElementData(FPCGContext* InContext) const override;
 	
 	PCG_API virtual bool IsValid() const override;
 	
-	PCG_API UPCGPointArrayData* GetPointArrayData() const { return GeneratedPointData; }
+	PCG_API UPCGPointArrayData* GetPointArrayData() const;
 
 #if WITH_EDITORONLY_DATA
 	PCG_API virtual void InitializeInternal(const FPCGInteractiveToolWorkingDataContext& Context) override;
@@ -160,15 +162,16 @@ struct FPCGInteractiveToolWorkingData_Spline : public FPCGInteractiveToolWorking
 {
 	GENERATED_BODY()
 public:
-	PCG_API FPCGInteractiveToolWorkingData_Spline() = default;
-
+	FPCGInteractiveToolWorkingData_Spline() = default;
+	virtual ~FPCGInteractiveToolWorkingData_Spline() override = default;
+	
 	PCG_API virtual void InitializeRuntimeElementData(FPCGContext* InContext) const override;
 	
 	PCG_API virtual bool IsValid() const override;
 	
 	PCG_API USplineComponent* GetSplineComponent() const;
 
-#if WITH_EDITORONLY_DATA
+#if WITH_EDITOR
 public:
 	virtual void OnToolStart(const FPCGInteractiveToolWorkingDataContext& Context) override;
 	virtual void OnToolCancel(const FPCGInteractiveToolWorkingDataContext& Context) override;
@@ -196,7 +199,8 @@ struct FPCGInteractiveToolWorkingData_SplineSurface : public FPCGInteractiveTool
 {
 	GENERATED_BODY()
 	
-	PCG_API FPCGInteractiveToolWorkingData_SplineSurface() = default;
+	FPCGInteractiveToolWorkingData_SplineSurface() = default;
+	virtual ~FPCGInteractiveToolWorkingData_SplineSurface() override = default;
 
 #if WITH_EDITOR
 private:
@@ -212,8 +216,8 @@ struct FPCGGraphToolData
 {
 	GENERATED_BODY()
 
-	PCG_API FPCGGraphToolData() = default;
-	PCG_API ~FPCGGraphToolData() = default;
+	FPCGGraphToolData() = default;
+	virtual ~FPCGGraphToolData() = default;
 	
 	UPROPERTY(EditAnywhere, Category = Tool)
 	FText DisplayName;
@@ -230,8 +234,8 @@ struct FPCGInteractiveToolWorkingDataContainer
 {
 	GENERATED_BODY()
 
-	PCG_API FPCGInteractiveToolWorkingDataContainer() = default;
-	PCG_API ~FPCGInteractiveToolWorkingDataContainer() = default;
+	FPCGInteractiveToolWorkingDataContainer() = default;
+	virtual ~FPCGInteractiveToolWorkingDataContainer() = default;
 	
 	UPROPERTY(BlueprintReadOnly, VisibleAnywhere, Category = "Tool Data")
 	TMap<FName, TInstancedStruct<FPCGInteractiveToolWorkingData>> ToolData;
