[PCG] Fix GPU overrides failing on CPU nodes

Attribute overrides coming from GPU data would silently do nothing when applied to CPU node override pins. This is because the readback of GPU data takes place too late in the element execution flow, so the override data is not CPU visible when overrides are applied to the settings.

#rnx

diff --git a/Engine/Plugins/PCG/Source/PCG/Private/PCGData.cpp b/Engine/Plugins/PCG/Source/PCG/Private/PCGData.cpp
--- a/Engine/Plugins/PCG/Source/PCG/Private/PCGData.cpp
+++ b/Engine/Plugins/PCG/Source/PCG/Private/PCGData.cpp
@@ -8,6 +8,7 @@
 #include "PCGParamData.h"
 #include "PCGPin.h"
 #include "PCGSettings.h"
+#include "Compute/Data/PCGProxyForGPUData.h"
 #include "Data/PCGBasePointData.h"
 #include "Data/PCGSpatialData.h"
 #include "Data/PCGUnionData.h"
@@ -407,23 +408,26 @@ TArray<FPCGTaggedData> FPCGDataCollection::GetAllSettings() const
 
 TArray<FPCGTaggedData> FPCGDataCollection::GetAllParams() const
 {
-	return TaggedData.FilterByPredicate([](const FPCGTaggedData& Data) {
-		return Cast<UPCGParamData>(Data.Data) != nullptr;
+	return TaggedData.FilterByPredicate([](const FPCGTaggedData& Data)
+	{
+		return Cast<UPCGParamData>(Data.Data) || (Data.Data && Data.Data->GetUnderlyingDataType() == EPCGDataType::Param);
 	});
 }
 
 TArray<FPCGTaggedData> FPCGDataCollection::GetParamsByPin(const FName& InPinLabel) const
 {
-	return TaggedData.FilterByPredicate([&InPinLabel](const FPCGTaggedData& Data) {
-		return Data.Pin == InPinLabel && Cast<UPCGParamData>(Data.Data);
-		});
+	return TaggedData.FilterByPredicate([&InPinLabel](const FPCGTaggedData& Data)
+	{
+		return Data.Pin == InPinLabel && (Cast<UPCGParamData>(Data.Data) || (Data.Data && Data.Data->GetUnderlyingDataType() == EPCGDataType::Param));
+	});
 }
 
 TArray<FPCGTaggedData> FPCGDataCollection::GetTaggedParams(const FString& InTag) const
 {
-	return TaggedData.FilterByPredicate([&InTag](const FPCGTaggedData& Data) {
-		return Data.Tags.Contains(InTag) && Cast<UPCGParamData>(Data.Data) != nullptr;
-		});
+	return TaggedData.FilterByPredicate([&InTag](const FPCGTaggedData& Data)
+	{
+		return Data.Tags.Contains(InTag) && (Cast<UPCGParamData>(Data.Data) || (Data.Data && Data.Data->GetUnderlyingDataType() == EPCGDataType::Param));
+	});
 }
 
 UPCGParamData* FPCGDataCollection::GetParamsWithDeprecation(const UPCGNode* Node) const
diff --git a/Engine/Plugins/PCG/Source/PCG/Private/PCGElement.cpp b/Engine/Plugins/PCG/Source/PCG/Private/PCGElement.cpp
--- a/Engine/Plugins/PCG/Source/PCG/Private/PCGElement.cpp
+++ b/Engine/Plugins/PCG/Source/PCG/Private/PCGElement.cpp
@@ -222,7 +222,10 @@ bool IPCGElement::Execute(FPCGContext* Context) const
 				TRACE_CPUPROFILER_EVENT_SCOPE(EPCGExecutionPhase::NotExecuted);
 				PCG_ELEMENT_EXECUTION_BREAKPOINT();
 
-				PreExecute(Context);
+				if (!PreExecute(Context))
+				{
+					bExecutionPostponed = true;
+				}
 
 				break;
 			}
@@ -293,7 +296,7 @@ bool IPCGElement::Execute(FPCGContext* Context) const
 	return Context->CurrentPhase == EPCGExecutionPhase::Done;
 }
 
-void IPCGElement::PreExecute(FPCGContext* Context) const
+bool IPCGElement::PreExecute(FPCGContext* Context) const
 {
 	check(Context);
 	// Check for early outs (task cancelled + node disabled)
@@ -306,7 +309,7 @@ void IPCGElement::PreExecute(FPCGContext* Context) const
 		{
 			// Skip task completely
 			Context->CurrentPhase = EPCGExecutionPhase::Done;
-			return;
+			return true;
 		}
 	}
 
@@ -317,7 +320,7 @@ void IPCGElement::PreExecute(FPCGContext* Context) const
 
 	if (!SettingsInterface)
 	{
-		return;
+		return true;
 	}
 
 	if (!SettingsInterface->bEnabled)
@@ -325,22 +328,30 @@ void IPCGElement::PreExecute(FPCGContext* Context) const
 		//Pass-through - no execution
 		DisabledPassThroughData(Context);
 		Context->CurrentPhase = EPCGExecutionPhase::PostExecute;
-		return;
+		return true;
 	}
 
-	// Will override the settings if there is any override.
-	Context->OverrideSettings();
-
-	const UPCGSettings* Settings = Context->GetInputSettings<UPCGSettings>();
+	const UPCGSettings* OriginalSettings = Context->GetOriginalSettings<UPCGSettings>();
 
 	// If we were supposed to execute on GPU and end up here, then GPU compilation failed. Pass through.
-	if (Settings && Settings->ShouldExecuteOnGPU())
+	if (OriginalSettings && OriginalSettings->ShouldExecuteOnGPU())
 	{
 		DisabledPassThroughData(Context);
 		Context->CurrentPhase = EPCGExecutionPhase::PostExecute;
-		return;
+		return true;
+	}
+
+	if (!ReadbackGPUDataForOverrides(Context))
+	{
+		Context->CurrentPhase = EPCGExecutionPhase::NotExecuted;
+		return false;
 	}
 
+	// Will override the settings if there is any override.
+	Context->OverrideSettings();
+
+	const UPCGSettings* Settings = Context->GetInputSettings<UPCGSettings>();
+
 	if (CVarPCGAllowPerDataCaching.GetValueOnAnyThread())
 	{
 		// Default implementation when the entries in a primary loop can be processed independently, e.g. they can appear in the cache separately
@@ -352,6 +363,8 @@ void IPCGElement::PreExecute(FPCGContext* Context) const
 			PreExecutePrimaryLoopElement(Context, Settings);
 		}
 	}
+
+	return true;
 }
 
 void IPCGElement::PreExecutePrimaryLoopElement(FPCGContext* Context, const UPCGSettings* Settings) const
@@ -1214,5 +1227,86 @@ EPCGCachingStatus IPCGElement::RetrieveResultsFromCache(IPCGGraphCache* Cache, c
 	}
 }
 
+bool IPCGElement::ReadbackGPUDataForOverrides(FPCGContext* Context) const
+{
+	check(Context);
+	
+	if (Context->InputData.TaggedData.IsEmpty())
+	{
+		return true;
+	}
+
+	const UPCGSettings* OriginalSettings = Context->GetOriginalSettings<UPCGSettings>();
+
+	if (!OriginalSettings)
+	{
+		return true;
+	}
+
+	const TArray<FPCGSettingsOverridableParam>& OverridableParams = OriginalSettings->OverridableParams();
+
+	if (OverridableParams.IsEmpty())
+	{
+		return true;
+	}
+
+	bool bHasPendingReadbacks = false;
+
+	// If there are any proxies in the input data, request readback to CPU.
+	for (FPCGTaggedData& TaggedData : Context->InputData.TaggedData)
+	{
+		const UPCGProxyForGPUData* DataWithGPUSupport = Cast<UPCGProxyForGPUData>(TaggedData.Data);
+
+		if (DataWithGPUSupport && OverridableParams.ContainsByPredicate([&TaggedData](const FPCGSettingsOverridableParam& InOverridableParam) { return InOverridableParam.Label == TaggedData.Pin; }))
+		{
+#if WITH_EDITOR
+			if (Context->Node && Context->GetStack() && Context->ExecutionSource.IsValid())
+			{
+				Context->ExecutionSource->GetExecutionState().GetInspection().NotifyGPUToCPUReadback(Context->Node, Context->GetStack());
+			}
+#endif
+
+			// Poll until readback is done.
+			UPCGProxyForGPUData::FReadbackResult Result = DataWithGPUSupport->GetCPUData(Context);
+
+			if (Result.bComplete)
+			{
+				// Readback done, replace the proxy with the result data.
+				ensure(Result.TaggedData.Data);
+				TaggedData.Data = Result.TaggedData.Data;
+				TaggedData.Tags.Append(Result.TaggedData.Tags);
+
+				// Will update referenced data objects.
+				Context->bInputDataModified = true;
+			}
+			else
+			{
+				bHasPendingReadbacks = true;
+			}
+		}
+	}
+
+	if (bHasPendingReadbacks)
+	{
+		Context->bIsPaused = true;
+
+		// Not ready to execute and unlikely to be in the very short term, sleep until next frame.
+		ExecuteOnGameThread(UE_SOURCE_LOCATION, [ContextHandle = Context->GetOrCreateHandle()]()
+		{
+			if (TSharedPtr<FPCGContextHandle> SharedHandle = ContextHandle.Pin())
+			{
+				if (FPCGContext* ContextPtr = SharedHandle->GetContext())
+				{
+					ContextPtr->bIsPaused = false;
+				}
+			}
+		});
+
+		return false;
+	}
+
+	return true;
+}
+
 #undef PCG_ELEMENT_EXECUTION_BREAKPOINT
 #undef LOCTEXT_NAMESPACE
diff --git a/Engine/Plugins/PCG/Source/PCG/Public/PCGElement.h b/Engine/Plugins/PCG/Source/PCG/Public/PCGElement.h
--- a/Engine/Plugins/PCG/Source/PCG/Public/PCGElement.h
+++ b/Engine/Plugins/PCG/Source/PCG/Public/PCGElement.h
@@ -194,8 +194,8 @@ public:
 	UE_API bool ConvertInputsIfNeeded(FPCGContext* Context) const;
 
 protected:
-	/** This function will be called once and once only, at the beginning of an execution */
-	UE_API void PreExecute(FPCGContext* Context) const;
+	/** This function is called at the beginning of an execution. Will be called until it returns true. */
+	UE_API bool PreExecute(FPCGContext* Context) const;
 	/** The prepare data phase is one where it is more likely to be able to multithread */
 	UE_API bool PrepareData(FPCGContext* Context) const;
 	UE_API virtual bool PrepareDataInternal(FPCGContext* Context) const;
@@ -240,6 +240,9 @@ protected:
 
 private:
 	UE_API void CleanupAndValidateOutput(FPCGContext* Context) const;
+
+	/** Poll GPU overridden attributes for readback. Replaces the GPU data with a CPU param data upon successful readback. Returns false until readback is complete. */
+	bool ReadbackGPUDataForOverrides(FPCGContext* Context) const;
 };
 
 /**
