[PCG] Added different color for polygons for visual separation from spline/polyline

#rnx

diff --git a/Engine/Plugins/PCG/Source/PCGEditor/Private/PCGEditorModule.cpp b/Engine/Plugins/PCG/Source/PCGEditor/Private/PCGEditorModule.cpp
--- a/Engine/Plugins/PCG/Source/PCGEditor/Private/PCGEditorModule.cpp
+++ b/Engine/Plugins/PCG/Source/PCGEditor/Private/PCGEditorModule.cpp
@@ -873,6 +873,7 @@ void FPCGEditorModule::RegisterPinColorAndIcons()
 
 	InRegistry.RegisterPinColorFunction(FPCGDataTypeInfoPoint::AsId(), [](const FPCGDataTypeIdentifier&) { return GetDefault<UPCGEditorSettings>()->PointDataPinColor; });
 	InRegistry.RegisterPinColorFunction(FPCGDataTypeInfoConcrete::AsId(), [](const FPCGDataTypeIdentifier&) { return GetDefault<UPCGEditorSettings>()->ConcreteDataPinColor; });
+	InRegistry.RegisterPinColorFunction(FPCGDataTypeInfoPolygon2D::AsId(), [](const FPCGDataTypeIdentifier&) { return GetDefault<UPCGEditorSettings>()->Polygon2DDataPinColor; });
 	InRegistry.RegisterPinColorFunction(FPCGDataTypeInfoPolyline::AsId(), [](const FPCGDataTypeIdentifier&) { return GetDefault<UPCGEditorSettings>()->PolyLineDataPinColor; });
 	InRegistry.RegisterPinColorFunction(FPCGDataTypeInfoLandscape::AsId(), [](const FPCGDataTypeIdentifier&) { return GetDefault<UPCGEditorSettings>()->LandscapeDataPinColor; });
 	InRegistry.RegisterPinColorFunction(FPCGDataTypeInfoBaseTexture2D::AsId(), [](const FPCGDataTypeIdentifier&) { return GetDefault<UPCGEditorSettings>()->BaseTextureDataPinColor; });
diff --git a/Engine/Plugins/PCG/Source/PCGEditor/Private/PCGEditorSettings.cpp b/Engine/Plugins/PCG/Source/PCGEditor/Private/PCGEditorSettings.cpp
--- a/Engine/Plugins/PCG/Source/PCGEditor/Private/PCGEditorSettings.cpp
+++ b/Engine/Plugins/PCG/Source/PCGEditor/Private/PCGEditorSettings.cpp
@@ -43,6 +43,7 @@ UPCGEditorSettings::UPCGEditorSettings(const FObjectInitializer& ObjectInitializ
 	ConcreteDataPinColor = FLinearColor(0.45f, 0.38f, 0.96f);
 	PointDataPinColor = FLinearColor(0.05f, 0.25f, 1.0f);
 	PolyLineDataPinColor = FLinearColor(0.05f, 0.75f, 0.82f);
+	Polygon2DDataPinColor = FLinearColor(0.055f, 0.65f, 0.515f);
 	SurfaceDataPinColor = FLinearColor(0.06f, 0.55f, 0.21f);
 	LandscapeDataPinColor = FLinearColor(0.66f, 0.66f, 0.07f);
 	BaseTextureDataPinColor = FLinearColor(0.81f, 0.28f, 0.22f);
diff --git a/Engine/Plugins/PCG/Source/PCGEditor/Public/PCGEditorSettings.h b/Engine/Plugins/PCG/Source/PCGEditor/Public/PCGEditorSettings.h
--- a/Engine/Plugins/PCG/Source/PCGEditor/Public/PCGEditorSettings.h
+++ b/Engine/Plugins/PCG/Source/PCGEditor/Public/PCGEditorSettings.h
@@ -181,6 +181,10 @@ public:
 	UPROPERTY(EditAnywhere, config, Category = Node, meta = (HideAlphaChannel))
 	FLinearColor PolyLineDataPinColor;
 
+	/** Color used for data pins of type Polygon2D */
+	UPROPERTY(EditAnywhere, config, Category = Node, meta = (HideAlphaChannel))
+	FLinearColor Polygon2DDataPinColor;
+
 	/** Color used for data pins of type Surface */
 	UPROPERTY(EditAnywhere, config, Category = Node, meta = (HideAlphaChannel))
 	FLinearColor SurfaceDataPinColor;
