[PCG] Added missing pin color function unregister

#rnx

diff --git a/Engine/Plugins/PCG/Source/PCGEditor/Private/PCGEditorModule.cpp b/Engine/Plugins/PCG/Source/PCGEditor/Private/PCGEditorModule.cpp
--- a/Engine/Plugins/PCG/Source/PCGEditor/Private/PCGEditorModule.cpp
+++ b/Engine/Plugins/PCG/Source/PCGEditor/Private/PCGEditorModule.cpp
@@ -967,6 +967,7 @@ void FPCGEditorModule::UnregisterPinColorAndIcons()
 	InRegistry.UnregisterPinColorFunction(FPCGDataTypeInfoBaseTexture2D::AsId());
 	InRegistry.UnregisterPinColorFunction(FPCGDataTypeInfoLandscape::AsId());
 	InRegistry.UnregisterPinColorFunction(FPCGDataTypeInfoPolyline::AsId());
+	InRegistry.UnregisterPinColorFunction(FPCGDataTypeInfoPolygon2D::AsId());
 	InRegistry.UnregisterPinColorFunction(FPCGDataTypeInfoConcrete::AsId());
 	InRegistry.UnregisterPinColorFunction(FPCGDataTypeInfoPoint::AsId());
 	InRegistry.UnregisterPinColorFunction(FPCGDataTypeInfoParam::AsId());
