[PCG] Fixed type output on gather node with multiple input pins

#rnx

diff --git a/Engine/Plugins/PCG/Source/PCG/Private/Elements/PCGGather.cpp b/Engine/Plugins/PCG/Source/PCG/Private/Elements/PCGGather.cpp
--- a/Engine/Plugins/PCG/Source/PCG/Private/Elements/PCGGather.cpp
+++ b/Engine/Plugins/PCG/Source/PCG/Private/Elements/PCGGather.cpp
@@ -5,6 +5,8 @@
 #include "PCGContext.h"
 #include "PCGPin.h"
 
+#include "Algo/Transform.h"
+
 #include UE_INLINE_GENERATED_CPP_BY_NAME(PCGGather)
 
 #define LOCTEXT_NAMESPACE "PCGGatherElement"
@@ -46,6 +48,24 @@ FName UPCGGatherSettings::GetDynamicInputPinsBaseLabel() const
 	return PCGPinConstants::DefaultInputLabel;
 }
 
+FPCGDataTypeIdentifier UPCGGatherSettings::GetCurrentPinTypesID(const UPCGPin* Pin) const
+{
+	// Output pin depends on the union of all input pins.
+	if (Pin->IsOutputPin())
+	{
+		TArray<FPCGPinProperties> InputPins = InputPinProperties();
+		TArray<FName> InputPinLabels;
+		Algo::Transform(InputPins, InputPinLabels, [](const FPCGPinProperties& InPinProperty) { return InPinProperty.Label; });
+		
+		const FPCGDataTypeIdentifier InputTypeUnion = GetTypeUnionIDOfAllIncidentEdges(InputPinLabels);
+		return InputTypeUnion.IsValid() ? InputTypeUnion : FPCGDataTypeIdentifier{ EPCGDataType::Any };
+	}
+	else
+	{
+		return Super::GetCurrentPinTypesID(Pin);
+	}
+}
+
 #if WITH_EDITOR
 void UPCGGatherSettings::AddDefaultDynamicInputPin()
 {
diff --git a/Engine/Plugins/PCG/Source/PCG/Public/Elements/PCGGather.h b/Engine/Plugins/PCG/Source/PCG/Public/Elements/PCGGather.h
--- a/Engine/Plugins/PCG/Source/PCG/Public/Elements/PCGGather.h
+++ b/Engine/Plugins/PCG/Source/PCG/Public/Elements/PCGGather.h
@@ -24,19 +24,20 @@ public:
 #endif
 
 	virtual bool HasDynamicPins() const override { return true; }
+	virtual FPCGDataTypeIdentifier GetCurrentPinTypesID(const UPCGPin* Pin) const override;
 
 protected:
-	TArray<FPCGPinProperties> StaticInputPinProperties() const override;
+	virtual TArray<FPCGPinProperties> StaticInputPinProperties() const override;
 	virtual TArray<FPCGPinProperties> OutputPinProperties() const override;
 	virtual FPCGElementPtr CreateElement() const override;
 	//~End UPCGSettings interface
 
 	//~Begin PCGSettingsWithDynamicInputs interface
-	FName GetDynamicInputPinsBaseLabel() const override;
+	virtual FName GetDynamicInputPinsBaseLabel() const override;
 
 public:
 #if WITH_EDITOR
-	void AddDefaultDynamicInputPin() override;
+	virtual void AddDefaultDynamicInputPin() override;
 #endif // WITH_EDITOR
 	//~End PCGSettingsWithDynamicInputs interface
 };
