[PCG] Fixed issue with the default settings not loaded early enough, causing issues with new actor creation. Also fixed the default instance layer not being used in a few places, which led to cancellation not working as well.

#rnx

diff --git a/Engine/Plugins/PCG/Source/PCGEditor/Private/EditorMode/Tools/Helpers/PCGEdModeActorHelpers.cpp b/Engine/Plugins/PCG/Source/PCGEditor/Private/EditorMode/Tools/Helpers/PCGEdModeActorHelpers.cpp
--- a/Engine/Plugins/PCG/Source/PCGEditor/Private/EditorMode/Tools/Helpers/PCGEdModeActorHelpers.cpp
+++ b/Engine/Plugins/PCG/Source/PCGEditor/Private/EditorMode/Tools/Helpers/PCGEdModeActorHelpers.cpp
@@ -48,7 +48,7 @@ namespace UE::PCG::EditorMode
 	{
 		UPCGComponent* Find(const AActor* Owner)
 		{
-			if (!ensure(Owner))
+			if (!Owner)
 			{
 				return nullptr;
 			}
diff --git a/Engine/Plugins/PCG/Source/PCGEditor/Private/EditorMode/Tools/Line/PCGDrawSplineTool.cpp b/Engine/Plugins/PCG/Source/PCGEditor/Private/EditorMode/Tools/Line/PCGDrawSplineTool.cpp
--- a/Engine/Plugins/PCG/Source/PCGEditor/Private/EditorMode/Tools/Line/PCGDrawSplineTool.cpp
+++ b/Engine/Plugins/PCG/Source/PCGEditor/Private/EditorMode/Tools/Line/PCGDrawSplineTool.cpp
@@ -219,7 +219,7 @@ TArray<USplineComponent*> UPCGInteractiveToolSettings_SplineBase::GetWorkingSpli
 {
 	TArray<USplineComponent*> Result;
 
-	TArray<FName> AllWorkingDataInstanceNames = GetValidDataInstanceNamesForGraph();
+	TArray<FName> AllWorkingDataInstanceNames = GetDataInstanceNamesForGraph();
 	for(const FName& DataInstanceIdentifier : AllWorkingDataInstanceNames)
 	{
 		if(USplineComponent* WorkingSplineComponent = GetWorkingSplineComponent(DataInstanceIdentifier))
diff --git a/Engine/Plugins/PCG/Source/PCGEditor/Private/EditorMode/Tools/PCGInteractiveToolSettings.cpp b/Engine/Plugins/PCG/Source/PCGEditor/Private/EditorMode/Tools/PCGInteractiveToolSettings.cpp
--- a/Engine/Plugins/PCG/Source/PCGEditor/Private/EditorMode/Tools/PCGInteractiveToolSettings.cpp
+++ b/Engine/Plugins/PCG/Source/PCGEditor/Private/EditorMode/Tools/PCGInteractiveToolSettings.cpp
@@ -67,6 +67,9 @@ bool UPCGInteractiveToolSettings::Initialize(UInteractiveTool* InTool)
 	InTool->GetToolManager()->GetContextQueriesAPI()->GetCurrentSelectionState(State);
 	TArray<AActor*> SelectedActors = State.SelectedActors;
 
+	// In order to constrain component selection and actor creation, we need to make sure to setup the default settings first.
+	TryLoadDefaults();
+
 	// If we have selected one actor, that's the actor we want to edit
 	// This actor might:
 	// - Not have a PCG Component
@@ -98,17 +101,16 @@ bool UPCGInteractiveToolSettings::Initialize(UInteractiveTool* InTool)
 	}
 
 	OnModified.AddUObject(this, &UPCGInteractiveToolSettings::OnPropertyModified);
-	
-	// If we don't have a valid graph already, we try to load a default graph
-	if(PCGComponent->GetGraph() == nullptr)
-	{
-		TryLoadDefaultGraph();
-	}
-	// Otherwise, we use the graph we already have
-	else
+
+	// If we don't have a valid graph already, we try to assign the default graph to the newly created pcg component.
+	if (PCGComponent->GetGraph() == nullptr)
 	{
-		SetToolGraph(PCGComponent->GetGraph());
+		PCGComponent->SetGraph(ToolGraph);
 	}
+
+	// Finally, update the tool using the component graph (either the newly set one or a previously existing one).
+	// This will also update all tool data.
+	SetToolGraph(PCGComponent->GetGraph());
 	
 	// Keyboard setup for switching between data instance identifiers
 	UKeyInputBehavior* KeyInputBehavior = NewObject<UKeyInputBehavior>();
@@ -331,10 +333,14 @@ void UPCGInteractiveToolSettings::DeleteGeneratedResources() const
 	}
 }
 
-void UPCGInteractiveToolSettings::TryLoadDefaultGraph()
+void UPCGInteractiveToolSettings::TryLoadDefaults()
 {
 	check(!HasAnyFlags(RF_ClassDefaultObject));
-	SetToolGraph(CastChecked<UPCGGraphInterface>(GetDefaultGraph().TryLoad(), ECastCheckedType::NullAllowed));
+	if (UPCGGraphInterface* DefaultGraphInterface = CastChecked<UPCGGraphInterface>(GetDefaultGraph().TryLoad(), ECastCheckedType::NullAllowed))
+	{
+		SetToolGraph(DefaultGraphInterface);
+		SetActorClassToSpawn(DefaultGraphInterface->ToolData.InitialActorClassToSpawn);
+	}
 }
 
 void UPCGInteractiveToolSettings::Tick(float DeltaTime)
@@ -409,7 +415,7 @@ void UPCGInteractiveToolSettings::CommitPropertyChangedEvents()
 
 void UPCGInteractiveToolSettings::OnKeyPressed(const FKey& InKeyID)
 {
-	TArray<FName> ValidDataInstanceNames = GetValidDataInstanceNamesForGraph();
+	TArray<FName> ValidDataInstanceNames = GetDataInstanceNamesForGraph();
 
 	int32 Index = INDEX_NONE;
 	
@@ -470,7 +476,10 @@ void UPCGInteractiveToolSettings::OnPropertyModified(UObject* Object, FProperty*
 {
 	if(Object == this && Property->GetFName() == GET_MEMBER_NAME_CHECKED(UPCGInteractiveToolSettings, ToolGraph))
 	{
-		GetWorkingPCGComponent()->SetGraph(ToolGraph);
+		if (UPCGComponent* WorkingComponent = GetWorkingPCGComponent())
+		{
+			WorkingComponent->SetGraph(ToolGraph);
+		}
 
 		if(ToolGraph)
 		{
@@ -482,30 +491,12 @@ void UPCGInteractiveToolSettings::OnPropertyModified(UObject* Object, FProperty*
 		FName DefaultDataInstanceName = GetDefaultDataInstanceName();
 		SetDataInstance(DefaultDataInstanceName);
 		
-		// If we have a graph, we update the working data to the available data instances ones from the graph
-		if(ToolGraph)
-		{
-			// We update the working data
-			TArray<FName> ValidDataInstanceNames = GetValidDataInstanceNamesForGraph();
+		// Update the working data to the available data instances ones from the graph
+		TArray<FName> ValidDataInstanceNames = GetDataInstanceNamesForGraph();
 
-			if(ValidDataInstanceNames.Num() > 0)
-			{
-				for(const FName& DataInstanceIdentifier : ValidDataInstanceNames)
-				{
-					InitializeWorkingData(DataInstanceIdentifier);
-				}
-			}
-			// We make sure at least the default is available so that the tool data exists and can be used even if the graph doesn't make use of it
-			else
-			{
-				InitializeWorkingData(DefaultDataInstanceName);
-			}
-			
-		}
-		// If not, we use the default
-		else
+		for (const FName& DataInstanceIdentifier : ValidDataInstanceNames)
 		{
-			InitializeWorkingData(DefaultDataInstanceName);
+			InitializeWorkingData(DataInstanceIdentifier);
 		}
 	}
 	if(Object == this && Property->GetFName() == GET_MEMBER_NAME_CHECKED(UPCGInteractiveToolSettings, ActorClassToSpawn))
@@ -576,6 +567,18 @@ UPCGGraphInstance* UPCGInteractiveToolSettings::GetGraphInstance() const
 	return nullptr;
 }
 
+TArray<FName> UPCGInteractiveToolSettings::GetDataInstanceNamesForGraph() const
+{
+	TArray<FName> Result = GetValidDataInstanceNamesForGraph();
+
+	if (Result.IsEmpty())
+	{
+		Result.Add(GetDefaultDataInstanceName());
+	}
+
+	return Result;
+}
+
 TArray<FName> UPCGInteractiveToolSettings::GetValidDataInstanceNamesForGraph() const
 {
 	TArray<FName> Result;
@@ -596,11 +599,11 @@ TArray<FName> UPCGInteractiveToolSettings::GetValidDataInstanceNamesForGraph() c
 			{
 				continue;
 			}
-			
+
 			Result.AddUnique(DataFromTool->DataInstance);
 		}
 	}
-	
+
 	return Result;
 }
 
diff --git a/Engine/Plugins/PCG/Source/PCGEditor/Private/EditorMode/Tools/Paint/PCGPaintTool.cpp b/Engine/Plugins/PCG/Source/PCGEditor/Private/EditorMode/Tools/Paint/PCGPaintTool.cpp
--- a/Engine/Plugins/PCG/Source/PCGEditor/Private/EditorMode/Tools/Paint/PCGPaintTool.cpp
+++ b/Engine/Plugins/PCG/Source/PCGEditor/Private/EditorMode/Tools/Paint/PCGPaintTool.cpp
@@ -334,7 +334,7 @@ void UPCGPaintTool::TryDeletePointsFromBrush()
 	}
 	else
 	{
-		for(const FName& DataInstanceName : ToolSettings->GetValidDataInstanceNamesForGraph())
+		for(const FName& DataInstanceName : ToolSettings->GetDataInstanceNamesForGraph())
 		{
 			DeletePointsFromDataInstance(DataInstanceName);
 		}
diff --git a/Engine/Plugins/PCG/Source/PCGEditor/Public/EditorMode/Tools/PCGInteractiveToolSettings.h b/Engine/Plugins/PCG/Source/PCGEditor/Public/EditorMode/Tools/PCGInteractiveToolSettings.h
--- a/Engine/Plugins/PCG/Source/PCGEditor/Public/EditorMode/Tools/PCGInteractiveToolSettings.h
+++ b/Engine/Plugins/PCG/Source/PCGEditor/Public/EditorMode/Tools/PCGInteractiveToolSettings.h
@@ -67,6 +67,9 @@ public:
 	UPCGComponent* GetWorkingPCGComponent() const;
 	UPCGGraphInstance* GetGraphInstance() const;
 
+	UFUNCTION()
+	TArray<FName> GetDataInstanceNamesForGraph() const;
+
 	UFUNCTION()
 	TArray<FName> GetValidDataInstanceNamesForGraph() const;
 	FName GetDefaultDataInstanceName() const;
@@ -123,7 +126,7 @@ public:
 
 		UPCGComponent* PCGComponent = GetWorkingPCGComponent();
 		
-		for(const FName& DataInstanceName : GetValidDataInstanceNamesForGraph())
+		for(const FName& DataInstanceName : GetDataInstanceNamesForGraph())
 		{
 			FName WorkingDataIdentifier = GetWorkingDataIdentifier(DataInstanceName);
 			
@@ -145,7 +148,7 @@ public:
 
 		UPCGComponent* PCGComponent = GetWorkingPCGComponent();
 		
-		for(const FName& DataInstanceName : GetValidDataInstanceNamesForGraph())
+		for(const FName& DataInstanceName : GetDataInstanceNamesForGraph())
 		{
 			FName WorkingDataIdentifier = GetWorkingDataIdentifier(DataInstanceName);
 			
@@ -159,7 +162,7 @@ public:
 	}
 protected:
 	/** Try to load the default graph into ToolGraph. */
-	virtual void TryLoadDefaultGraph();
+	virtual void TryLoadDefaults();
 	
 	virtual void Tick(float DeltaTime) override;
 	virtual bool IsTickable() const override;
