[PCG] Editor tool mode robustness fixes

#rnx

diff --git a/Engine/Plugins/PCG/Source/PCGEditor/Private/EditorMode/Tools/PCGInteractiveToolSettings.cpp b/Engine/Plugins/PCG/Source/PCGEditor/Private/EditorMode/Tools/PCGInteractiveToolSettings.cpp
--- a/Engine/Plugins/PCG/Source/PCGEditor/Private/EditorMode/Tools/PCGInteractiveToolSettings.cpp
+++ b/Engine/Plugins/PCG/Source/PCGEditor/Private/EditorMode/Tools/PCGInteractiveToolSettings.cpp
@@ -148,7 +148,7 @@ void UPCGInteractiveToolSettings::Apply(UInteractiveTool* OwningTool)
 		FPCGInteractiveToolWorkingDataContext Context;
 		Context.OwningPCGComponent = GetWorkingPCGComponent();
 		Context.DataInstanceIdentifier = WorkingDataMap.Key;
-		Context.OwningActor = Context.OwningPCGComponent->GetOwner();
+		Context.OwningActor = (Context.OwningPCGComponent.IsValid() ? Context.OwningPCGComponent->GetOwner() : nullptr);
 		Context.InteractiveTool = OwningTool;
 		Context.PCGSettings = this;
 		
@@ -166,7 +166,7 @@ void UPCGInteractiveToolSettings::Cancel(UInteractiveTool* OwningTool)
 		FPCGInteractiveToolWorkingDataContext Context;
 		Context.OwningPCGComponent = GetWorkingPCGComponent();
 		Context.DataInstanceIdentifier = WorkingDataMap.Key;
-		Context.OwningActor = Context.OwningPCGComponent->GetOwner();
+		Context.OwningActor = (Context.OwningPCGComponent.IsValid() ? Context.OwningPCGComponent->GetOwner() : nullptr);
 		Context.InteractiveTool = OwningTool;
 		Context.PCGSettings = this;
 
@@ -207,13 +207,13 @@ void UPCGInteractiveToolSettings::InitializeWorkingData(FName DataInstanceIdenti
 	FPCGInteractiveToolWorkingDataContext Context;
 	Context.OwningPCGComponent = PCGComponent;
 	Context.DataInstanceIdentifier = DataInstanceIdentifier;
-	Context.OwningActor = Context.OwningPCGComponent->GetOwner();
+	Context.OwningActor = (Context.OwningPCGComponent.IsValid() ? Context.OwningPCGComponent->GetOwner() : nullptr);
 	Context.InteractiveTool = GetTypedOuter<UInteractiveTool>();
 	Context.PCGSettings = this;
 	
 	FName WorkingDataIdentifier = GetWorkingDataIdentifier(DataInstanceIdentifier);
 	
-	if (EditActor.IsValid())
+	if (EditActor.IsValid() && PCGComponent)
 	{
 		// If we don't find the requested working data, we create it
 		if(PCGComponent->ToolDataContainer.ToolData.Contains(WorkingDataIdentifier) == false)
@@ -257,7 +257,10 @@ void UPCGInteractiveToolSettings::InitializeWorkingData(FName DataInstanceIdenti
 		// If the working data is invalid despite being initialized, we delete and request it again.
 		// This can happen when data dependencies have been deleted outside of PCG, such as deleting a spline component on the actor
 		WorkingData->OnToolCancel(Context);
-		PCGComponent->ToolDataContainer.ToolData.Remove(WorkingDataIdentifier);
+		if (PCGComponent)
+		{
+			PCGComponent->ToolDataContainer.ToolData.Remove(WorkingDataIdentifier);
+		}
 
 		// In which case, requesting new working data will recreate a new spline
 		InitializeWorkingData(DataInstanceIdentifier);
diff --git a/Engine/Plugins/PCG/Source/PCGEditor/Public/EditorMode/Tools/PCGInteractiveToolSettings.h b/Engine/Plugins/PCG/Source/PCGEditor/Public/EditorMode/Tools/PCGInteractiveToolSettings.h
--- a/Engine/Plugins/PCG/Source/PCGEditor/Public/EditorMode/Tools/PCGInteractiveToolSettings.h
+++ b/Engine/Plugins/PCG/Source/PCGEditor/Public/EditorMode/Tools/PCGInteractiveToolSettings.h
@@ -92,7 +92,7 @@ public:
 		FName DataIdentifier = GetWorkingDataIdentifier(DataInstanceIdentifier);
 		
 		UPCGComponent* PCGComponent = GetWorkingPCGComponent();
-		if (PCGComponent->ToolDataContainer.ToolData.Contains(DataIdentifier))
+		if (PCGComponent && PCGComponent->ToolDataContainer.ToolData.Contains(DataIdentifier))
 		{
 			return PCGComponent->ToolDataContainer.ToolData[DataIdentifier].GetMutablePtr<T>(); 
 		}
@@ -109,7 +109,7 @@ public:
 		FName WorkingDataIdentifier = GetWorkingDataIdentifier(DataInstanceIdentifier);
 
 		UPCGComponent* PCGComponent = GetWorkingPCGComponent();
-		if(PCGComponent->ToolDataContainer.ToolData.Contains(WorkingDataIdentifier))
+		if(PCGComponent && PCGComponent->ToolDataContainer.ToolData.Contains(WorkingDataIdentifier))
 		{
 			return PCGComponent->ToolDataContainer.ToolData[WorkingDataIdentifier].GetPtr<T>(); 
 		}
@@ -124,15 +124,16 @@ public:
 
 		TMap<FName, T*> Result;
 
-		UPCGComponent* PCGComponent = GetWorkingPCGComponent();
-		
-		for(const FName& DataInstanceName : GetDataInstanceNamesForGraph())
+		if (UPCGComponent* PCGComponent = GetWorkingPCGComponent())
 		{
-			FName WorkingDataIdentifier = GetWorkingDataIdentifier(DataInstanceName);
-			
-			if(PCGComponent->ToolDataContainer.ToolData.Contains(WorkingDataIdentifier))
+			for(const FName& DataInstanceName : GetDataInstanceNamesForGraph())
 			{
-				Result.Add(WorkingDataIdentifier, PCGComponent->ToolDataContainer.ToolData[WorkingDataIdentifier].GetMutablePtr<T>()); 
+				FName WorkingDataIdentifier = GetWorkingDataIdentifier(DataInstanceName);
+			
+				if(PCGComponent->ToolDataContainer.ToolData.Contains(WorkingDataIdentifier))
+				{
+					Result.Add(WorkingDataIdentifier, PCGComponent->ToolDataContainer.ToolData[WorkingDataIdentifier].GetMutablePtr<T>()); 
+				}
 			}
 		}
 
@@ -146,15 +147,16 @@ public:
 
 		TMap<FName, const T*> Result;
 
-		UPCGComponent* PCGComponent = GetWorkingPCGComponent();
-		
-		for(const FName& DataInstanceName : GetDataInstanceNamesForGraph())
+		if (UPCGComponent* PCGComponent = GetWorkingPCGComponent())
 		{
-			FName WorkingDataIdentifier = GetWorkingDataIdentifier(DataInstanceName);
-			
-			if(PCGComponent->ToolDataContainer.ToolData.Contains(WorkingDataIdentifier))
+			for(const FName& DataInstanceName : GetDataInstanceNamesForGraph())
 			{
-				Result.Add(WorkingDataIdentifier, PCGComponent->ToolDataContainer.ToolData[WorkingDataIdentifier].GetPtr<T>()); 
+				FName WorkingDataIdentifier = GetWorkingDataIdentifier(DataInstanceName);
+			
+				if(PCGComponent->ToolDataContainer.ToolData.Contains(WorkingDataIdentifier))
+				{
+					Result.Add(WorkingDataIdentifier, PCGComponent->ToolDataContainer.ToolData[WorkingDataIdentifier].GetPtr<T>()); 
+				}
 			}
 		}
 
