[PCG] Moved the tool data to the data folder for clarity, and defaulted the actor class on the graph tool data.

#rnx

diff --git a/Engine/Plugins/PCG/Source/PCG/Private/Data/PCGToolData.cpp b/Engine/Plugins/PCG/Source/PCG/Private/Data/PCGToolData.cpp
new file mode 100644
--- /dev/null
+++ b/Engine/Plugins/PCG/Source/PCG/Private/Data/PCGToolData.cpp
@@ -0,0 +1,385 @@
+﻿// Copyright Epic Games, Inc. All Rights Reserved.
+
+#include "Data/PCGToolData.h"
+
+#include "PCGComponent.h"
+#include "PCGContext.h"
+#include "Data/PCGPointArrayData.h"
+#include "Helpers/PCGDynamicTrackingHelpers.h"
+
+#include "Engine/Engine.h"
+#include "Components/SplineComponent.h"
+#include "Components/BoxComponent.h"
+#include "Logging/StructuredLog.h"
+#include "UObject/UObjectGlobals.h"
+#include "UObject/Package.h"
+#if WITH_EDITOR
+#include "InteractiveTool.h"
+#endif
+
+#if WITH_EDITORONLY_DATA
+void FPCGInteractiveToolWorkingData::Initialize(const FPCGInteractiveToolWorkingDataContext& Context)
+{
+	ensureMsgf(Context.OwningPCGComponent.IsValid() && Context.OwningActor.IsValid(), TEXT("Owning PCG Component or Owning Actor invalid! This should never happen so an assumption was hurt."));
+
+	InitializeInternal(Context);
+	
+	OnInitializedDelegate.ExecuteIfBound(this);
+
+	bInitialized = true;
+}
+
+void FPCGInteractiveToolWorkingData::InitializeInternal(const FPCGInteractiveToolWorkingDataContext& Context)
+{
+}
+
+void FPCGInteractiveToolWorkingData::OnToolStart(const FPCGInteractiveToolWorkingDataContext& Context)
+{
+}
+
+void FPCGInteractiveToolWorkingData::OnToolApply(const FPCGInteractiveToolWorkingDataContext& Context)
+{
+	GeneratedResources.Finalize();
+}
+
+void FPCGInteractiveToolWorkingData::OnToolCancel(const FPCGInteractiveToolWorkingDataContext& Context)
+{
+	GeneratedResources.Cleanup();
+}
+
+void FPCGInteractiveToolWorkingDataGeneratedResources::Finalize()
+{
+	for(TWeakObjectPtr<UActorComponent> GeneratedComponent : GeneratedComponents)
+	{
+		if(GeneratedComponent.IsValid())
+		{
+			GeneratedComponent->ClearFlags(RF_Transient);
+		}
+	}
+
+	GeneratedComponents.Empty();
+}
+
+void FPCGInteractiveToolWorkingDataGeneratedResources::Cleanup()
+{
+	for(TWeakObjectPtr<UActorComponent> GeneratedComponent : GeneratedComponents)
+	{
+		if(GeneratedComponent.IsValid())
+		{
+			if(UPCGComponent* PCGComponent = Cast<UPCGComponent>(GeneratedComponent))
+			{
+				PCGComponent->Cleanup(true);	
+			}
+			
+			GeneratedComponent->DestroyComponent();
+		}
+	}
+
+	GeneratedComponents.Empty();
+}
+#endif
+
+#if WITH_EDITORONLY_DATA
+FPCGInteractiveToolWorkingDataContext::FPCGInteractiveToolWorkingDataContext(UPCGComponent* InOwningComponent, UInteractiveTool* InInteractiveTool,	UInteractiveToolPropertySet* InPropertySet, FName InDataInstanceIdentifier)
+{
+	OwningPCGComponent = InOwningComponent;
+	OwningActor = OwningPCGComponent->GetOwner();
+	DataInstanceIdentifier = InDataInstanceIdentifier;
+	InteractiveTool = InInteractiveTool;
+	PCGSettings = InPropertySet;
+}
+#endif
+
+bool FPCGInteractiveToolWorkingData::IsInitialized() const
+{
+	return bInitialized;
+}
+
+bool FPCGInteractiveToolWorkingData::IsValid() const
+{
+	return true;
+}
+
+void FPCGInteractiveToolWorkingData::InitializeRuntimeElementData(FPCGContext* InContext) const
+{
+}
+
+#if WITH_EDITORONLY_DATA
+void FPCGInteractiveToolWorkingData_PointArrayData::InitializeInternal(const FPCGInteractiveToolWorkingDataContext& Context)
+{
+	GeneratedPointData = NewObject<UPCGPointArrayData>(Context.OwningPCGComponent.Get());
+}
+
+void FPCGInteractiveToolWorkingData_PointArrayData::OnToolStart(const FPCGInteractiveToolWorkingDataContext& Context)
+{
+	OnToolStartPointArray = Cast<UPCGPointArrayData>(StaticDuplicateObject(GeneratedPointData, GetTransientPackage()));
+	
+	FPCGInteractiveToolWorkingData::OnToolStart(Context);
+}
+
+void FPCGInteractiveToolWorkingData_PointArrayData::OnToolCancel(const FPCGInteractiveToolWorkingDataContext& Context)
+{
+	if(OnToolStartPointArray != nullptr)
+	{
+		UEngine::CopyPropertiesForUnrelatedObjects(OnToolStartPointArray, GeneratedPointData);
+	}
+	
+	FPCGInteractiveToolWorkingData::OnToolCancel(Context);
+}
+#endif
+
+void FPCGInteractiveToolWorkingData_PointArrayData::InitializeRuntimeElementData(FPCGContext* Context) const
+{
+	FPCGInteractiveToolWorkingData::InitializeRuntimeElementData(Context);
+
+	TArray<FPCGTaggedData>& Outputs = Context->OutputData.TaggedData;
+
+	FPCGTaggedData& Output = Outputs.Emplace_GetRef();
+	// We duplicate it so that PCG can take 'ownership'.
+	// If not, it can't process references and caches properly
+	// Might be relying on the fact it shouldn't access the same object over multiple generations
+	// And if we recreate it every time it gets properly managed
+	// If it gets dropped from cache, it's expected to be invalid
+	Output.Data = GeneratedPointData->DuplicateData(Context);
+	Output.Pin = PCGPinConstants::DefaultOutputLabel;
+#if WITH_EDITOR
+	FPCGSelectionKey Key = FPCGSelectionKey::CreateFromObjectPtr(GeneratedPointData);
+	FPCGDynamicTrackingHelper::AddSingleDynamicTrackingKey(Context, MoveTemp(Key), false);
+#endif
+}
+
+bool FPCGInteractiveToolWorkingData_PointArrayData::IsValid() const
+{
+	return GeneratedPointData != nullptr;
+}
+
+UPCGPointArrayData* FPCGInteractiveToolWorkingData_PointArrayData::GetPointArrayData() const
+{
+	return GeneratedPointData;
+}
+
+USplineComponent* FPCGInteractiveToolWorkingData_Spline::GetSplineComponent() const
+{
+	return SplineComponent.Get();
+}
+
+#if WITH_EDITORONLY_DATA
+void FPCGInteractiveToolWorkingData_Spline::InitializeInternal(const FPCGInteractiveToolWorkingDataContext& Context)
+{	
+	SplineComponent = FindOrGenerateSplineComponent(Context);
+}
+
+void FPCGInteractiveToolWorkingData_Spline::OnToolStart(const FPCGInteractiveToolWorkingDataContext& Context)
+{
+	FPCGInteractiveToolWorkingData::OnToolStart(Context);
+
+	if(SplineComponent.IsValid())
+	{
+		OnToolStartSplineComponent = Cast<USplineComponent>(StaticDuplicateObject(SplineComponent.Get(), GetTransientPackage()));
+
+		// We wipe the spline points on tool start so it becomes easy to edit them after the fact.
+		// @todo_pcg
+		// When the Spline Tool exposes more controls, such as selecting spline points to move them, resetting the spline etc., we will remove this.
+		SplineComponent->bSplineHasBeenEdited = true;
+		SplineComponent->ClearSplinePoints(true);
+	}
+}
+
+void FPCGInteractiveToolWorkingData_Spline::OnToolCancel(const FPCGInteractiveToolWorkingDataContext& Context)
+{
+	// If the Spline Component is valid and isn't a generated component, we revert its state.
+	if(SplineComponent.IsValid() && GeneratedResources.GeneratedComponents.Contains(SplineComponent) == false)
+	{
+		if(OnToolStartSplineComponent)
+		{
+			// We restore the spline component to its previous state
+			UEngine::CopyPropertiesForUnrelatedObjects(OnToolStartSplineComponent, SplineComponent.Get());
+			
+			SplineComponent->UpdateSpline();
+			
+			OnToolStartSplineComponent->MarkAsGarbage();
+			OnToolStartSplineComponent = nullptr;
+		}
+	}
+
+	// This will delete the spline component if it was generated
+	FPCGInteractiveToolWorkingData::OnToolCancel(Context);
+}
+
+USplineComponent* FPCGInteractiveToolWorkingData_Spline::FindOrGenerateSplineComponent(const FPCGInteractiveToolWorkingDataContext& Context)
+{
+	USplineComponent* MatchingSplineComponent = FindMatchingSplineComponent(Context);
+
+	// If we didn't find any spline component to match with, we create it now
+	if (MatchingSplineComponent == nullptr)
+	{
+		MatchingSplineComponent = GenerateMatchingSplineComponent(Context);
+		GeneratedResources.GeneratedComponents.Add(MatchingSplineComponent);
+	}	
+	
+	ensureMsgf(MatchingSplineComponent != nullptr, TEXT("Matching Spline Component should always be found after initialization as it was either found or created."));
+	return MatchingSplineComponent;
+}
+
+USplineComponent* FPCGInteractiveToolWorkingData_Spline::FindMatchingSplineComponent(const FPCGInteractiveToolWorkingDataContext& Context)
+{
+	TArray<USplineComponent*> AllSplineComponents;
+	Context.OwningActor->GetComponents<USplineComponent>(AllSplineComponents, false);
+	
+	USplineComponent* MatchingSplineComponent = nullptr;
+	if (Context.DataInstanceIdentifier.IsNone())
+	{
+		if(AllSplineComponents.Num() > 0)
+		{			
+			MatchingSplineComponent = AllSplineComponents[0];
+		}
+	}
+	else
+	{
+		// When the DataInstanceIdentifier is set, we only care about the splines with the given tag
+		AllSplineComponents.RemoveAll([DataInstanceIdentifier = Context.DataInstanceIdentifier](USplineComponent* Candidate)
+		{
+			return Candidate->ComponentHasTag(DataInstanceIdentifier) == false;
+		});
+		
+		// If we find a matching spline component with the specific tag, use it
+		if (AllSplineComponents.Num() > 0)
+		{
+			MatchingSplineComponent = Cast<USplineComponent>(AllSplineComponents[0]);
+		}
+	}
+
+	if(AllSplineComponents.Num() > 1 && MatchingSplineComponent != nullptr)
+	{
+		UE_LOGFMT(LogPCG, Warning, "More than one spline component found. Choosing the first available: {0}", MatchingSplineComponent->GetName());
+	}
+	
+	return MatchingSplineComponent;
+}
+
+USplineComponent* FPCGInteractiveToolWorkingData_Spline::GenerateMatchingSplineComponent(const FPCGInteractiveToolWorkingDataContext& Context)
+{
+	USplineComponent* MatchingSplineComponent = NewObject<USplineComponent>(Context.OwningActor.Get());
+	MatchingSplineComponent->SetMobility(Context.OwningActor->GetRootComponent()->GetMobility());
+	MatchingSplineComponent->SetupAttachment(Context.OwningActor->GetRootComponent());
+	// Transient until Applied. 
+	MatchingSplineComponent->SetFlags(RF_Transient);
+
+	if (Context.DataInstanceIdentifier.IsNone() == false)
+	{
+		MatchingSplineComponent->ComponentTags.Add(Context.DataInstanceIdentifier);
+	}
+
+	MatchingSplineComponent->ClearSplinePoints();
+	MatchingSplineComponent->bSplineHasBeenEdited = true;
+
+	// We set it to open or closed using OnInitialized based on the user property
+	MatchingSplineComponent->SetClosedLoop(false);
+		
+	MatchingSplineComponent->RegisterComponent();
+	Context.OwningActor->AddInstanceComponent(MatchingSplineComponent);
+
+	return MatchingSplineComponent;
+}
+
+USplineComponent* FPCGInteractiveToolWorkingData_SplineSurface::FindMatchingSplineComponent(const FPCGInteractiveToolWorkingDataContext& Context)
+{
+	TArray<USplineComponent*> AllClosedSplineComponents;
+	Context.OwningActor->GetComponents<USplineComponent>(AllClosedSplineComponents, false);
+
+	// We remove all spline components that are open loops since this working data struct is only for spline surfaces
+	AllClosedSplineComponents.RemoveAll([](USplineComponent* Candidate)
+	{
+		return Candidate->IsClosedLoop() == false;
+	});
+	
+	USplineComponent* MatchingSplineComponent = nullptr;
+	if (Context.DataInstanceIdentifier.IsNone())
+	{
+		if(AllClosedSplineComponents.Num() > 0)
+		{
+			MatchingSplineComponent = AllClosedSplineComponents[0];
+		}
+	}
+	else
+	{
+		// When the DataInstanceIdentifier is set, we only care about the closed splines with the given tag
+		AllClosedSplineComponents.RemoveAll([DataInstanceIdentifier = Context.DataInstanceIdentifier](USplineComponent* Candidate)
+		{
+			return Candidate->ComponentHasTag(DataInstanceIdentifier) == false;
+		});
+		
+		// If we find a matching spline component with the specific tag, use i
+		if (AllClosedSplineComponents.Num() > 0)
+		{
+			MatchingSplineComponent = Cast<USplineComponent>(AllClosedSplineComponents[0]);
+		}
+	}
+
+	if(AllClosedSplineComponents.Num() > 1 && MatchingSplineComponent != nullptr)
+	{
+		UE_LOGFMT(LogPCG, Warning, "More than one spline component found. Choosing the first available: {0}", MatchingSplineComponent->GetName());
+	}
+
+	return MatchingSplineComponent;
+}
+
+USplineComponent* FPCGInteractiveToolWorkingData_SplineSurface::GenerateMatchingSplineComponent(const FPCGInteractiveToolWorkingDataContext& Context)
+{
+	USplineComponent* MatchingSplineComponent = NewObject<USplineComponent>(Context.OwningActor.Get());
+	MatchingSplineComponent->SetMobility(Context.OwningActor->GetRootComponent()->GetMobility());
+	MatchingSplineComponent->SetupAttachment(Context.OwningActor->GetRootComponent());
+	// Transient until Applied. 
+	MatchingSplineComponent->SetFlags(RF_Transient);
+
+	if (Context.DataInstanceIdentifier.IsNone() == false)
+	{
+		MatchingSplineComponent->ComponentTags.Add(Context.DataInstanceIdentifier);
+	}
+
+	MatchingSplineComponent->ClearSplinePoints();
+	MatchingSplineComponent->bSplineHasBeenEdited = true;
+	MatchingSplineComponent->SetClosedLoop(true);
+		
+	MatchingSplineComponent->RegisterComponent();
+	Context.OwningActor->AddInstanceComponent(MatchingSplineComponent);
+
+	return MatchingSplineComponent;
+}
+#endif
+
+void FPCGInteractiveToolWorkingData_Spline::InitializeRuntimeElementData(FPCGContext* Context) const
+{
+	FPCGInteractiveToolWorkingData::InitializeRuntimeElementData(Context);
+
+	TArray<FPCGTaggedData>& Outputs = Context->OutputData.TaggedData;
+
+	if(SplineComponent.IsValid())
+	{
+		FPCGTaggedData& Output = Outputs.Emplace_GetRef();
+		// We duplicate it so that PCG can take 'ownership'.
+		// If not, it can't process references and caches properly
+		// Might be relying on the fact it shouldn't access the same object over multiple generations
+		// And if we recreate it every time it gets properly managed
+		// If it gets dropped from cache, it's expected to be invalid
+		UPCGSplineData* SplineData = FPCGContext::NewObject_AnyThread<UPCGSplineData>(Context);
+		SplineData->Initialize(SplineComponent.Get());
+		Output.Data = SplineData;
+		Output.Pin = PCGPinConstants::DefaultOutputLabel;
+#if WITH_EDITOR
+		FPCGSelectionKey Key = FPCGSelectionKey::CreateFromObjectPtr(SplineComponent.Get());
+		FPCGDynamicTrackingHelper::AddSingleDynamicTrackingKey(Context, MoveTemp(Key), false);
+#endif
+	}
+}
+
+bool FPCGInteractiveToolWorkingData_Spline::IsValid() const
+{
+	return SplineComponent.IsValid();
+}
+
+FPCGGraphToolData::FPCGGraphToolData()
+	: InitialActorClassToSpawn(AActor::StaticClass())
+{
+
+}
\ No newline at end of file
diff --git a/Engine/Plugins/PCG/Source/PCG/Private/Editor/PCGGraphToolData.cpp b/Engine/Plugins/PCG/Source/PCG/Private/Editor/PCGGraphToolData.cpp
deleted file mode 100644
diff --git a/Engine/Plugins/PCG/Source/PCG/Public/Data/PCGToolData.h b/Engine/Plugins/PCG/Source/PCG/Public/Data/PCGToolData.h
new file mode 100644
--- /dev/null
+++ b/Engine/Plugins/PCG/Source/PCG/Public/Data/PCGToolData.h
@@ -0,0 +1,241 @@
+﻿// Copyright Epic Games, Inc. All Rights Reserved.
+
+#pragma once
+
+#include "Data/PCGSplineData.h"
+#include "UObject/WeakObjectPtrTemplates.h"
+
+#include "PCGToolData.generated.h"
+
+class AActor;
+class USplineComponent;
+class UPCGPointArrayData;
+class UPCGComponent;
+class UInteractiveTool;
+class UInteractiveToolPropertySet;
+
+/** A transient struct that contains resources generated during tool usage for a specific tool data struct. Used to finalize or clean up when a tool is applied or cancelled. */
+USTRUCT()
+struct FPCGInteractiveToolWorkingDataGeneratedResources
+{
+	GENERATED_BODY()
+
+	FPCGInteractiveToolWorkingDataGeneratedResources() = default;
+	virtual ~FPCGInteractiveToolWorkingDataGeneratedResources() = default;
+
+#if WITH_EDITORONLY_DATA
+	/** Finalizes generated resources (e.g. remove Transient flag). */
+	PCG_API void Finalize();
+
+	/** Remove and destroy generated resources. */
+	PCG_API void Cleanup();
+	
+	UPROPERTY(Transient)
+	TArray<TWeakObjectPtr<UActorComponent>> GeneratedComponents;
+#endif
+};
+
+/** A transient context struct that is passed in so the tool data can manage its data properly during tool usage. Used at editor-time only. */
+USTRUCT()
+struct FPCGInteractiveToolWorkingDataContext
+{
+	GENERATED_BODY()
+
+	FPCGInteractiveToolWorkingDataContext() = default;
+	virtual ~FPCGInteractiveToolWorkingDataContext() = default;
+
+#if WITH_EDITORONLY_DATA
+	PCG_API FPCGInteractiveToolWorkingDataContext(UPCGComponent* OwningComponent, UInteractiveTool* InteractiveTool, UInteractiveToolPropertySet* PropertySet, FName DataInstanceIdentifier);
+
+	/** The owning PCG Component. Should always be valid. Can be used as an Outer for new objects. */
+	UPROPERTY(Transient)
+	TWeakObjectPtr<UPCGComponent> OwningPCGComponent;
+
+	/** An optional DataInstanceIdentifier. Can be used to find or create a specific object, such as a Spline Component with a certain tag. */
+	UPROPERTY(Transient)
+	FName DataInstanceIdentifier;
+
+	/** The owning actor. The owning PCG component's outer. Should always be valid. Convenience accessor. */
+	UPROPERTY(Transient)
+	TWeakObjectPtr<AActor> OwningActor;
+
+	/** The interactive tool currently editing the tool data. */
+	UPROPERTY(Transient)
+	TWeakObjectPtr<UInteractiveTool> InteractiveTool;
+
+	/** The PCG Setting properties of the Interactive Tool. */
+	UPROPERTY(Transient)
+	TWeakObjectPtr<UInteractiveToolPropertySet> PCGSettings;
+#endif
+};
+
+/** The base struct for working data for a tool. This data is created by the tool, stored on the PCG Component,
+ * then read back by the tool on subsequent tool invocations. */
+USTRUCT()
+struct FPCGInteractiveToolWorkingData
+{
+	GENERATED_BODY()
+
+	FPCGInteractiveToolWorkingData() = default;
+	virtual ~FPCGInteractiveToolWorkingData() = default;
+
+	PCG_API bool IsInitialized() const;
+
+	PCG_API virtual bool IsValid() const;
+
+	/** Construct and add pcg data to the runtime context.
+	 * This will translate data from this struct into UPCGData to be processed by the graph.
+	 * Add dynamic tracking keys in here if you want to refresh when a change to your source data happens. */
+	PCG_API virtual void InitializeRuntimeElementData(FPCGContext* InContext) const;
+	
+#if WITH_EDITORONLY_DATA
+public:
+	DECLARE_DELEGATE_OneParam(FOnInitialized, FPCGInteractiveToolWorkingData*)
+	
+	/** Initializes data used by the tool working data. Initialize is called once when the working data is created. */
+	PCG_API void Initialize(const FPCGInteractiveToolWorkingDataContext& Context);
+	
+	/** OnToolStart is called every time the tool is started up, or when the data is created.
+	 * Can be used to save the initial state of persistent data to restore to on Cancel, or to Modify existing data. */
+	PCG_API virtual void OnToolStart(const FPCGInteractiveToolWorkingDataContext& Context);
+	
+	/** Called when the tool is applied. Finalize changes here. */
+	PCG_API virtual void OnToolApply(const FPCGInteractiveToolWorkingDataContext& Context);
+	
+	/** Called when the tool is canceled. Revert changes you have made here. */
+	PCG_API virtual void OnToolCancel(const FPCGInteractiveToolWorkingDataContext& Context);
+
+	PCG_API FOnInitialized& OnInitialized() { return OnInitializedDelegate; }
+private:
+	/** Implement your initialization logic here. Called by Initialize. */
+	PCG_API virtual void InitializeInternal(const FPCGInteractiveToolWorkingDataContext& Context);
+protected:
+	/** Each working data struct keeps track of its generated resources. They have to be reverted during OnToolCancel. */
+	UPROPERTY(Transient)
+	FPCGInteractiveToolWorkingDataGeneratedResources GeneratedResources;
+
+	/** OnInitialized delegate to allow the settings class to do more work. */
+	FOnInitialized OnInitializedDelegate;
+#endif
+	
+private:
+	UPROPERTY()
+	bool bInitialized = false;
+};
+
+/** Tool Data that stores points. */
+USTRUCT(BlueprintType, DisplayName="Point Array Tool Data")
+struct FPCGInteractiveToolWorkingData_PointArrayData : public FPCGInteractiveToolWorkingData
+{
+	GENERATED_BODY()
+	
+public:
+	FPCGInteractiveToolWorkingData_PointArrayData() = default;
+
+	PCG_API virtual void InitializeRuntimeElementData(FPCGContext* InContext) const override;
+	
+	PCG_API virtual bool IsValid() const override;
+	
+	PCG_API UPCGPointArrayData* GetPointArrayData() const;
+
+#if WITH_EDITORONLY_DATA
+	PCG_API virtual void InitializeInternal(const FPCGInteractiveToolWorkingDataContext& Context) override;
+	
+	PCG_API virtual void OnToolStart(const FPCGInteractiveToolWorkingDataContext& Context) override;
+	
+	PCG_API virtual void OnToolCancel(const FPCGInteractiveToolWorkingDataContext& Context) override;
+#endif
+	
+protected:
+	UPROPERTY(VisibleAnywhere, BlueprintReadOnly, Category="Working Data")
+	TObjectPtr<UPCGPointArrayData> GeneratedPointData;
+
+	/** A transient copy of the point array data, created on tool start/data access. We revert back to these points on Cancel. */
+	UPROPERTY(Transient)
+	TObjectPtr<UPCGPointArrayData> OnToolStartPointArray;
+};
+
+/** The working data for splines. */
+USTRUCT(BlueprintType, DisplayName="Spline Tool Data")
+struct FPCGInteractiveToolWorkingData_Spline : public FPCGInteractiveToolWorkingData
+{
+	GENERATED_BODY()
+public:
+	FPCGInteractiveToolWorkingData_Spline() = default;
+	virtual ~FPCGInteractiveToolWorkingData_Spline() override = default;
+	
+	PCG_API virtual void InitializeRuntimeElementData(FPCGContext* InContext) const override;
+	
+	PCG_API virtual bool IsValid() const override;
+	
+	PCG_API USplineComponent* GetSplineComponent() const;
+
+#if WITH_EDITOR
+public:
+	virtual void OnToolStart(const FPCGInteractiveToolWorkingDataContext& Context) override;
+	virtual void OnToolCancel(const FPCGInteractiveToolWorkingDataContext& Context) override;
+	
+private:
+	virtual void InitializeInternal(const FPCGInteractiveToolWorkingDataContext& Context) override;
+	
+private:
+	USplineComponent* FindOrGenerateSplineComponent(const FPCGInteractiveToolWorkingDataContext& Context);
+	
+	virtual USplineComponent* FindMatchingSplineComponent(const FPCGInteractiveToolWorkingDataContext& Context);
+	virtual USplineComponent* GenerateMatchingSplineComponent(const FPCGInteractiveToolWorkingDataContext& Context);
+#endif
+	
+protected:
+	UPROPERTY(VisibleAnywhere, BlueprintReadOnly, Category="Working Data")
+	TWeakObjectPtr<USplineComponent> SplineComponent;
+
+	UPROPERTY(Transient)
+	TObjectPtr<USplineComponent> OnToolStartSplineComponent;
+};
+
+USTRUCT(BlueprintType, DisplayName="Spline Surface Tool Data")
+struct FPCGInteractiveToolWorkingData_SplineSurface : public FPCGInteractiveToolWorkingData_Spline
+{
+	GENERATED_BODY()
+	
+	FPCGInteractiveToolWorkingData_SplineSurface() = default;
+	virtual ~FPCGInteractiveToolWorkingData_SplineSurface() override = default;
+
+#if WITH_EDITOR
+private:
+	/** We look for and generate a closed spline component. */
+	virtual USplineComponent* FindMatchingSplineComponent(const FPCGInteractiveToolWorkingDataContext& Context) override;
+	virtual USplineComponent* GenerateMatchingSplineComponent(const FPCGInteractiveToolWorkingDataContext& Context) override;
+#endif
+};
+
+/** The tool data stored on the graph itself. Informs the UI. */
+USTRUCT(BlueprintType)
+struct FPCGGraphToolData
+{
+	GENERATED_BODY()
+
+	PCG_API FPCGGraphToolData();
+	virtual ~FPCGGraphToolData() = default;
+	
+	UPROPERTY(EditAnywhere, Category = Tool)
+	FText DisplayName;
+	
+	UPROPERTY(EditAnywhere, Category = Tool)
+	TArray<FString> CompatibleToolTags;
+
+	UPROPERTY(EditAnywhere, Category = Tool)
+	TSubclassOf<AActor> InitialActorClassToSpawn;
+};
+
+USTRUCT(BlueprintType)
+struct FPCGInteractiveToolWorkingDataContainer
+{
+	GENERATED_BODY()
+
+	FPCGInteractiveToolWorkingDataContainer() = default;
+	virtual ~FPCGInteractiveToolWorkingDataContainer() = default;
+	
+	UPROPERTY(BlueprintReadOnly, VisibleAnywhere, Category = "Tool Data")
+	TMap<FName, TInstancedStruct<FPCGInteractiveToolWorkingData>> ToolData;
+};
\ No newline at end of file
diff --git a/Engine/Plugins/PCG/Source/PCG/Public/Editor/PCGGraphToolData.h b/Engine/Plugins/PCG/Source/PCG/Public/Editor/PCGGraphToolData.h
deleted file mode 100644
diff --git a/Engine/Plugins/PCG/Source/PCG/Public/PCGComponent.h b/Engine/Plugins/PCG/Source/PCG/Public/PCGComponent.h
--- a/Engine/Plugins/PCG/Source/PCG/Public/PCGComponent.h
+++ b/Engine/Plugins/PCG/Source/PCG/Public/PCGComponent.h
@@ -7,6 +7,7 @@
 #include "PCGGraphExecutionStateInterface.h"
 #include "PCGNode.h"
 #include "PCGSettings.h"
+#include "Data/PCGToolData.h"
 #include "Graph/PCGStackContext.h"
 #include "Grid/PCGGridDescriptor.h"
 #include "Utils/PCGExtraCapture.h"
@@ -14,7 +15,6 @@
 #include "ComponentInstanceDataCache.h"
 #include "Components/ActorComponent.h"
 #include "Containers/Map.h"
-#include "Editor/PCGGraphToolData.h"
 #include "Misc/ScopeExit.h"
 #include "UObject/ObjectKey.h"
 #include "Misc/TransactionallySafeCriticalSection.h"
diff --git a/Engine/Plugins/PCG/Source/PCG/Public/PCGGraph.h b/Engine/Plugins/PCG/Source/PCG/Public/PCGGraph.h
--- a/Engine/Plugins/PCG/Source/PCG/Public/PCGGraph.h
+++ b/Engine/Plugins/PCG/Source/PCG/Public/PCGGraph.h
@@ -5,6 +5,7 @@
 #include "PCGCommon.h"
 #include "PCGNode.h"
 #include "PCGSettings.h"
+#include "Data/PCGToolData.h"
 #include "Editor/PCGGraphCustomization.h"
 #include "Graph/PCGStackContext.h"
 #include "Helpers/PCGGraphParameterExtension.h"
@@ -14,7 +15,6 @@
 #endif // WITH_EDITOR
 
 #include "ComputeFramework/ComputeGraphInstance.h"
-#include "Editor/PCGGraphToolData.h"
 #include "StructUtils/PropertyBag.h"
 #include "UObject/ObjectPtr.h"
 
diff --git a/Engine/Plugins/PCG/Source/PCGEditor/Public/EditorMode/Tools/Helpers/PCGEdModeEditorUtilities.h b/Engine/Plugins/PCG/Source/PCGEditor/Public/EditorMode/Tools/Helpers/PCGEdModeEditorUtilities.h
--- a/Engine/Plugins/PCG/Source/PCGEditor/Public/EditorMode/Tools/Helpers/PCGEdModeEditorUtilities.h
+++ b/Engine/Plugins/PCG/Source/PCGEditor/Public/EditorMode/Tools/Helpers/PCGEdModeEditorUtilities.h
@@ -1,10 +1,11 @@
-﻿// Copyright Epic Games, Inc. All Rights Reserved.
+// Copyright Epic Games, Inc. All Rights Reserved.
 
 #pragma once
 
-#include "Templates/SharedPointer.h"
-#include "Editor/PCGGraphToolData.h"
+#include "Data/PCGToolData.h"
+
 #include "AssetRegistry/AssetData.h"
+#include "Templates/SharedPointer.h"
 
 class FUICommandInfo;
 class FBindingContext;
