[PCG] FPCGComputeGraphElement: avoid crashing on standalone graphs

#rnx

diff --git a/Engine/Plugins/PCG/Source/PCG/Private/Compute/Elements/PCGComputeGraphElement.cpp b/Engine/Plugins/PCG/Source/PCG/Private/Compute/Elements/PCGComputeGraphElement.cpp
--- a/Engine/Plugins/PCG/Source/PCG/Private/Compute/Elements/PCGComputeGraphElement.cpp
+++ b/Engine/Plugins/PCG/Source/PCG/Private/Compute/Elements/PCGComputeGraphElement.cpp
@@ -130,6 +130,12 @@ bool FPCGComputeGraphElement::ExecuteInternal(FPCGContext* InContext) const
 		return true;
 	}
 
+	if (!InContext->ExecutionSource->GetExecutionState().GetWorld())
+	{
+		UE_LOG(LogPCG, Error, TEXT("FPCGComputeGraphElement: Execution source is expected to have a world."));
+		return true;
+	}
+
 	auto SleepUntilNextFrame = [Context]()
 	{
 		Context->bIsPaused = true;
