[PCG] Disabled caching for tool data retrieval until we introduce per-execution cache for this

#rnx

diff --git a/Engine/Plugins/PCG/Source/PCG/Private/Data/PCGToolData.cpp b/Engine/Plugins/PCG/Source/PCG/Private/Data/PCGToolData.cpp
--- a/Engine/Plugins/PCG/Source/PCG/Private/Data/PCGToolData.cpp
+++ b/Engine/Plugins/PCG/Source/PCG/Private/Data/PCGToolData.cpp
@@ -135,13 +135,15 @@ void FPCGInteractiveToolWorkingData_PointArrayData::InitializeRuntimeElementData
 	TArray<FPCGTaggedData>& Outputs = Context->OutputData.TaggedData;
 
 	FPCGTaggedData& Output = Outputs.Emplace_GetRef();
-	// We duplicate it so that PCG can take 'ownership'.
-	// If not, it can't process references and caches properly
-	// Might be relying on the fact it shouldn't access the same object over multiple generations
-	// And if we recreate it every time it gets properly managed
-	// If it gets dropped from cache, it's expected to be invalid
+
+#if WITH_EDITOR
+	// Implementation note - we duplicate here to make sure we are not preserving the parent, because it might still change, which would cause issues.
+	check(IsInGameThread());
+	Output.Data = Cast<UPCGData>(StaticDuplicateObject(GeneratedPointData, GetTransientPackage()));
+#else
 	Output.Data = GeneratedPointData->DuplicateData(Context);
-	Output.Pin = PCGPinConstants::DefaultOutputLabel;
+#endif
+
 #if WITH_EDITOR
 	FPCGSelectionKey Key = FPCGSelectionKey::CreateFromObjectPtr(GeneratedPointData);
 	FPCGDynamicTrackingHelper::AddSingleDynamicTrackingKey(Context, MoveTemp(Key), false);
@@ -365,7 +367,6 @@ void FPCGInteractiveToolWorkingData_Spline::InitializeRuntimeElementData(FPCGCon
 		UPCGSplineData* SplineData = FPCGContext::NewObject_AnyThread<UPCGSplineData>(Context);
 		SplineData->Initialize(SplineComponent.Get());
 		Output.Data = SplineData;
-		Output.Pin = PCGPinConstants::DefaultOutputLabel;
 #if WITH_EDITOR
 		FPCGSelectionKey Key = FPCGSelectionKey::CreateFromObjectPtr(SplineComponent.Get());
 		FPCGDynamicTrackingHelper::AddSingleDynamicTrackingKey(Context, MoveTemp(Key), false);
diff --git a/Engine/Plugins/PCG/Source/PCG/Public/Elements/PCGDataFromTool.h b/Engine/Plugins/PCG/Source/PCG/Public/Elements/PCGDataFromTool.h
--- a/Engine/Plugins/PCG/Source/PCG/Public/Elements/PCGDataFromTool.h
+++ b/Engine/Plugins/PCG/Source/PCG/Public/Elements/PCGDataFromTool.h
@@ -49,6 +49,8 @@ public:
 
 class FPCGDataFromToolElement : public IPCGElement
 {
+public:
+	virtual bool IsCacheable(const UPCGSettings* InSettings) const override { return false; }
 protected:
 	virtual bool CanExecuteOnlyOnMainThread(FPCGContext* Context) const override { return true; }
 	PCG_API virtual bool ExecuteInternal(FPCGContext* Context) const override;
