[PCG] Fix Custom HLSL node attribute parsing failure when pin names have underscores

#rnx

diff --git a/Engine/Plugins/PCG/Source/PCG/Private/Compute/Elements/PCGCustomHLSLKernel.cpp b/Engine/Plugins/PCG/Source/PCG/Private/Compute/Elements/PCGCustomHLSLKernel.cpp
--- a/Engine/Plugins/PCG/Source/PCG/Private/Compute/Elements/PCGCustomHLSLKernel.cpp
+++ b/Engine/Plugins/PCG/Source/PCG/Private/Compute/Elements/PCGCustomHLSLKernel.cpp
@@ -968,7 +968,10 @@ void UPCGCustomHLSLKernel::ParseShaderSource()
 				{
 					AttributeFunctionMatch.Reset();
 
-					const int32 DelimiterIndex = TokenString.Find("_");
+					// Use the last underscore as a delimiter in case the pin name contains underscores.
+					int32 DelimiterIndex = INDEX_NONE;
+					TokenString.FindLastChar('_', DelimiterIndex);
+
 					AttributeFunctionMatch.PinStr = TokenString.Left(DelimiterIndex);
 					AttributeFunctionMatch.FuncStr = TokenString.Mid(DelimiterIndex + 1, 3);
 					AttributeFunctionMatch.TypeStr = TokenString.Mid(DelimiterIndex + 4, Token.Range.Len() - DelimiterIndex + 1);
