[PCG] Fix failing GPU get/set transform tests

Transform attributes should be transposed on pack/unpack to match the column major layout of Get/SetPointTransform

#rnx

diff --git a/Engine/Plugins/PCG/Source/PCG/Private/Compute/PCGDataForGPU.cpp b/Engine/Plugins/PCG/Source/PCG/Private/Compute/PCGDataForGPU.cpp
--- a/Engine/Plugins/PCG/Source/PCG/Private/Compute/PCGDataForGPU.cpp
+++ b/Engine/Plugins/PCG/Source/PCG/Private/Compute/PCGDataForGPU.cpp
@@ -238,20 +238,20 @@ namespace PCGDataForGPUHelpers
 
 			check(StrideBytes == 64);
 			OutPackedDataCollection[InOutAddressUints++] = FMath::AsUInt(static_cast<float>(Matrix.M[0][0]));
-			OutPackedDataCollection[InOutAddressUints++] = FMath::AsUInt(static_cast<float>(Matrix.M[0][1]));
-			OutPackedDataCollection[InOutAddressUints++] = FMath::AsUInt(static_cast<float>(Matrix.M[0][2]));
-			OutPackedDataCollection[InOutAddressUints++] = FMath::AsUInt(static_cast<float>(Matrix.M[0][3]));
 			OutPackedDataCollection[InOutAddressUints++] = FMath::AsUInt(static_cast<float>(Matrix.M[1][0]));
-			OutPackedDataCollection[InOutAddressUints++] = FMath::AsUInt(static_cast<float>(Matrix.M[1][1]));
-			OutPackedDataCollection[InOutAddressUints++] = FMath::AsUInt(static_cast<float>(Matrix.M[1][2]));
-			OutPackedDataCollection[InOutAddressUints++] = FMath::AsUInt(static_cast<float>(Matrix.M[1][3]));
 			OutPackedDataCollection[InOutAddressUints++] = FMath::AsUInt(static_cast<float>(Matrix.M[2][0]));
-			OutPackedDataCollection[InOutAddressUints++] = FMath::AsUInt(static_cast<float>(Matrix.M[2][1]));
-			OutPackedDataCollection[InOutAddressUints++] = FMath::AsUInt(static_cast<float>(Matrix.M[2][2]));
-			OutPackedDataCollection[InOutAddressUints++] = FMath::AsUInt(static_cast<float>(Matrix.M[2][3]));
 			OutPackedDataCollection[InOutAddressUints++] = FMath::AsUInt(static_cast<float>(Matrix.M[3][0]));
+			OutPackedDataCollection[InOutAddressUints++] = FMath::AsUInt(static_cast<float>(Matrix.M[0][1]));
+			OutPackedDataCollection[InOutAddressUints++] = FMath::AsUInt(static_cast<float>(Matrix.M[1][1]));
+			OutPackedDataCollection[InOutAddressUints++] = FMath::AsUInt(static_cast<float>(Matrix.M[2][1]));
 			OutPackedDataCollection[InOutAddressUints++] = FMath::AsUInt(static_cast<float>(Matrix.M[3][1]));
+			OutPackedDataCollection[InOutAddressUints++] = FMath::AsUInt(static_cast<float>(Matrix.M[0][2]));
+			OutPackedDataCollection[InOutAddressUints++] = FMath::AsUInt(static_cast<float>(Matrix.M[1][2]));
+			OutPackedDataCollection[InOutAddressUints++] = FMath::AsUInt(static_cast<float>(Matrix.M[2][2]));
 			OutPackedDataCollection[InOutAddressUints++] = FMath::AsUInt(static_cast<float>(Matrix.M[3][2]));
+			OutPackedDataCollection[InOutAddressUints++] = FMath::AsUInt(static_cast<float>(Matrix.M[0][3]));
+			OutPackedDataCollection[InOutAddressUints++] = FMath::AsUInt(static_cast<float>(Matrix.M[1][3]));
+			OutPackedDataCollection[InOutAddressUints++] = FMath::AsUInt(static_cast<float>(Matrix.M[2][3]));
 			OutPackedDataCollection[InOutAddressUints++] = FMath::AsUInt(static_cast<float>(Matrix.M[3][3]));
 
 			break;
@@ -507,20 +507,20 @@ namespace PCGDataForGPUHelpers
 				const uint32 PackedElementIndex = InAddressUints + ElementIndex * 16;
 
 				Matrix.M[0][0] = DataAsFloat[PackedElementIndex + 0];
-				Matrix.M[0][1] = DataAsFloat[PackedElementIndex + 1];
-				Matrix.M[0][2] = DataAsFloat[PackedElementIndex + 2];
-				Matrix.M[0][3] = DataAsFloat[PackedElementIndex + 3];
-				Matrix.M[1][0] = DataAsFloat[PackedElementIndex + 4];
+				Matrix.M[1][0] = DataAsFloat[PackedElementIndex + 1];
+				Matrix.M[2][0] = DataAsFloat[PackedElementIndex + 2];
+				Matrix.M[3][0] = DataAsFloat[PackedElementIndex + 3];
+				Matrix.M[0][1] = DataAsFloat[PackedElementIndex + 4];
 				Matrix.M[1][1] = DataAsFloat[PackedElementIndex + 5];
-				Matrix.M[1][2] = DataAsFloat[PackedElementIndex + 6];
-				Matrix.M[1][3] = DataAsFloat[PackedElementIndex + 7];
-				Matrix.M[2][0] = DataAsFloat[PackedElementIndex + 8];
-				Matrix.M[2][1] = DataAsFloat[PackedElementIndex + 9];
+				Matrix.M[2][1] = DataAsFloat[PackedElementIndex + 6];
+				Matrix.M[3][1] = DataAsFloat[PackedElementIndex + 7];
+				Matrix.M[0][2] = DataAsFloat[PackedElementIndex + 8];
+				Matrix.M[1][2] = DataAsFloat[PackedElementIndex + 9];
 				Matrix.M[2][2] = DataAsFloat[PackedElementIndex + 10];
-				Matrix.M[2][3] = DataAsFloat[PackedElementIndex + 11];
-				Matrix.M[3][0] = DataAsFloat[PackedElementIndex + 12];
-				Matrix.M[3][1] = DataAsFloat[PackedElementIndex + 13];
-				Matrix.M[3][2] = DataAsFloat[PackedElementIndex + 14];
+				Matrix.M[3][2] = DataAsFloat[PackedElementIndex + 11];
+				Matrix.M[0][3] = DataAsFloat[PackedElementIndex + 12];
+				Matrix.M[1][3] = DataAsFloat[PackedElementIndex + 13];
+				Matrix.M[2][3] = DataAsFloat[PackedElementIndex + 14];
 				Matrix.M[3][3] = DataAsFloat[PackedElementIndex + 15];
 
 				new (&Values[ElementIndex]) FTransform(Matrix);
