[PCG] Fix runtime gen artifacts being left behind during blueprint teardown

#rnx

diff --git a/Engine/Plugins/PCG/Source/PCG/Private/RuntimeGen/PCGRuntimeGenScheduler.cpp b/Engine/Plugins/PCG/Source/PCG/Private/RuntimeGen/PCGRuntimeGenScheduler.cpp
--- a/Engine/Plugins/PCG/Source/PCG/Private/RuntimeGen/PCGRuntimeGenScheduler.cpp
+++ b/Engine/Plugins/PCG/Source/PCG/Private/RuntimeGen/PCGRuntimeGenScheduler.cpp
@@ -1491,6 +1491,29 @@ void FPCGRuntimeGenScheduler::OnOriginalComponentUnregistered(UPCGComponent* InO
 	CleanupRemainingComponents(InOriginalComponent);
 }
 
+#if WITH_EDITOR
+void FPCGRuntimeGenScheduler::OnObjectsReplaced(const TMap<UObject*, UObject*>& InOldToNewInstances)
+{
+	for (const TPair<UObject*, UObject*>& OldToNew : InOldToNewInstances)
+	{
+		UPCGComponent* OldComponent = Cast<UPCGComponent>(OldToNew.Key);
+		UPCGComponent* NewComponent = Cast<UPCGComponent>(OldToNew.Value);
+
+		// Replace all usage of OldComponent with NewComponent.
+		if (OldComponent && NewComponent)
+		{
+			for (FGridGenerationKey& GenerationKey : GeneratedComponents)
+			{
+				if (GenerationKey.OriginalComponent == OldComponent)
+				{
+					GenerationKey.OriginalComponent = NewComponent;
+				}
+			}
+		}
+	}
+}
+#endif
+
 void FPCGRuntimeGenScheduler::CleanupRemainingComponents(UPCGComponent* InOriginalComponent)
 {
 	// Check for remaining PAs to cleanup
diff --git a/Engine/Plugins/PCG/Source/PCG/Private/Subsystems/PCGSubsystem.cpp b/Engine/Plugins/PCG/Source/PCG/Private/Subsystems/PCGSubsystem.cpp
--- a/Engine/Plugins/PCG/Source/PCG/Private/Subsystems/PCGSubsystem.cpp
+++ b/Engine/Plugins/PCG/Source/PCG/Private/Subsystems/PCGSubsystem.cpp
@@ -209,6 +209,10 @@ void UPCGSubsystem::Deinitialize()
 
 	ActorAndComponentMapping->Deinitialize();
 
+#if WITH_EDITOR
+	FCoreUObjectDelegates::OnObjectsReplaced.RemoveAll(this);
+#endif
+
 	Super::Deinitialize();
 }
 
@@ -238,6 +242,10 @@ void UPCGSubsystem::PostInitialize()
 	// Initialize runtime generation scheduler
 	check(!RuntimeGenScheduler);
 	RuntimeGenScheduler = new FPCGRuntimeGenScheduler(GetWorld(), ActorAndComponentMapping.Get());
+
+#if WITH_EDITOR
+	FCoreUObjectDelegates::OnObjectsReplaced.AddUObject(this, &UPCGSubsystem::OnObjectsReplaced);
+#endif
 }
 
 UPCGSubsystem* UPCGSubsystem::GetInstance(UWorld* World)
@@ -1897,6 +1905,14 @@ void UPCGSubsystem::SetChainedDispatchToLocalComponents(bool bInChainedDispatch)
 	ActorAndComponentMapping->SetChainedDispatchToLocalComponents(bInChainedDispatch);
 }
 
+void UPCGSubsystem::OnObjectsReplaced(const TMap<UObject*, UObject*>& InOldToNewInstances)
+{
+	if (RuntimeGenScheduler)
+	{
+		RuntimeGenScheduler->OnObjectsReplaced(InOldToNewInstances);
+	}
+}
+
 void UPCGSubsystem::OnPCGGraphCancelled(UPCGComponent* InComponent)
 {
 	ActorAndComponentMapping->OnPCGGraphCancelled(InComponent);
diff --git a/Engine/Plugins/PCG/Source/PCG/Public/RuntimeGen/PCGRuntimeGenScheduler.h b/Engine/Plugins/PCG/Source/PCG/Public/RuntimeGen/PCGRuntimeGenScheduler.h
--- a/Engine/Plugins/PCG/Source/PCG/Public/RuntimeGen/PCGRuntimeGenScheduler.h
+++ b/Engine/Plugins/PCG/Source/PCG/Public/RuntimeGen/PCGRuntimeGenScheduler.h
@@ -82,6 +82,10 @@ public:
 	/** Destroy all runtime gen partition actors (both generated and pooled). Executed in next tick. */
 	void FlushAllGeneratedActors() { bActorFlushRequested = true; }
 
+#if WITH_EDITOR
+	void OnObjectsReplaced(const TMap<UObject*, UObject*>& InOldToNewInstances);
+#endif
+
 protected:
 	// Grid size, grid coords, original component, local componnent
 	struct FGridGenerationKey
diff --git a/Engine/Plugins/PCG/Source/PCG/Public/Subsystems/PCGSubsystem.h b/Engine/Plugins/PCG/Source/PCG/Public/Subsystems/PCGSubsystem.h
--- a/Engine/Plugins/PCG/Source/PCG/Public/Subsystems/PCGSubsystem.h
+++ b/Engine/Plugins/PCG/Source/PCG/Public/Subsystems/PCGSubsystem.h
@@ -365,6 +365,8 @@ private:
 	UE_API void UpdateMappingPCGComponentPartitionActor(UPCGComponent* InComponent);
 
 	UE_API void SetChainedDispatchToLocalComponents(bool bInChainedDispatch);
+
+	void OnObjectsReplaced(const TMap<UObject*, UObject*>& InOldToNewInstances);
 #endif // WITH_EDITOR
 	
 private:
