[PCG] Improve error logging when trying to override GPU static mesh spawner component properties with non-basic attributes

#rnx

diff --git a/Engine/Plugins/PCG/Source/PCG/Internal/Compute/PCGDataForGPU.h b/Engine/Plugins/PCG/Source/PCG/Internal/Compute/PCGDataForGPU.h
--- a/Engine/Plugins/PCG/Source/PCG/Internal/Compute/PCGDataForGPU.h
+++ b/Engine/Plugins/PCG/Source/PCG/Internal/Compute/PCGDataForGPU.h
@@ -322,7 +322,7 @@ namespace PCGDataForGPUHelpers
 	void ComputeCustomFloatPacking(
 		FPCGContext* InContext,
 		const UPCGSettings* InSettings,
-		TArray<FName>& InAttributeNames,
+		TConstArrayView<FPCGAttributePropertyInputSelector> InAttributeSelectors,
 		const UPCGDataBinding* InBinding,
 		const TSharedPtr<const FPCGDataCollectionDesc> InDataCollectionDescription,
 		uint32& OutCustomFloatCount,
diff --git a/Engine/Plugins/PCG/Source/PCG/Private/Compute/DataInterfaces/Elements/PCGSkinnedMeshSpawnerDataInterface.cpp b/Engine/Plugins/PCG/Source/PCG/Private/Compute/DataInterfaces/Elements/PCGSkinnedMeshSpawnerDataInterface.cpp
--- a/Engine/Plugins/PCG/Source/PCG/Private/Compute/DataInterfaces/Elements/PCGSkinnedMeshSpawnerDataInterface.cpp
+++ b/Engine/Plugins/PCG/Source/PCG/Private/Compute/DataInterfaces/Elements/PCGSkinnedMeshSpawnerDataInterface.cpp
@@ -488,13 +488,21 @@ void UPCGSkinnedMeshSpawnerDataProvider::CreatePrimitiveDescriptors(FPCGContext*
 		return;
 	}
 
-	if (Settings->InstanceDataPackerParameters && Settings->InstanceDataPackerParameters->GetAttributeNames(/*OutNames=*/nullptr))
-	{
-		TArray<FName> AttributeNames;
-		Settings->InstanceDataPackerParameters->GetAttributeNames(&AttributeNames);
+	const TOptional<TConstArrayView<FPCGAttributePropertyInputSelector>> AttributeSelectors = Settings->InstanceDataPackerParameters ? Settings->InstanceDataPackerParameters->GetAttributeSelectors() : TOptional<TConstArrayView<FPCGAttributePropertyInputSelector>>();
 
-		for (FName AttributeName : AttributeNames)
+	if (AttributeSelectors.IsSet())
+	{
+		for (const FPCGAttributePropertyInputSelector& AttributeSelector : AttributeSelectors.GetValue())
 		{
+			const FName AttributeName = AttributeSelector.GetAttributeName();
+
+			if (!AttributeSelector.IsBasicAttribute() || AttributeName == PCGMetadataAttributeConstants::LastAttributeName)
+			{
+				PCG_KERNEL_VALIDATION_ERR(InContext, Settings,
+					FText::Format(LOCTEXT("OnlyBasicAttributesSupported", "Attribute '{0}' is invalid. GPU instance packer implementation currently only supports basic attributes."), AttributeSelector.GetDisplayText()));
+				continue;
+			}
+
 			if (AttributeName == NAME_None)
 			{
 				PCG_KERNEL_VALIDATION_ERR(InContext, Settings,
@@ -512,7 +520,7 @@ void UPCGSkinnedMeshSpawnerDataProvider::CreatePrimitiveDescriptors(FPCGContext*
 			}
 		}
 
-		PCGDataForGPUHelpers::ComputeCustomFloatPacking(InContext, Settings, AttributeNames, InBinding, InputDataDesc, CustomFloatCount, AttributeIdOffsetStrides);
+		PCGDataForGPUHelpers::ComputeCustomFloatPacking(InContext, Settings, AttributeSelectors.GetValue(), InBinding, InputDataDesc, CustomFloatCount, AttributeIdOffsetStrides);
 	}
 
 	if (Settings->MeshSelectorParameters)
diff --git a/Engine/Plugins/PCG/Source/PCG/Private/Compute/DataInterfaces/Elements/PCGStaticMeshSpawnerDataInterface.cpp b/Engine/Plugins/PCG/Source/PCG/Private/Compute/DataInterfaces/Elements/PCGStaticMeshSpawnerDataInterface.cpp
--- a/Engine/Plugins/PCG/Source/PCG/Private/Compute/DataInterfaces/Elements/PCGStaticMeshSpawnerDataInterface.cpp
+++ b/Engine/Plugins/PCG/Source/PCG/Private/Compute/DataInterfaces/Elements/PCGStaticMeshSpawnerDataInterface.cpp
@@ -551,13 +551,21 @@ void UPCGStaticMeshSpawnerDataProvider::CreatePrimitiveDescriptors(FPCGContext*
 		return;
 	}
 
-	if (Settings->InstanceDataPackerParameters && Settings->InstanceDataPackerParameters->GetAttributeNames(/*OutNames=*/nullptr))
-	{
-		TArray<FName> AttributeNames;
-		Settings->InstanceDataPackerParameters->GetAttributeNames(&AttributeNames);
+	const TOptional<TConstArrayView<FPCGAttributePropertyInputSelector>> AttributeSelectors = Settings->InstanceDataPackerParameters ? Settings->InstanceDataPackerParameters->GetAttributeSelectors() : TOptional<TConstArrayView<FPCGAttributePropertyInputSelector>>();
 
-		for (FName AttributeName : AttributeNames)
+	if (AttributeSelectors.IsSet())
+	{
+		for (const FPCGAttributePropertyInputSelector& AttributeSelector : AttributeSelectors.GetValue())
 		{
+			const FName AttributeName = AttributeSelector.GetAttributeName();
+
+			if (!AttributeSelector.IsBasicAttribute() || AttributeName == PCGMetadataAttributeConstants::LastAttributeName)
+			{
+				PCG_KERNEL_VALIDATION_ERR(InContext, Settings,
+					FText::Format(LOCTEXT("OnlyBasicAttributesSupported", "Attribute '{0}' is invalid. GPU instance packer implementation currently only supports basic attributes."), AttributeSelector.GetDisplayText()));
+				continue;
+			}
+
 			if (AttributeName == NAME_None)
 			{
 				PCG_KERNEL_VALIDATION_ERR(InContext, Settings,
@@ -575,7 +583,7 @@ void UPCGStaticMeshSpawnerDataProvider::CreatePrimitiveDescriptors(FPCGContext*
 			}
 		}
 
-		PCGDataForGPUHelpers::ComputeCustomFloatPacking(InContext, Settings, AttributeNames, InBinding, InputDataDesc, CustomFloatCount, AttributeIdOffsetStrides);
+		PCGDataForGPUHelpers::ComputeCustomFloatPacking(InContext, Settings, AttributeSelectors.GetValue(), InBinding, InputDataDesc, CustomFloatCount, AttributeIdOffsetStrides);
 	}
 
 	if (const UPCGMeshSelectorByAttribute* SelectorByAttribute = Cast<UPCGMeshSelectorByAttribute>(Settings->MeshSelectorParameters))
diff --git a/Engine/Plugins/PCG/Source/PCG/Private/Compute/PCGDataForGPU.cpp b/Engine/Plugins/PCG/Source/PCG/Private/Compute/PCGDataForGPU.cpp
--- a/Engine/Plugins/PCG/Source/PCG/Private/Compute/PCGDataForGPU.cpp
+++ b/Engine/Plugins/PCG/Source/PCG/Private/Compute/PCGDataForGPU.cpp
@@ -574,7 +574,7 @@ namespace PCGDataForGPUHelpers
 	void ComputeCustomFloatPacking(
 		FPCGContext* InContext,
 		const UPCGSettings* InSettings,
-		TArray<FName>& InAttributeNames,
+		TConstArrayView<FPCGAttributePropertyInputSelector> InAttributeSelectors,
 		const UPCGDataBinding* InBinding,
 		const TSharedPtr<const FPCGDataCollectionDesc> InDataCollectionDescription,
 		uint32& OutCustomFloatCount,
@@ -586,8 +586,10 @@ namespace PCGDataForGPUHelpers
 
 		uint32 OffsetFloats = 0;
 
-		for (FName AttributeName : InAttributeNames)
+		for (const FPCGAttributePropertyInputSelector& AttributeSelector : InAttributeSelectors)
 		{
+			const FName AttributeName = AttributeSelector.GetAttributeName();
+
 			// We need to do a lookup here as the user does not provide the full attribute description, only the attribute name. So we see what we can
 			// find in the input data similar to the CPU SM Spawner.
 			FPCGKernelAttributeDesc AttributeDesc;
diff --git a/Engine/Plugins/PCG/Source/PCG/Private/InstanceDataPackers/PCGInstanceDataPackerByAttribute.cpp b/Engine/Plugins/PCG/Source/PCG/Private/InstanceDataPackers/PCGInstanceDataPackerByAttribute.cpp
--- a/Engine/Plugins/PCG/Source/PCG/Private/InstanceDataPackers/PCGInstanceDataPackerByAttribute.cpp
+++ b/Engine/Plugins/PCG/Source/PCG/Private/InstanceDataPackers/PCGInstanceDataPackerByAttribute.cpp
@@ -96,4 +96,9 @@ bool UPCGInstanceDataPackerByAttribute::GetAttributeNames(TArray<FName>* OutName
 	return true;
 }
 
+TOptional<TConstArrayView<FPCGAttributePropertyInputSelector>> UPCGInstanceDataPackerByAttribute::GetAttributeSelectors() const
+{
+	return MakeConstArrayView(AttributeSelectors);
+}
+
 #undef LOCTEXT_NAMESPACE
diff --git a/Engine/Plugins/PCG/Source/PCG/Private/InstanceDataPackers/PCGSkinnedMeshInstanceDataPackerByAttribute.cpp b/Engine/Plugins/PCG/Source/PCG/Private/InstanceDataPackers/PCGSkinnedMeshInstanceDataPackerByAttribute.cpp
--- a/Engine/Plugins/PCG/Source/PCG/Private/InstanceDataPackers/PCGSkinnedMeshInstanceDataPackerByAttribute.cpp
+++ b/Engine/Plugins/PCG/Source/PCG/Private/InstanceDataPackers/PCGSkinnedMeshInstanceDataPackerByAttribute.cpp
@@ -86,4 +86,9 @@ bool UPCGSkinnedMeshInstanceDataPackerByAttribute::GetAttributeNames(TArray<FNam
 	return true;
 }
 
+TOptional<TConstArrayView<FPCGAttributePropertyInputSelector>> UPCGSkinnedMeshInstanceDataPackerByAttribute::GetAttributeSelectors() const
+{
+	return MakeConstArrayView(AttributeSelectors);
+}
+
 #undef LOCTEXT_NAMESPACE
diff --git a/Engine/Plugins/PCG/Source/PCG/Public/InstanceDataPackers/PCGInstanceDataPackerBase.h b/Engine/Plugins/PCG/Source/PCG/Public/InstanceDataPackers/PCGInstanceDataPackerBase.h
--- a/Engine/Plugins/PCG/Source/PCG/Public/InstanceDataPackers/PCGInstanceDataPackerBase.h
+++ b/Engine/Plugins/PCG/Source/PCG/Public/InstanceDataPackers/PCGInstanceDataPackerBase.h
@@ -11,6 +11,7 @@ class IPCGAttributeAccessor;
 class IPCGAttributeAccessorKeys;
 class UPCGMetadata;
 class UPCGSpatialData;
+struct FPCGAttributePropertyInputSelector;
 struct FPCGContext;
 struct FPCGMeshInstanceList;
 
@@ -54,6 +55,9 @@ public:
 
 	/** If OutNames is not null, returns a list of all attributes that will be packed. Returns true if this list can be statically determined (prior to execution). */
 	virtual bool GetAttributeNames(TArray<FName>* OutNames) { return false; }
+
+	/** Returns a list of all attribute selectors that will be packed, if they can be statically determined (prior to execution). */
+	virtual TOptional<TConstArrayView<FPCGAttributePropertyInputSelector>> GetAttributeSelectors() const { return {}; }
 };
 
 #undef UE_API
diff --git a/Engine/Plugins/PCG/Source/PCG/Public/InstanceDataPackers/PCGInstanceDataPackerByAttribute.h b/Engine/Plugins/PCG/Source/PCG/Public/InstanceDataPackers/PCGInstanceDataPackerByAttribute.h
--- a/Engine/Plugins/PCG/Source/PCG/Public/InstanceDataPackers/PCGInstanceDataPackerByAttribute.h
+++ b/Engine/Plugins/PCG/Source/PCG/Public/InstanceDataPackers/PCGInstanceDataPackerByAttribute.h
@@ -20,6 +20,7 @@ public:
 #endif
 	virtual void PackInstances_Implementation(UPARAM(ref) FPCGContext& Context, const UPCGSpatialData* InSpatialData, UPARAM(ref) const FPCGMeshInstanceList& InstanceList, FPCGPackedCustomData& OutPackedCustomData) const override;
 	virtual bool GetAttributeNames(TArray<FName>* OutNames) override;
+	virtual TOptional<TConstArrayView<FPCGAttributePropertyInputSelector>> GetAttributeSelectors() const override;
 
 public:
 	UPROPERTY(BlueprintReadWrite, EditAnywhere, Category = InstanceDataPacker)
diff --git a/Engine/Plugins/PCG/Source/PCG/Public/InstanceDataPackers/PCGSkinnedMeshInstanceDataPackerBase.h b/Engine/Plugins/PCG/Source/PCG/Public/InstanceDataPackers/PCGSkinnedMeshInstanceDataPackerBase.h
--- a/Engine/Plugins/PCG/Source/PCG/Public/InstanceDataPackers/PCGSkinnedMeshInstanceDataPackerBase.h
+++ b/Engine/Plugins/PCG/Source/PCG/Public/InstanceDataPackers/PCGSkinnedMeshInstanceDataPackerBase.h
@@ -11,6 +11,7 @@ class IPCGAttributeAccessor;
 class IPCGAttributeAccessorKeys;
 class UPCGMetadata;
 class UPCGSpatialData;
+struct FPCGAttributePropertyInputSelector;
 struct FPCGContext;
 struct FPCGSkinnedMeshInstanceList;
 
@@ -57,6 +58,9 @@ public:
 
 	/** If OutNames is not null, returns a list of all attributes that will be packed. Returns true if this list can be statically determined (prior to execution). */
 	virtual bool GetAttributeNames(TArray<FName>* OutNames) { return false; }
+
+	/** Returns a list of all attribute selectors that will be packed, if they can be statically determined (prior to execution). */
+	virtual TOptional<TConstArrayView<FPCGAttributePropertyInputSelector>> GetAttributeSelectors() const { return {}; }
 };
 
 #undef UE_API
diff --git a/Engine/Plugins/PCG/Source/PCG/Public/InstanceDataPackers/PCGSkinnedMeshInstanceDataPackerByAttribute.h b/Engine/Plugins/PCG/Source/PCG/Public/InstanceDataPackers/PCGSkinnedMeshInstanceDataPackerByAttribute.h
--- a/Engine/Plugins/PCG/Source/PCG/Public/InstanceDataPackers/PCGSkinnedMeshInstanceDataPackerByAttribute.h
+++ b/Engine/Plugins/PCG/Source/PCG/Public/InstanceDataPackers/PCGSkinnedMeshInstanceDataPackerByAttribute.h
@@ -20,6 +20,7 @@ public:
 #endif
 	virtual void PackInstances_Implementation(UPARAM(ref) FPCGContext& Context, const UPCGSpatialData* InSpatialData, UPARAM(ref) const FPCGSkinnedMeshInstanceList& InstanceList, FPCGSkinnedMeshPackedCustomData& OutPackedCustomData) const override;
 	virtual bool GetAttributeNames(TArray<FName>* OutNames) override;
+	virtual TOptional<TConstArrayView<FPCGAttributePropertyInputSelector>> GetAttributeSelectors() const override;
 
 public:
 	UPROPERTY(BlueprintReadWrite, EditAnywhere, Category = InstanceDataPacker)
