[PCG] Added component filtering so we select the most likely candidate. Adjusted update rate and fixed first generation on new components.

#rnx

diff --git a/Engine/Plugins/PCG/Source/PCGEditor/Private/EditorMode/Tools/Helpers/PCGEdModeActorHelpers.cpp b/Engine/Plugins/PCG/Source/PCGEditor/Private/EditorMode/Tools/Helpers/PCGEdModeActorHelpers.cpp
--- a/Engine/Plugins/PCG/Source/PCGEditor/Private/EditorMode/Tools/Helpers/PCGEdModeActorHelpers.cpp
+++ b/Engine/Plugins/PCG/Source/PCGEditor/Private/EditorMode/Tools/Helpers/PCGEdModeActorHelpers.cpp
@@ -3,6 +3,7 @@
 #include "EditorMode/Tools/Helpers/PCGEdModeActorHelpers.h"
 
 #include "PCGComponent.h"
+#include "PCGGraph.h"
 #include "ActorFactories/ActorFactoryEmptyActor.h"
 #include "Helpers/PCGActorHelpers.h"
 
@@ -46,15 +47,40 @@ namespace UE::PCG::EditorMode
 
 	namespace PCGComponent
 	{
-		UPCGComponent* Find(const AActor* Owner)
+		UPCGComponent* Find(const AActor* Owner, const FName ToolTag)
 		{
 			if (!Owner)
 			{
 				return nullptr;
 			}
 
-			UPCGComponent* PCGComponent = Owner->GetComponentByClass<UPCGComponent>();
-			return PCGComponent;
+			TArray<UPCGComponent*, TInlineAllocator<8>> Components;
+			Owner->GetComponents<UPCGComponent>(Components);
+
+			// Priority order:
+			// 1. A tool-marked pcg component with a graph that has the same tool tags
+			// 2. A pcg component with a graph that has the same tool tags
+			UPCGComponent* TentativeComponent = nullptr;
+
+			for (UPCGComponent* Component : Components)
+			{
+				if (Component)
+				{
+					const bool bComponentHasToolTag = Component->ComponentHasTag(Tags::ToolGeneratedTag);
+					const bool bComponentHasCompatibleGraph = Component->GetGraph() && Component->GetGraph()->ToolData.CompatibleToolTags.Contains(ToolTag);
+
+					if (bComponentHasToolTag && bComponentHasCompatibleGraph)
+					{
+						return Component;
+					}
+					else if (bComponentHasCompatibleGraph && !TentativeComponent)
+					{
+						TentativeComponent = Component;
+					}
+				}
+			}
+
+			return TentativeComponent;
 		}
 
 		UPCGComponent* Create(AActor* Owner, const FName ComponentName)
diff --git a/Engine/Plugins/PCG/Source/PCGEditor/Private/EditorMode/Tools/PCGInteractiveToolSettings.cpp b/Engine/Plugins/PCG/Source/PCGEditor/Private/EditorMode/Tools/PCGInteractiveToolSettings.cpp
--- a/Engine/Plugins/PCG/Source/PCGEditor/Private/EditorMode/Tools/PCGInteractiveToolSettings.cpp
+++ b/Engine/Plugins/PCG/Source/PCGEditor/Private/EditorMode/Tools/PCGInteractiveToolSettings.cpp
@@ -68,7 +68,10 @@ bool UPCGInteractiveToolSettings::Initialize(UInteractiveTool* InTool)
 	TArray<AActor*> SelectedActors = State.SelectedActors;
 
 	// In order to constrain component selection and actor creation, we need to make sure to setup the default settings first.
-	TryLoadDefaults();
+	if (!TryLoadDefaults())
+	{
+		return false;
+	}
 
 	// If we have selected one actor, that's the actor we want to edit
 	// This actor might:
@@ -76,7 +79,7 @@ bool UPCGInteractiveToolSettings::Initialize(UInteractiveTool* InTool)
 	// - Have a PCG Component but no valid working data
 	// - Have a PCG Component with valid working data we want to edit
 	if(SelectedActors.Num() == 1)
-	{		
+	{
 		EditActor = SelectedActors[0];
 	}
 	// If we have no selected actor, we spawn a working actor that needs to fulfill the same actor requirements. This actor is fully instantiated when Applied.
@@ -105,7 +108,14 @@ bool UPCGInteractiveToolSettings::Initialize(UInteractiveTool* InTool)
 	// If we don't have a valid graph already, we try to assign the default graph to the newly created pcg component.
 	if (PCGComponent->GetGraph() == nullptr)
 	{
+		if (!ensureMsgf(ToolGraph, TEXT("The PCG component has no graph and no default graph for this tool was found in the editor preferences. Shutting down the tool.")))
+		{
+			return false;
+		}
+
 		PCGComponent->SetGraph(ToolGraph);
+		// Mark as we should try to generate - note that it will fail until the actor has bounds.
+		bTryToGenerate = true;
 	}
 
 	// Finally, update the tool using the component graph (either the newly set one or a previously existing one).
@@ -115,10 +125,10 @@ bool UPCGInteractiveToolSettings::Initialize(UInteractiveTool* InTool)
 	// Keyboard setup for switching between data instance identifiers
 	UKeyInputBehavior* KeyInputBehavior = NewObject<UKeyInputBehavior>();
 	KeyInputBehavior->Initialize(this,
-		{
-			EKeys::One, EKeys::Two, EKeys::Three, EKeys::Four, EKeys::Five,
-			EKeys::Six, EKeys::Seven, EKeys::Eight, EKeys::Nine, EKeys::Zero
-		});
+	{
+		EKeys::One, EKeys::Two, EKeys::Three, EKeys::Four, EKeys::Five,
+		EKeys::Six, EKeys::Seven, EKeys::Eight, EKeys::Nine, EKeys::Zero
+	});
 	KeyInputBehavior->bRequireAllKeys = false;
 	InTool->AddInputBehavior(KeyInputBehavior);
 
@@ -298,7 +308,7 @@ UPCGComponent* UPCGInteractiveToolSettings::FindOrGeneratePCGComponent(AActor& O
 {
 	using namespace UE::PCG::EditorMode;
 	
-	UPCGComponent* PCGComponent = PCGComponent::Find(&OwningActor);
+	UPCGComponent* PCGComponent = PCGComponent::Find(&OwningActor, GetToolTag());
 
 	if(PCGComponent == nullptr)
 	{
@@ -336,13 +346,19 @@ void UPCGInteractiveToolSettings::DeleteGeneratedResources() const
 	}
 }
 
-void UPCGInteractiveToolSettings::TryLoadDefaults()
+bool UPCGInteractiveToolSettings::TryLoadDefaults()
 {
 	check(!HasAnyFlags(RF_ClassDefaultObject));
 	if (UPCGGraphInterface* DefaultGraphInterface = CastChecked<UPCGGraphInterface>(GetDefaultGraph().TryLoad(), ECastCheckedType::NullAllowed))
 	{
 		SetToolGraph(DefaultGraphInterface);
 		SetActorClassToSpawn(DefaultGraphInterface->ToolData.InitialActorClassToSpawn);
+
+		return true;
+	}
+	else
+	{
+		return false;
 	}
 }
 
@@ -369,17 +385,17 @@ void UPCGInteractiveToolSettings::OnRefresh()
 {
 	if(HasAnyFlags(RF_MirroredGarbage) == false)
 	{
-		if(bAutoGenerate)
+		// We commit the property changed events now to let the graph regenerate, if the object is tracked
+		CommitPropertyChangedEvents();
+
+		if (bTryToGenerate)
 		{
-			UPCGComponent* PCGComponent = GetWorkingPCGComponent();
-			if(PCGComponent && !PCGComponent->IsCleaningUp() && !PCGComponent->IsGenerating())
+			if (UPCGComponent* PCGComponent = GetWorkingPCGComponent())
 			{
-				PCGComponent->GenerateLocal(true);
+				PCGComponent->GenerateLocal(/*bForce=*/true);
+				bTryToGenerate = !PCGComponent->IsGenerating();
 			}
 		}
-
-		// We commit the property changed events now to let the graph regenerate, if the object is tracked
-		CommitPropertyChangedEvents();
 		
 		// Reset the refresh timer
 		GraphRefreshTimer = GetDefault<UPCGEditorModeSettings>()->GraphRefreshRate;
@@ -516,9 +532,9 @@ void UPCGInteractiveToolSettings::OnRefreshInternal()
 FSoftObjectPath UPCGInteractiveToolSettings::GetDefaultGraph() const
 {
 	const FPCGPerInteractiveToolSettingSettings* Settings = GetDefault<UPCGEditorModeSettings>()->InteractiveToolSettings.FindByPredicate([this](const FPCGPerInteractiveToolSettingSettings& SettingsCandidate)
-		{
-			return SettingsCandidate.SettingsClass == this->GetClass();
-		});
+	{
+		return SettingsCandidate.SettingsClass == this->GetClass();
+	});
 
 	if(Settings != nullptr)
 	{
@@ -554,7 +570,7 @@ UPCGComponent* UPCGInteractiveToolSettings::GetWorkingPCGComponent() const
 {
 	if(WorkingComponentCache.IsValid() == false)
 	{
-		WorkingComponentCache = UE::PCG::EditorMode::PCGComponent::Find(EditActor.Get());
+		WorkingComponentCache = UE::PCG::EditorMode::PCGComponent::Find(EditActor.Get(), GetToolTag());
 	}
 	
 	return WorkingComponentCache.Get();
diff --git a/Engine/Plugins/PCG/Source/PCGEditor/Private/EditorMode/Tools/Paint/PCGPaintTool.cpp b/Engine/Plugins/PCG/Source/PCGEditor/Private/EditorMode/Tools/Paint/PCGPaintTool.cpp
--- a/Engine/Plugins/PCG/Source/PCGEditor/Private/EditorMode/Tools/Paint/PCGPaintTool.cpp
+++ b/Engine/Plugins/PCG/Source/PCGEditor/Private/EditorMode/Tools/Paint/PCGPaintTool.cpp
@@ -29,8 +29,6 @@ namespace PCGPaintTool::Constants
 
 UPCGInteractiveToolSettings_PaintTool::UPCGInteractiveToolSettings_PaintTool()
 {
-	// We turn auto generation off, as we'll make use of the tracking system to regenerate instead
-	bAutoGenerate = false;
 } 
 
 FName UPCGInteractiveToolSettings_PaintTool::GetToolTag() const
diff --git a/Engine/Plugins/PCG/Source/PCGEditor/Public/EditorMode/PCGEdModeSettings.h b/Engine/Plugins/PCG/Source/PCGEditor/Public/EditorMode/PCGEdModeSettings.h
--- a/Engine/Plugins/PCG/Source/PCGEditor/Public/EditorMode/PCGEdModeSettings.h
+++ b/Engine/Plugins/PCG/Source/PCGEditor/Public/EditorMode/PCGEdModeSettings.h
@@ -47,7 +47,7 @@ public:
 	UPROPERTY(EditAnywhere, Category = "Component", meta = (
 		ForceUnits = "Seconds",
 		ConsoleVariable = "pcg.editormode.GraphRefreshRate"))
-	float GraphRefreshRate = 1.5f;
+	float GraphRefreshRate = 0.5f;
 
 	/** Hide the tool buttons from the UI when a tool is actively being used. */
 	UPROPERTY(EditAnywhere, Config, Category = "Tools", meta = (
diff --git a/Engine/Plugins/PCG/Source/PCGEditor/Public/EditorMode/Tools/Helpers/PCGEdModeActorHelpers.h b/Engine/Plugins/PCG/Source/PCGEditor/Public/EditorMode/Tools/Helpers/PCGEdModeActorHelpers.h
--- a/Engine/Plugins/PCG/Source/PCGEditor/Public/EditorMode/Tools/Helpers/PCGEdModeActorHelpers.h
+++ b/Engine/Plugins/PCG/Source/PCGEditor/Public/EditorMode/Tools/Helpers/PCGEdModeActorHelpers.h
@@ -24,7 +24,7 @@ namespace UE::PCG::EditorMode
 	namespace PCGComponent
 	{
 		/** Find a PCG Component on an actor if it exists. */
-		UPCGComponent* Find(const AActor* Owner);
+		UPCGComponent* Find(const AActor* Owner, const FName ToolTag);
 		/** Spawn a PCG Component on an actor. */
 		UPCGComponent* Create(AActor* Owner, FName ComponentName);
 		/** Helper to spawn a temporary working PCG component for a tool. */
diff --git a/Engine/Plugins/PCG/Source/PCGEditor/Public/EditorMode/Tools/PCGInteractiveToolSettings.h b/Engine/Plugins/PCG/Source/PCGEditor/Public/EditorMode/Tools/PCGInteractiveToolSettings.h
--- a/Engine/Plugins/PCG/Source/PCGEditor/Public/EditorMode/Tools/PCGInteractiveToolSettings.h
+++ b/Engine/Plugins/PCG/Source/PCGEditor/Public/EditorMode/Tools/PCGInteractiveToolSettings.h
@@ -97,7 +97,6 @@ public:
 			return PCGComponent->ToolDataContainer.ToolData[DataIdentifier].GetMutablePtr<T>(); 
 		}
 
-		checkNoEntry();
 		return nullptr;
 	}
 	
@@ -163,8 +162,8 @@ public:
 		return Result;
 	}
 protected:
-	/** Try to load the default graph into ToolGraph. */
-	virtual void TryLoadDefaults();
+	/** Try to load the defaults (graph, actor class, ...) into the tool settings (ToolGraph, etc.). */
+	virtual bool TryLoadDefaults();
 	
 	virtual void Tick(float DeltaTime) override;
 	virtual bool IsTickable() const override;
@@ -194,10 +193,6 @@ public:
 	UPROPERTY(EditAnywhere, BlueprintReadOnly, Category = "PCG", meta = (NoResetToDefault, EditCondition="ToolGraph != nullptr", EditConditionHides, GetOptions="GetValidDataInstanceNamesForGraph"))
 	FName DataInstance = NAME_None;
 	
-	/** Automatically generate the PCG component while using the tool. */
-	UPROPERTY(EditAnywhere, Category = "PCG")
-	bool bAutoGenerate = true;
-
 	// @todo_pcg: Convert to a target object, rather than actor, to ease transition into SceneGraph
 	/** The selected actor to modify with the tool. */
 	UPROPERTY(EditAnywhere, Category = "PCG", meta = (TransientToolProperty))
@@ -235,4 +230,6 @@ public:
 	};
 	
 	mutable FGeneratedResources GeneratedResources;
+
+	bool bTryToGenerate = false;
 };
