[PCG] Always spawn Partition Actors in the persistent level
- Fixes issue where Partition Actors were spawned while a Level Instance was being edited, causing a crash when discarding/saving the edit as the PAs were destroyed at the same time

diff --git a/Engine/Plugins/PCG/Source/PCG/Private/RuntimeGen/PCGRuntimeGenScheduler.cpp b/Engine/Plugins/PCG/Source/PCG/Private/RuntimeGen/PCGRuntimeGenScheduler.cpp
--- a/Engine/Plugins/PCG/Source/PCG/Private/RuntimeGen/PCGRuntimeGenScheduler.cpp
+++ b/Engine/Plugins/PCG/Source/PCG/Private/RuntimeGen/PCGRuntimeGenScheduler.cpp
@@ -2017,6 +2017,9 @@ void FPCGRuntimeGenScheduler::AddPartitionActorPoolCount(int32 Count)
 	SpawnParams.ObjectFlags |= RF_Transient;
 	SpawnParams.ObjectFlags &= ~RF_Transactional;
 
+	// Important to override the level to make sure we spawn in the persistent level and not the editor's current editing level
+	SpawnParams.OverrideLevel = World->PersistentLevel;
+
 	// Create RuntimeGen PA pool.
 	for (int32 I = 0; I < Count; ++I)
 	{
diff --git a/Engine/Plugins/PCG/Source/PCG/Private/Subsystems/PCGSubsystem.cpp b/Engine/Plugins/PCG/Source/PCG/Private/Subsystems/PCGSubsystem.cpp
--- a/Engine/Plugins/PCG/Source/PCG/Private/Subsystems/PCGSubsystem.cpp
+++ b/Engine/Plugins/PCG/Source/PCG/Private/Subsystems/PCGSubsystem.cpp
@@ -1260,6 +1260,10 @@ APCGPartitionActor* UPCGSubsystem::FindOrCreatePCGPartitionActor(const FPCGGridD
 	}
 
 	FActorSpawnParameters SpawnParams;
+
+	// Important to override the level to make sure we spawn in the persistent level and not the editor's current editing level
+	SpawnParams.OverrideLevel = World->PersistentLevel;
+
 #if WITH_EDITOR
 	SpawnParams.Name = *APCGPartitionActor::GetPCGPartitionActorName(GridDescriptor, GridCoords);
 	SpawnParams.NameMode = FActorSpawnParameters::ESpawnActorNameMode::Required_Fatal;
