PCG Raycast Settings: Instanced Struct-based rule system that will determine how to filter while using raycast based tools such as the draw spline or paint tools

diff --git a/Engine/Plugins/PCG/Source/PCGEditor/PCGEditor.Build.cs b/Engine/Plugins/PCG/Source/PCGEditor/PCGEditor.Build.cs
--- a/Engine/Plugins/PCG/Source/PCGEditor/PCGEditor.Build.cs
+++ b/Engine/Plugins/PCG/Source/PCGEditor/PCGEditor.Build.cs
@@ -71,7 +71,8 @@ namespace UnrealBuildTool.Rules
 					"UnrealEd",
 					"WidgetRegistration", 
 					"ModelingComponents",
-					"ModelingComponentsEditorOnly"
+					"ModelingComponentsEditorOnly",
+					"Landscape"
 				});
 			
 			// @todo_pcg
diff --git a/Engine/Plugins/PCG/Source/PCGEditor/Private/Details/PCGRaycastFilterRuleCollectionDetails.cpp b/Engine/Plugins/PCG/Source/PCGEditor/Private/Details/PCGRaycastFilterRuleCollectionDetails.cpp
new file mode 100644
--- /dev/null
+++ b/Engine/Plugins/PCG/Source/PCGEditor/Private/Details/PCGRaycastFilterRuleCollectionDetails.cpp
@@ -0,0 +1,122 @@
+// Copyright Epic Games, Inc. All Rights Reserved.
+
+#include "Details/PCGRaycastFilterRuleCollectionDetails.h"
+#include "PCGModule.h"
+#include "Helpers/PCGHelpers.h"
+#include "EditorMode/Tools/Helpers/PCGEdModeSceneQueryHelpers.h"
+
+#include "DetailLayoutBuilder.h"
+#include "DetailWidgetRow.h"
+#include "IDetailChildrenBuilder.h"
+#include "IPropertyUtilities.h"
+#include "PropertyCustomizationHelpers.h"
+#include "ScopedTransaction.h"
+#include "Widgets/Input/SCheckBox.h"
+
+#define LOCTEXT_NAMESPACE "PCGEditorMode"
+
+void FPCGRaycastFilterRuleCollectionDetails::CustomizeHeader(TSharedRef<IPropertyHandle> InPropertyHandle, FDetailWidgetRow& HeaderRow, IPropertyTypeCustomizationUtils& CustomizationUtils)
+{
+	PropertyHandle = InPropertyHandle;
+	check(PropertyHandle.IsValid());
+
+	HeaderRow.Visibility(EVisibility::Collapsed);
+}
+
+void FPCGRaycastFilterRuleCollectionDetails::CustomizeChildren(TSharedRef<IPropertyHandle> InPropertyHandle, IDetailChildrenBuilder& ChildBuilder, IPropertyTypeCustomizationUtils& CustomizationUtils)
+{
+	if(TSharedPtr<IPropertyHandle> RaycastRuleHandle = InPropertyHandle->GetChildHandle(GET_MEMBER_NAME_CHECKED(FPCGRaycastFilterRuleCollection, RaycastRules)))
+	{
+		TSharedRef<FDetailArrayBuilder> RaycastRuleArrayBuilder = MakeShared<FDetailArrayBuilder>(RaycastRuleHandle.ToSharedRef(), true, true, true);
+		RaycastRuleArrayBuilder->OnGenerateArrayElementWidget(FOnGenerateArrayElementWidget::CreateSP(this, &FPCGRaycastFilterRuleCollectionDetails::OnGenerateRaycastRuleElement, CustomizationUtils.GetPropertyUtilities()));
+		
+		ChildBuilder.AddCustomBuilder(RaycastRuleArrayBuilder);
+	}
+}
+
+void FPCGRaycastFilterRuleCollectionDetails::OnGenerateRaycastRuleElement(TSharedRef<IPropertyHandle> RaycastRuleArrayElementHandle, int32 ArrayIndex,	IDetailChildrenBuilder& DetailChildrenBuilder, TSharedPtr<IPropertyUtilities> PropertyUtilities) const
+{
+	IDetailPropertyRow& ArrayElementRow = DetailChildrenBuilder.AddProperty(RaycastRuleArrayElementHandle);
+	ArrayElementRow.ShouldAutoExpand(true);
+	
+	FInstancedStruct* Instanced = nullptr;
+	if (RaycastRuleArrayElementHandle->GetValueData((void*&)Instanced) != FPropertyAccess::Success || !Instanced || !Instanced->IsValid())
+	{
+		return;
+	}
+
+	const UScriptStruct* InnerScript = Instanced->GetScriptStruct();
+	void* Memory = Instanced->GetMutableMemory();
+	if (!InnerScript || !Memory)
+	{
+		return;
+	}
+	
+	FBoolProperty* EnabledProp = FindFProperty<FBoolProperty>(InnerScript, GET_MEMBER_NAME_CHECKED(FPCGRaycastFilterRule, bEnabled));
+	auto GetEnabled = [EnabledProp, Memory]() -> bool
+	{
+		bool bEnabled = true;
+		if (EnabledProp)
+		{
+			bEnabled = EnabledProp->GetPropertyValue(EnabledProp->ContainerPtrToValuePtr<void>(Memory));
+		}
+			
+		return bEnabled;
+	};
+
+	auto SetEnabled = [EnabledProp, Memory](bool bEnabled) -> void
+	{
+		if (EnabledProp)
+		{
+			EnabledProp->SetPropertyValue(EnabledProp->ContainerPtrToValuePtr<void>(Memory), bEnabled);
+		}
+	};
+
+	TSharedPtr<SWidget> DefaultNameWidget;
+	TSharedPtr<SWidget> DefaultValueWidget;	
+	ArrayElementRow.GetDefaultWidgets(DefaultNameWidget, DefaultValueWidget, false);
+	
+	ArrayElementRow.CustomWidget(true)
+	.NameContent()
+	[
+		DefaultValueWidget.ToSharedRef()
+	]
+	.ValueContent()
+	[
+		SNew(SHorizontalBox)
+		+ SHorizontalBox::Slot()
+		.AutoWidth()
+		[
+			SNew(SCheckBox)
+			.IsChecked_Lambda([GetEnabled]() -> ECheckBoxState
+			{
+				return GetEnabled() ? ECheckBoxState::Checked : ECheckBoxState::Unchecked;
+			})
+			.OnCheckStateChanged_Lambda([SetEnabled, RaycastRuleArrayElementHandle, PropertyUtilities](ECheckBoxState InCheckBoxState)
+			{
+				FScopedTransaction Transaction(LOCTEXT("ToggleRuleEnabled", "Toggle Rule Enabled"));
+					
+				RaycastRuleArrayElementHandle->NotifyPreChange();
+				SetEnabled(InCheckBoxState == ECheckBoxState::Checked);
+				RaycastRuleArrayElementHandle->NotifyPostChange(EPropertyChangeType::ValueSet);
+
+				PropertyUtilities->ForceRefresh();
+			})
+		]
+	];
+}
+
+FPCGRaycastFilterRuleCollection* FPCGRaycastFilterRuleCollectionDetails::GetStruct()
+{
+	void* Data = nullptr;
+	FPropertyAccess::Result Result = PropertyHandle->GetValueData(Data);
+
+	return (Result == FPropertyAccess::Success) ? static_cast<FPCGRaycastFilterRuleCollection*>(Data) : nullptr;
+}
+
+const FPCGRaycastFilterRuleCollection* FPCGRaycastFilterRuleCollectionDetails::GetStruct() const
+{
+	return const_cast<FPCGRaycastFilterRuleCollectionDetails&>(*this).GetStruct();
+}
+
+#undef LOCTEXT_NAMESPACE
diff --git a/Engine/Plugins/PCG/Source/PCGEditor/Private/EditorMode/Tools/Helpers/PCGEdModeSceneQueryHelpers.cpp b/Engine/Plugins/PCG/Source/PCGEditor/Private/EditorMode/Tools/Helpers/PCGEdModeSceneQueryHelpers.cpp
--- a/Engine/Plugins/PCG/Source/PCGEditor/Private/EditorMode/Tools/Helpers/PCGEdModeSceneQueryHelpers.cpp
+++ b/Engine/Plugins/PCG/Source/PCGEditor/Private/EditorMode/Tools/Helpers/PCGEdModeSceneQueryHelpers.cpp
@@ -3,12 +3,16 @@
 #include "EditorMode/Tools/Helpers/PCGEdModeSceneQueryHelpers.h"
 
 #include "CollisionQueryParams.h"
+#include "Landscape.h"
 #include "BaseGizmos/GizmoMath.h"
 #include "Components/PrimitiveComponent.h"
+#include "Components/SkeletalMeshComponent.h"
+#include "Components/StaticMeshComponent.h"
 #include "Engine/HitResult.h"
 #include "Engine/World.h"
+#include "Helpers/PCGHelpers.h"
 
-namespace PCG::EditorMode::Scene
+namespace UE::PCG::EditorMode::Scene
 {
 	bool IsHitTargetVisible(const FHitResult& HitResult)
 	{
@@ -31,7 +35,7 @@ namespace PCG::EditorMode::Scene
 		return true;
 	}
 
-	TOptional<FHitResult> TraceToNearestObject(const FInputDeviceRay& Ray, const FViewRayParams& Params)
+	TOptional<FHitResult> TraceToNearestObject(const FRay& WorldRay, const FViewRayParams& Params)
 	{
 		if (!Params.World)
 		{
@@ -44,8 +48,8 @@ namespace PCG::EditorMode::Scene
 		QueryParams.bTraceComplex = true;
 
 		TArray<FHitResult> OutHits;
-		const FVector& Origin = Ray.WorldRay.Origin;
-		const FVector End = Origin + Ray.WorldRay.Direction * Params.Distance;
+		const FVector& Origin = WorldRay.Origin;
+		const FVector End = Origin + WorldRay.Direction * Params.Distance;
 		if (!Params.World->LineTraceMultiByObjectType(OutHits, Origin, End, ObjectQueryParams, QueryParams))
 		{
 			return {};
@@ -59,6 +63,15 @@ namespace PCG::EditorMode::Scene
 				continue;
 			}
 
+			if (Params.FilterRuleCollection != nullptr)
+			{
+				TOptional<bool> Result = Params.FilterRuleCollection->IsRaycastHitValid(Hit);
+				if (Result.IsSet() == false || Result.GetValue() == false)
+				{
+					continue;
+				}
+			}
+			
 			// Check for a visible result
 			bool bHitIsValid = IsHitTargetVisible(Hit);
 			// If invisible, check if it's explicitly allowed.
@@ -74,10 +87,15 @@ namespace PCG::EditorMode::Scene
 
 		return {};
 	}
+	
+	TOptional<FHitResult> TraceToNearestObject(const FInputDeviceRay& Ray, const FViewRayParams& Params)
+	{
+		return TraceToNearestObject(Ray.WorldRay, Params);
+	}
 
-	TOptional<FViewHitResult> ViewportRay(const FInputDeviceRay& Ray, const FViewRayParams& Params)
+	TOptional<FViewHitResult> ViewportRay(const FRay& WorldRay, const FViewRayParams& Params)
 	{
-		if (TOptional<FHitResult> Result = TraceToNearestObject(Ray, Params))
+		if (TOptional<FHitResult> Result = TraceToNearestObject(WorldRay, Params))
 		{
 			return FViewHitResult(Result->ImpactPoint, Result->ImpactNormal, Result->Distance);
 		}
@@ -86,6 +104,23 @@ namespace PCG::EditorMode::Scene
 			return {};
 		}
 	}
+	
+	TOptional<FViewHitResult> ViewportRay(const FInputDeviceRay& Ray, const FViewRayParams& Params)
+	{
+		return ViewportRay(Ray.WorldRay, Params);
+	}
+
+	TOptional<FHitResult> ViewportRay_HitResult(const FRay& WorldRay, const FViewRayParams& Params)
+	{
+		if (TOptional<FHitResult> Result = TraceToNearestObject(WorldRay, Params))
+		{
+			return Result;
+		}
+		else
+		{
+			return {};
+		}
+	}
 
 	TOptional<FViewHitResult> ViewportRayAgainstPlane(const FInputDeviceRay& Ray, const FViewRayParams& Params, const FPlane& Plane)
 	{
@@ -116,4 +151,137 @@ namespace PCG::EditorMode::Scene
 		const FVector PlaneNormal = -CameraState.Forward();
 		return ViewportRayAgainstPlane(Ray, Params, FPlane(PlaneOrigin, PlaneNormal));
 	}
-} // namespace PCG::EditorMode::Scene
+}
+
+FPCGRaycastFilterRule::FPCGRaycastFilterRule()
+{
+}
+
+FPCGRaycastFilterRule::~FPCGRaycastFilterRule()
+{
+}
+
+TOptional<bool> FPCGRaycastFilterRule_Landscape::IsRaycastHitValid(const FHitResult& HitResult) const
+{
+	if (HitResult.bBlockingHit)
+	{
+		if (ALandscape* LandscapeActor = Cast<ALandscape>(HitResult.GetActor()))
+		{
+			return true;
+		}
+	}
+
+	return {};
+}
+
+TOptional<bool> FPCGRaycastFilterRule_Meshes::IsRaycastHitValid(const FHitResult& HitResult) const
+{
+	if (HitResult.bBlockingHit)
+	{
+		UActorComponent* HitComponent = HitResult.GetComponent();
+		if (HitComponent && (HitComponent->IsA<UStaticMeshComponent>() || HitComponent->IsA<USkeletalMeshComponent>()))
+		{
+			return true;
+		}
+	}
+
+	return {};
+}
+
+TOptional<bool> FPCGRaycastFilterRule_IgnorePCGGeneratedComponents::IsRaycastHitValid(const FHitResult& HitResult) const
+{
+	if (HitResult.bBlockingHit)
+	{
+		UActorComponent* HitComponent = HitResult.GetComponent();
+		if (HitComponent && HitComponent->ComponentHasTag(PCGHelpers::DefaultPCGTag))
+		{
+			return false;
+		}
+	}
+
+	return {};
+}
+
+TOptional<bool> FPCGRaycastFilterRule_ConstrainToActor::IsRaycastHitValid(const FHitResult& HitResult) const
+{
+	if (SelectedActor == nullptr)
+	{
+		return {};
+	}
+	
+	AActor* HitActor = HitResult.GetActor();
+	if (HitActor && HitActor == SelectedActor)
+	{
+		return true;
+	}
+
+	return false;
+}
+
+TOptional<bool> FPCGRaycastFilterRule_AllowedClasses::IsRaycastHitValid(const FHitResult& HitResult) const
+{
+	if (AllowedClasses.Num() == 0)
+	{
+		return {};
+	}
+
+	AActor* HitActor = HitResult.GetActor();
+	if (UClass* HitActorClass = HitActor ? HitActor->GetClass() : nullptr)
+	{
+		if (AllowedClasses.Contains(HitActorClass))
+		{
+			return true;
+		}
+		else
+		{
+			for (UClass* Class : AllowedClasses)
+			{
+				if (HitActorClass->IsChildOf(Class))
+				{
+					return true;
+				}
+			}
+		}
+	}
+
+	return false;
+}
+
+TOptional<bool> FPCGRaycastFilterRuleCollection::IsRaycastHitValid(const FHitResult& RaycastHit) const
+{
+	TOptional<bool> Result;
+
+	for (const TInstancedStruct<FPCGRaycastFilterRule>& Rule : RaycastRules)
+	{
+		const FPCGRaycastFilterRule* RulePtr = Rule.GetPtr<FPCGRaycastFilterRule>();
+		if (RulePtr && RulePtr->bEnabled)
+		{
+			TOptional<bool> RuleResult = RulePtr->IsRaycastHitValid(RaycastHit);
+
+			// If there haven't been concrete results yet, we use the current rule's result 
+			if (Result.IsSet() == false)
+			{
+				Result = RuleResult;
+			}
+			// If there have been concrete results, we have to make sure all subsequent rules return true
+			else
+			{
+				if (Result.GetValue() == true)
+				{
+					if (RuleResult.IsSet())
+					{
+						Result = Result.GetValue() && RuleResult.GetValue();
+					}
+				}
+			}
+
+			// Early out if any rule has returned false
+			if (Result.IsSet() && Result.GetValue() == false)
+			{
+				return false;
+			}
+		}
+	}
+
+	return Result;
+}
diff --git a/Engine/Plugins/PCG/Source/PCGEditor/Private/EditorMode/Tools/Line/PCGDrawSplineTool.cpp b/Engine/Plugins/PCG/Source/PCGEditor/Private/EditorMode/Tools/Line/PCGDrawSplineTool.cpp
--- a/Engine/Plugins/PCG/Source/PCGEditor/Private/EditorMode/Tools/Line/PCGDrawSplineTool.cpp
+++ b/Engine/Plugins/PCG/Source/PCGEditor/Private/EditorMode/Tools/Line/PCGDrawSplineTool.cpp
@@ -18,7 +18,7 @@
 
 namespace PCGDrawSplineTool
 {
-	using namespace PCG::EditorMode;
+	using namespace UE::PCG::EditorMode;
 
 	// Helper classes for making undo/redo transactions
 	namespace CommandChange
@@ -288,10 +288,18 @@ void UPCGDrawSplineToolBase::Setup()
 	ClickOrDragBehavior->Initialize(this, this);
 	AddInputBehavior(ClickOrDragBehavior);
 
+	TInstancedStruct<FPCGRaycastFilterRule_Landscape> LandscapeRule = TInstancedStruct<FPCGRaycastFilterRule_Landscape>::Make();
+	TInstancedStruct<FPCGRaycastFilterRule_Meshes> MeshesRule = TInstancedStruct<FPCGRaycastFilterRule_Meshes>::Make();
+	TInstancedStruct<FPCGRaycastFilterRule_IgnorePCGGeneratedComponents> IgnorePCGGeneratedComponentsRule = TInstancedStruct<FPCGRaycastFilterRule_IgnorePCGGeneratedComponents>::Make();
+	ToolSettings->RaycastRuleCollection.RaycastRules.Add(LandscapeRule);
+	ToolSettings->RaycastRuleCollection.RaycastRules.Add(MeshesRule);
+	ToolSettings->RaycastRuleCollection.RaycastRules.Add(IgnorePCGGeneratedComponentsRule);
+	
 	// @todo_pcg: Make adjustments to this and check when properties change.
 	CachedRayParams =
 	{
-		.World = GetWorld()
+		.World = GetWorld(),
+		.FilterRuleCollection = &ToolSettings->RaycastRuleCollection
 	};
 }
 
@@ -389,8 +397,9 @@ FVector UPCGDrawSplineToolBase::GetNewPointNormal(const FVector& HitNormal) cons
 
 FInputRayHit UPCGDrawSplineToolBase::IsHitByClick(const FInputDeviceRay& ClickPosition)
 {
-	using namespace PCG::EditorMode;
-	if (TOptional<Scene::FViewHitResult> Result = ViewportRay(ClickPosition, CachedRayParams))
+	using namespace UE::PCG::EditorMode::Scene;
+	
+	if (TOptional<FViewHitResult> Result = ViewportRay(ClickPosition, CachedRayParams))
 	{
 		return Result->ToInputRayHit();
 	}
@@ -402,9 +411,9 @@ void UPCGDrawSplineToolBase::OnClicked(const FInputDeviceRay& ClickPosition)
 {
 	using namespace PCGDrawSplineTool;
 
-	if (TOptional<Scene::FViewHitResult> Result = ViewportRay(ClickPosition, CachedRayParams))
+	if (TOptional<Scene::FViewHitResult> Result = Scene::ViewportRay(ClickPosition, CachedRayParams))
 	{
-	UPCGInteractiveToolSettings_SplineBase* ToolSettings = GetSplineSettings();
+		UPCGInteractiveToolSettings_SplineBase* ToolSettings = GetSplineSettings();
 		
 		USplineComponent* WorkingSpline = ToolSettings->GetWorkingSplineComponent(ToolSettings->DataInstance);
 		switch (ToolSettings->DrawMode)
@@ -448,7 +457,7 @@ void UPCGDrawSplineToolBase::OnClicked(const FInputDeviceRay& ClickPosition)
 
 FInputRayHit UPCGDrawSplineToolBase::CanBeginClickDragSequence(const FInputDeviceRay& DragPosition)
 {
-	using namespace PCG::EditorMode::Scene;
+	using namespace UE::PCG::EditorMode::Scene;
 	if (TOptional<FViewHitResult> Result = ViewportRay(DragPosition, CachedRayParams))
 	{
 		return Result->ToInputRayHit();
@@ -459,7 +468,7 @@ FInputRayHit UPCGDrawSplineToolBase::CanBeginClickDragSequence(const FInputDevic
 
 void UPCGDrawSplineToolBase::OnClickPress(const FInputDeviceRay& ClickPosition)
 {
-	using namespace PCG::EditorMode::Scene;
+	using namespace UE::PCG::EditorMode::Scene;
 	if (TOptional<FViewHitResult> Result = ViewportRay(ClickPosition, CachedRayParams))
 	{
 		UPCGInteractiveToolSettings_SplineBase* ToolSettings = GetSplineSettings();
@@ -476,7 +485,7 @@ void UPCGDrawSplineToolBase::OnClickPress(const FInputDeviceRay& ClickPosition)
 
 void UPCGDrawSplineToolBase::OnClickDrag(const FInputDeviceRay& DragPosition)
 {
-	using namespace PCG::EditorMode::Scene;
+	using namespace UE::PCG::EditorMode::Scene;
 
 	UPCGInteractiveToolSettings_SplineBase* ToolSettings = GetSplineSettings();
 	
diff --git a/Engine/Plugins/PCG/Source/PCGEditor/Private/EditorMode/Tools/Paint/PCGPaintTool.cpp b/Engine/Plugins/PCG/Source/PCGEditor/Private/EditorMode/Tools/Paint/PCGPaintTool.cpp
--- a/Engine/Plugins/PCG/Source/PCGEditor/Private/EditorMode/Tools/Paint/PCGPaintTool.cpp
+++ b/Engine/Plugins/PCG/Source/PCGEditor/Private/EditorMode/Tools/Paint/PCGPaintTool.cpp
@@ -6,6 +6,7 @@
 
 #include "InteractiveGizmoManager.h"
 #include "InteractiveToolManager.h"
+#include "LandscapeComponent.h"
 #include "ToolTargetManager.h"
 #include "TargetInterfaces/PrimitiveComponentBackedTarget.h"
 #include "ToolTargets/PrimitiveComponentToolTarget.h"
@@ -53,6 +54,19 @@ void UPCGPaintTool::Setup()
 	{
 		ToolSettings->GeneratedResources.SpawnedActor->SetActorLocation(TargetActor->GetActorLocation());
 	}
+
+	TInstancedStruct<FPCGRaycastFilterRule_Landscape> LandscapeRule = TInstancedStruct<FPCGRaycastFilterRule_Landscape>::Make();
+	TInstancedStruct<FPCGRaycastFilterRule_Meshes> MeshesRule = TInstancedStruct<FPCGRaycastFilterRule_Meshes>::Make();
+	TInstancedStruct<FPCGRaycastFilterRule_IgnorePCGGeneratedComponents> IgnorePCGGeneratedComponentsRule = TInstancedStruct<FPCGRaycastFilterRule_IgnorePCGGeneratedComponents>::Make();
+	ToolSettings->RaycastRuleCollection.RaycastRules.Add(LandscapeRule);
+	ToolSettings->RaycastRuleCollection.RaycastRules.Add(MeshesRule);
+	ToolSettings->RaycastRuleCollection.RaycastRules.Add(IgnorePCGGeneratedComponentsRule);
+	
+	CachedRayParams =
+	{
+		.World = GetWorld(),
+		.FilterRuleCollection = &ToolSettings->RaycastRuleCollection
+	};
 	
 	Super::Setup();	
 }
@@ -123,6 +137,18 @@ void UPCGPaintTool::OnTick(float DeltaTime)
 	bShiftPressedLastTick = bShiftPressedThisTick;
 }
 
+bool UPCGPaintTool::HitTest(const FRay& Ray, FHitResult& OutHit)
+{
+	using namespace UE::PCG::EditorMode::Scene;
+	if (TOptional<FHitResult> Result = TraceToNearestObject(Ray, CachedRayParams))
+	{
+		OutHit = Result.GetValue();
+		return Result->bBlockingHit;
+	}
+
+	return false;
+}
+
 void UPCGPaintTool::OnClickPress(const FInputDeviceRay& PressPos)
 {
 	Super::OnClickPress(PressPos);
diff --git a/Engine/Plugins/PCG/Source/PCGEditor/Private/EditorMode/Tools/Volume/PCGCreateVolumeTool.cpp b/Engine/Plugins/PCG/Source/PCGEditor/Private/EditorMode/Tools/Volume/PCGCreateVolumeTool.cpp
--- a/Engine/Plugins/PCG/Source/PCGEditor/Private/EditorMode/Tools/Volume/PCGCreateVolumeTool.cpp
+++ b/Engine/Plugins/PCG/Source/PCGEditor/Private/EditorMode/Tools/Volume/PCGCreateVolumeTool.cpp
@@ -79,7 +79,7 @@ void UPCGCreateVolumeTool::Render(IToolsContextRenderAPI* RenderAPI)
 
 FInputRayHit UPCGCreateVolumeTool::IsHitByClick(const FInputDeviceRay& ClickPosition)
 {
-	using namespace PCG::EditorMode::Scene;
+	using namespace UE::PCG::EditorMode::Scene;
 	if (TOptional<FViewHitResult> Result = ViewportRay(ClickPosition, CachedRayParams))
 	{
 		return Result->ToInputRayHit();
diff --git a/Engine/Plugins/PCG/Source/PCGEditor/Private/PCGEditorModule.cpp b/Engine/Plugins/PCG/Source/PCGEditor/Private/PCGEditorModule.cpp
--- a/Engine/Plugins/PCG/Source/PCGEditor/Private/PCGEditorModule.cpp
+++ b/Engine/Plugins/PCG/Source/PCGEditor/Private/PCGEditorModule.cpp
@@ -70,6 +70,7 @@
 #include "PropertyEditorModule.h"
 #include "ScopedTransaction.h"
 #include "ToolMenus.h"
+#include "Details/PCGRaycastFilterRuleCollectionDetails.h"
 #include "Editor/SceneOutliner/Public/ISceneOutliner.h"
 #include "EditorMode/PCGEdMode.h"
 #include "EditorMode/Tools/PCGInteractiveToolSettings.h"
@@ -327,7 +328,8 @@ void FPCGEditorModule::RegisterDetailsCustomizations()
 	PropertyEditor.RegisterCustomPropertyTypeLayout("PCGOverrideInstancedPropertyBag", FOnGetPropertyTypeCustomizationInstance::CreateStatic(&FPCGOverrideInstancedPropertyBagDetails::MakeInstance));
 	PropertyEditor.RegisterCustomPropertyTypeLayout("EnumSelector", FOnGetPropertyTypeCustomizationInstance::CreateStatic(&FEnumSelectorDetails::MakeInstance));
 	PropertyEditor.RegisterCustomPropertyTypeLayout("PCGDataTypeIdentifier", FOnGetPropertyTypeCustomizationInstance::CreateStatic(&FPCGDataTypeIdentifierDetails::MakeInstance));
-
+	PropertyEditor.RegisterCustomPropertyTypeLayout("PCGRaycastFilterRuleCollection", FOnGetPropertyTypeCustomizationInstance::CreateStatic(&FPCGRaycastFilterRuleCollectionDetails::MakeInstance));
+	
 	PropertyEditor.NotifyCustomizationModuleChanged();
 }
 
@@ -345,6 +347,7 @@ void FPCGEditorModule::UnregisterDetailsCustomizations()
 		PropertyModule->UnregisterCustomClassLayout("PCGComponent");
 		PropertyModule->UnregisterCustomClassLayout("PCGBlueprintSettings");
 
+		PropertyModule->UnregisterCustomPropertyTypeLayout("PCGRaycastFilterRuleCollection");
 		PropertyModule->UnregisterCustomPropertyTypeLayout("PCGDataTypeIdentifier");
 		PropertyModule->UnregisterCustomPropertyTypeLayout("EnumSelector");
 		PropertyModule->UnregisterCustomPropertyTypeLayout("PCGOverrideInstancedPropertyBag");
diff --git a/Engine/Plugins/PCG/Source/PCGEditor/Public/Details/PCGRaycastFilterRuleCollectionDetails.h b/Engine/Plugins/PCG/Source/PCGEditor/Public/Details/PCGRaycastFilterRuleCollectionDetails.h
new file mode 100644
--- /dev/null
+++ b/Engine/Plugins/PCG/Source/PCGEditor/Public/Details/PCGRaycastFilterRuleCollectionDetails.h
@@ -0,0 +1,30 @@
+// Copyright Epic Games, Inc. All Rights Reserved.
+
+#pragma once
+
+#include "IPropertyTypeCustomization.h"
+
+struct FPCGRaycastFilterRuleCollection;
+class SWidget;
+
+class FPCGRaycastFilterRuleCollectionDetails : public IPropertyTypeCustomization
+{
+public:
+	static TSharedRef<IPropertyTypeCustomization> MakeInstance()
+	{
+		return MakeShared<FPCGRaycastFilterRuleCollectionDetails>();
+	}
+
+	/** ~Begin IPropertyTypeCustomization interface */
+	virtual void CustomizeHeader(TSharedRef<IPropertyHandle> InPropertyHandle, FDetailWidgetRow& HeaderRow, IPropertyTypeCustomizationUtils& CustomizationUtils) override;
+	virtual void CustomizeChildren(TSharedRef<IPropertyHandle> InPropertyHandle, IDetailChildrenBuilder& ChildBuilder, IPropertyTypeCustomizationUtils& CustomizationUtils) override;
+	/** ~End IPropertyTypeCustomization interface */
+
+protected:
+	void OnGenerateRaycastRuleElement(TSharedRef<IPropertyHandle> RaycastRuleArrayElementHandle, int32 ArrayIndex, IDetailChildrenBuilder& DetailChildrenBuilder, TSharedPtr<IPropertyUtilities> PropertyUtilities) const;
+	
+	FPCGRaycastFilterRuleCollection* GetStruct();
+	const FPCGRaycastFilterRuleCollection* GetStruct() const;
+
+	TSharedPtr<IPropertyHandle> PropertyHandle;
+};
diff --git a/Engine/Plugins/PCG/Source/PCGEditor/Public/EditorMode/Tools/Helpers/PCGEdModeSceneQueryHelpers.h b/Engine/Plugins/PCG/Source/PCGEditor/Public/EditorMode/Tools/Helpers/PCGEdModeSceneQueryHelpers.h
--- a/Engine/Plugins/PCG/Source/PCGEditor/Public/EditorMode/Tools/Helpers/PCGEdModeSceneQueryHelpers.h
+++ b/Engine/Plugins/PCG/Source/PCGEditor/Public/EditorMode/Tools/Helpers/PCGEdModeSceneQueryHelpers.h
@@ -5,11 +5,95 @@
 #include "InputState.h"
 #include "Math/Vector.h"
 #include "Misc/Optional.h"
+#include "StructUtils/InstancedStruct.h"
+#include "Templates/SubclassOf.h"
+#include "GameFramework/Actor.h"
+#include "PCGEdModeSceneQueryHelpers.generated.h"
 
 struct FViewCameraState;
 struct FHitResult;
+class ALandscape;
 
-namespace PCG::EditorMode::Scene
+USTRUCT()
+struct FPCGRaycastFilterRule
+{
+	GENERATED_BODY()
+
+	PCGEDITOR_API FPCGRaycastFilterRule();
+	PCGEDITOR_API virtual ~FPCGRaycastFilterRule();
+	
+	/** If this returns false, the raycast hit is filtered out.
+	 *  If it returns true, this rule has determined is valid under this given rule.
+	 *  If it returns an empty TOptional, this rule doesn't apply to the hit.
+	 */
+	virtual TOptional<bool> IsRaycastHitValid(const FHitResult& HitResult) const { return {}; }
+
+	/** Whether this rule is enabled or not. Not exposed as editing is done via custom UI. */
+	UPROPERTY()
+	bool bEnabled = true;
+};
+
+USTRUCT(DisplayName="Landscape")
+struct FPCGRaycastFilterRule_Landscape : public FPCGRaycastFilterRule
+{
+	GENERATED_BODY()
+
+	virtual TOptional<bool> IsRaycastHitValid(const FHitResult& HitResult) const override;
+};
+
+USTRUCT(DisplayName="Meshes")
+struct FPCGRaycastFilterRule_Meshes : public FPCGRaycastFilterRule
+{
+	GENERATED_BODY()
+
+	virtual TOptional<bool> IsRaycastHitValid(const FHitResult& HitResult) const override;
+};
+
+USTRUCT(DisplayName="Ignore PCG Components")
+struct FPCGRaycastFilterRule_IgnorePCGGeneratedComponents : public FPCGRaycastFilterRule
+{
+	GENERATED_BODY()
+
+	virtual TOptional<bool> IsRaycastHitValid(const FHitResult& HitResult) const override;
+};
+
+USTRUCT(DisplayName="Constrain to Actor")
+struct FPCGRaycastFilterRule_ConstrainToActor : public FPCGRaycastFilterRule
+{
+	GENERATED_BODY()
+
+	virtual TOptional<bool> IsRaycastHitValid(const FHitResult& HitResult) const override;
+
+	UPROPERTY(EditAnywhere, Category= "Raycast")
+	TObjectPtr<AActor> SelectedActor;
+};
+
+/** The Ray will only be valid when the hit actor belongs to any of the specified classes. */
+USTRUCT(DisplayName="Allowed Classes")
+struct FPCGRaycastFilterRule_AllowedClasses : public FPCGRaycastFilterRule
+{
+	GENERATED_BODY()
+
+	virtual TOptional<bool> IsRaycastHitValid(const FHitResult& HitResult) const override;
+
+	UPROPERTY(EditAnywhere, Category= "Raycast")
+	TSet<TSubclassOf<AActor>> AllowedClasses;
+};
+
+/** A collection struct used for details customization purposes. This allows us to add the collection to any PCGInteractiveToolSettings class where required. */
+USTRUCT()
+struct FPCGRaycastFilterRuleCollection
+{
+	GENERATED_BODY()
+
+	TOptional<bool> IsRaycastHitValid(const FHitResult& RaycastHit) const;
+
+	/** The raycast rules for a given tool. */
+	UPROPERTY(EditAnywhere, Category = "Raycast", meta=(ExcludeBaseStruct))
+	TArray<TInstancedStruct<FPCGRaycastFilterRule>> RaycastRules;
+};
+
+namespace UE::PCG::EditorMode::Scene
 {
 	/** Custom parameters for UI ray casts from the current view into the scene. */
 	struct FViewRayParams
@@ -18,6 +102,7 @@ namespace PCG::EditorMode::Scene
 		double Distance = UE_DOUBLE_BIG_NUMBER;
 		const TArray<const UPrimitiveComponent*>* ComponentsToIgnore = nullptr;
 		const TArray<const UPrimitiveComponent*>* AllowedInvisibleComponents = nullptr;
+		const FPCGRaycastFilterRuleCollection* FilterRuleCollection;
 	};
 
 	/** The resulting data from a UI ray cast hit. */
@@ -45,9 +130,11 @@ namespace PCG::EditorMode::Scene
 	bool IsHitTargetVisible(const FHitResult& HitResult);
 
 	/** Cast a ray and if a hit is recorded, return the hit result nearest object. */
+	TOptional<FHitResult> TraceToNearestObject(const FRay& WorldRay, const FViewRayParams& Params);
 	TOptional<FHitResult> TraceToNearestObject(const FInputDeviceRay& Ray, const FViewRayParams& Params);
 
 	/** Cast a ray from the view and if a hit is recorded, return a condensed view-related hit result. */
+	TOptional<FViewHitResult> ViewportRay(const FRay& WorldRay, const FViewRayParams& Params);
 	TOptional<FViewHitResult> ViewportRay(const FInputDeviceRay& Ray, const FViewRayParams& Params);
 
 	/** Cast a ray from the view to a plane in space. If a hit is recorded, return a condensed view-related hit result. */
@@ -55,4 +142,4 @@ namespace PCG::EditorMode::Scene
 
 	/** Cast a ray from the view to a plane parallel to the user view. If a hit is recorded, return a condensed view-related hit result. */
 	TOptional<FViewHitResult> ViewportRayAgainstCameraPlane(const FInputDeviceRay& Ray, const FViewRayParams& Params, const FViewCameraState& CameraState);
-} // namespace PCG::EditorMode::Scene
+}
diff --git a/Engine/Plugins/PCG/Source/PCGEditor/Public/EditorMode/Tools/Line/PCGDrawSplineTool.h b/Engine/Plugins/PCG/Source/PCGEditor/Public/EditorMode/Tools/Line/PCGDrawSplineTool.h
--- a/Engine/Plugins/PCG/Source/PCGEditor/Public/EditorMode/Tools/Line/PCGDrawSplineTool.h
+++ b/Engine/Plugins/PCG/Source/PCGEditor/Public/EditorMode/Tools/Line/PCGDrawSplineTool.h
@@ -51,6 +51,9 @@ public:
 
 	UPROPERTY(EditAnywhere, Category = "PCG", meta = (TransientToolProperty))
 	FName SplineName;
+
+	UPROPERTY(EditAnywhere, Category= "Raycast")
+	FPCGRaycastFilterRuleCollection RaycastRuleCollection;
 };
 
 UCLASS(BlueprintType)
@@ -108,7 +111,7 @@ public:
 	virtual void OnClickRelease(const FInputDeviceRay& ReleasePosition) override;
 	virtual void OnTerminateDragSequence() override;
 	// ~End IClickDragBehaviorTarget interface
-
+	
 	UPCGInteractiveToolSettings_SplineBase* GetSplineSettings() const;
 private:
 	void AddSplinePoint(const FVector& HitLocation, const FVector& HitNormal) const;
@@ -117,7 +120,7 @@ private:
 	UPROPERTY(Instanced)
 	TObjectPtr<USingleClickOrDragInputBehavior> ClickOrDragBehavior = nullptr;
 	
-	PCG::EditorMode::Scene::FViewRayParams CachedRayParams{};
+	UE::PCG::EditorMode::Scene::FViewRayParams CachedRayParams{};
 	bool bDrawTangentForPreviousPoint = false;
 	bool bFreeDrawPlacedPreviewPoint = false;
 	int32 FreeDrawStrokeStartIndex = 0;
diff --git a/Engine/Plugins/PCG/Source/PCGEditor/Public/EditorMode/Tools/Paint/PCGPaintTool.h b/Engine/Plugins/PCG/Source/PCGEditor/Public/EditorMode/Tools/Paint/PCGPaintTool.h
--- a/Engine/Plugins/PCG/Source/PCGEditor/Public/EditorMode/Tools/Paint/PCGPaintTool.h
+++ b/Engine/Plugins/PCG/Source/PCGEditor/Public/EditorMode/Tools/Paint/PCGPaintTool.h
@@ -6,6 +6,7 @@
 #include "BaseBehaviors/KeyAsModifierInputBehavior.h"
 #include "EditorMode/Tools/PCGInteractiveToolBuilder.h"
 #include "EditorMode/Tools/PCGInteractiveToolSettings.h"
+#include "ToolTargets/PrimitiveComponentToolTarget.h"
 #include "PCGPaintTool.generated.h"
 
 UCLASS()
@@ -24,6 +25,9 @@ public:
 	/** If false, will erase points on all paint data instances. */
 	UPROPERTY(BlueprintReadOnly, EditAnywhere, Category = "Paint")
 	bool bErasePointsOfSelectedDataInstance = true;
+
+	UPROPERTY(EditAnywhere, Category= "Raycast")
+	FPCGRaycastFilterRuleCollection RaycastRuleCollection;
 };
 
 UCLASS(Transient)
@@ -41,6 +45,8 @@ public:
 	virtual void OnTick(float DeltaTime) override;
 	// ~End UPCGInteractiveToolSettings interface
 
+	virtual bool HitTest(const FRay& Ray, FHitResult& OutHit) override;
+	
 protected:
 	virtual void OnClickPress(const FInputDeviceRay& PressPos) override;
 	virtual void OnUpdateDrag(const FRay& Ray) override;
@@ -51,9 +57,11 @@ protected:
 	
 	/** This will be used as the brush base, if not specifying a radius directly */
 	virtual double EstimateMaximumTargetDimension() override;
+	
 private:
 	void AddSinglePointFromBrush();
 	void TryDeletePointsFromBrush();
+	
 private:
 	UPROPERTY(Transient)
 	TObjectPtr<UPCGInteractiveToolSettings_PaintTool> ToolSettings;
@@ -63,6 +71,8 @@ private:
 
 	UPROPERTY(Transient)
 	bool bShiftPressedLastTick = false;
+
+	UE::PCG::EditorMode::Scene::FViewRayParams CachedRayParams{};
 };
 
 UCLASS(Transient)
diff --git a/Engine/Plugins/PCG/Source/PCGEditor/Public/EditorMode/Tools/Volume/PCGCreateVolumeTool.h b/Engine/Plugins/PCG/Source/PCGEditor/Public/EditorMode/Tools/Volume/PCGCreateVolumeTool.h
--- a/Engine/Plugins/PCG/Source/PCGEditor/Public/EditorMode/Tools/Volume/PCGCreateVolumeTool.h
+++ b/Engine/Plugins/PCG/Source/PCGEditor/Public/EditorMode/Tools/Volume/PCGCreateVolumeTool.h
@@ -65,7 +65,7 @@ public:
 protected:
 	void ApplyTool() const;
 private:
-	PCG::EditorMode::Scene::FViewRayParams CachedRayParams;
+	UE::PCG::EditorMode::Scene::FViewRayParams CachedRayParams;
 
 	UPROPERTY(Transient)
 	TObjectPtr<UPCGInteractiveToolSettings_CreateVolume> ToolSettings;
