Fixed uninitialized ptr

#rnx

diff --git a/Engine/Plugins/PCG/Source/PCGEditor/Public/EditorMode/Tools/Helpers/PCGEdModeSceneQueryHelpers.h b/Engine/Plugins/PCG/Source/PCGEditor/Public/EditorMode/Tools/Helpers/PCGEdModeSceneQueryHelpers.h
--- a/Engine/Plugins/PCG/Source/PCGEditor/Public/EditorMode/Tools/Helpers/PCGEdModeSceneQueryHelpers.h
+++ b/Engine/Plugins/PCG/Source/PCGEditor/Public/EditorMode/Tools/Helpers/PCGEdModeSceneQueryHelpers.h
@@ -102,7 +102,7 @@ namespace UE::PCG::EditorMode::Scene
 		double Distance = UE_DOUBLE_BIG_NUMBER;
 		const TArray<const UPrimitiveComponent*>* ComponentsToIgnore = nullptr;
 		const TArray<const UPrimitiveComponent*>* AllowedInvisibleComponents = nullptr;
-		const FPCGRaycastFilterRuleCollection* FilterRuleCollection;
+		const FPCGRaycastFilterRuleCollection* FilterRuleCollection = nullptr;
 	};
 
 	/** The resulting data from a UI ray cast hit. */
