[PCG] Make GPU header files public to prepare for move to beta

#rnx

diff --git a/Engine/Plugins/PCG/Source/PCG/Internal/Compute/Data/PCGProxyForGPUData.h b/Engine/Plugins/PCG/Source/PCG/Internal/Compute/Data/PCGProxyForGPUData.h
deleted file mode 100644
diff --git a/Engine/Plugins/PCG/Source/PCG/Internal/Compute/PCGComputeCommon.h b/Engine/Plugins/PCG/Source/PCG/Internal/Compute/PCGComputeCommon.h
deleted file mode 100644
diff --git a/Engine/Plugins/PCG/Source/PCG/Internal/Compute/PCGDataForGPU.h b/Engine/Plugins/PCG/Source/PCG/Internal/Compute/PCGDataForGPU.h
deleted file mode 100644
diff --git a/Engine/Plugins/PCG/Source/PCG/Private/Compute/PCGComputeCommon.cpp b/Engine/Plugins/PCG/Source/PCG/Private/Compute/PCGComputeCommon.cpp
--- a/Engine/Plugins/PCG/Source/PCG/Private/Compute/PCGComputeCommon.cpp
+++ b/Engine/Plugins/PCG/Source/PCG/Private/Compute/PCGComputeCommon.cpp
@@ -422,3 +422,9 @@ uint32 GetTypeHash(const FPCGPinReference& In)
 {
 	return HashCombine(/*GetTypeHash(In.TaskId),*/ PointerHash(In.Kernel), GetTypeHash(In.Label));
 }
+
+uint32 GetTypeHash(const FPCGKernelPin& In)
+{
+	return HashCombine(GetTypeHash(In.KernelIndex), GetTypeHash(In.PinLabel), GetTypeHash(In.bIsInput));
+}
+
diff --git a/Engine/Plugins/PCG/Source/PCG/Private/Compute/PCGComputeGraph.cpp b/Engine/Plugins/PCG/Source/PCG/Private/Compute/PCGComputeGraph.cpp
--- a/Engine/Plugins/PCG/Source/PCG/Private/Compute/PCGComputeGraph.cpp
+++ b/Engine/Plugins/PCG/Source/PCG/Private/Compute/PCGComputeGraph.cpp
@@ -18,11 +18,6 @@
 
 #define PCG_DATA_DESCRIPTION_LOGGING 0
 
-uint32 GetTypeHash(const FPCGKernelPin& In)
-{
-	return HashCombine(GetTypeHash(In.KernelIndex), GetTypeHash(In.PinLabel), GetTypeHash(In.bIsInput));
-}
-
 static FPCGCompilerDiagnostic ProcessCompilationMessage(const FComputeKernelCompileMessage& InMessage)
 {
 	FPCGCompilerDiagnostic Diagnostic;
diff --git a/Engine/Plugins/PCG/Source/PCG/Private/Compute/PCGComputeGraph.h b/Engine/Plugins/PCG/Source/PCG/Private/Compute/PCGComputeGraph.h
--- a/Engine/Plugins/PCG/Source/PCG/Private/Compute/PCGComputeGraph.h
+++ b/Engine/Plugins/PCG/Source/PCG/Private/Compute/PCGComputeGraph.h
@@ -19,37 +19,6 @@ class UPCGNode;
 class UPCGPin;
 class UPCGStaticMeshSpawnerKernel;
 
-/** An input or output pin of a kernel. Compute graph does not internally have 'pins' so this is useful for mapping between kernel data and PCG pins. */
-USTRUCT()
-struct FPCGKernelPin
-{
-	GENERATED_BODY()
-
-public:
-	FPCGKernelPin() = default;
-
-	explicit FPCGKernelPin(int32 InKernelIndex, FName InPinLabel, bool bInIsInput)
-		: KernelIndex(InKernelIndex)
-		, PinLabel(InPinLabel)
-		, bIsInput(bInIsInput)
-	{
-	}
-
-	bool operator==(const FPCGKernelPin& Other) const = default;
-
-public:
-	UPROPERTY()
-	int32 KernelIndex = INDEX_NONE;
-
-	UPROPERTY()
-	FName PinLabel = NAME_None;
-
-	UPROPERTY()
-	bool bIsInput = false;
-
-	friend uint32 GetTypeHash(const FPCGKernelPin& In);
-};
-
 UCLASS()
 class UPCGComputeGraph : public UComputeGraph
 {
diff --git a/Engine/Plugins/PCG/Source/PCG/Private/Compute/PCGComputeKernel.h b/Engine/Plugins/PCG/Source/PCG/Private/Compute/PCGComputeKernel.h
deleted file mode 100644
diff --git a/Engine/Plugins/PCG/Source/PCG/Private/Compute/PCGDataBinding.h b/Engine/Plugins/PCG/Source/PCG/Private/Compute/PCGDataBinding.h
deleted file mode 100644
diff --git a/Engine/Plugins/PCG/Source/PCG/Private/Compute/PCGKernelHelpers.h b/Engine/Plugins/PCG/Source/PCG/Private/Compute/PCGKernelHelpers.h
deleted file mode 100644
diff --git a/Engine/Plugins/PCG/Source/PCG/Private/Compute/PCGKernelParamsCache.h b/Engine/Plugins/PCG/Source/PCG/Private/Compute/PCGKernelParamsCache.h
deleted file mode 100644
diff --git a/Engine/Plugins/PCG/Source/PCG/Private/Compute/PCGPinPropertiesGPU.h b/Engine/Plugins/PCG/Source/PCG/Private/Compute/PCGPinPropertiesGPU.h
deleted file mode 100644
diff --git a/Engine/Plugins/PCG/Source/PCG/Private/Graph/PCGGPUGraphCompilationContext.h b/Engine/Plugins/PCG/Source/PCG/Private/Graph/PCGGPUGraphCompilationContext.h
deleted file mode 100644
diff --git a/Engine/Plugins/PCG/Source/PCG/Public/Compute/Data/PCGProxyForGPUData.h b/Engine/Plugins/PCG/Source/PCG/Public/Compute/Data/PCGProxyForGPUData.h
new file mode 100644
--- /dev/null
+++ b/Engine/Plugins/PCG/Source/PCG/Public/Compute/Data/PCGProxyForGPUData.h
@@ -0,0 +1,83 @@
+// Copyright Epic Games, Inc. All Rights Reserved.
+
+#pragma once
+
+#include "PCGData.h"
+
+#include "Compute/PCGDataForGPU.h"
+
+#include "PCGProxyForGPUData.generated.h"
+
+#define UE_API PCG_API
+
+struct FPCGContext;
+
+USTRUCT()
+struct FPCGDataTypeInfoProxyForGPU : public FPCGDataTypeInfo
+{
+	GENERATED_BODY()
+
+	PCG_DECLARE_LEGACY_TYPE_INFO(EPCGDataType::ProxyForGPU)
+
+#if WITH_EDITOR
+	virtual bool Hidden() const override { return true; }
+#endif // WITH_EDITOR
+};
+
+/** A proxy for data residing on the GPU with functionality to read the data back to the CPU. */
+UCLASS(MinimalAPI, ClassGroup = (Procedural), DisplayName = "GPU Proxy")
+class UPCGProxyForGPUData : public UPCGData
+{
+	GENERATED_BODY()
+
+public:
+	UE_API void Initialize(TSharedPtr<FPCGProxyForGPUDataCollection> InDataCollection, int InDataIndexInCollection);
+
+	//~Begin UPCGData interface
+	PCG_ASSIGN_DEFAULT_TYPE_INFO(FPCGDataTypeInfoProxyForGPU)
+
+	UE_API virtual void AddToCrc(FArchiveCrc32& Ar, bool bFullDataCrc) const override;
+	virtual bool CanBeSerialized() const override { return false; }
+	virtual bool HoldsTransientResources() const override { return true; }
+	virtual bool IsCacheable() const override { return false; }
+	UE_API virtual void ReleaseTransientResources(const TCHAR* InReason = nullptr);
+	UE_API virtual EPCGDataType GetUnderlyingDataType() const override;
+	UE_API virtual FPCGDataTypeBaseId GetUnderlyingDataTypeId() const override;
+	//~End UPCGData interface
+
+	/** Returns the GPU info. Returns null if the buffer has been discarded. */
+	UE_API TSharedPtr<const FPCGProxyForGPUDataCollection> GetInputDataCollectionInfo() const;
+
+	int GetDataIndexInCollection() const { return DataIndexInCollection; }
+
+	struct FReadbackResult
+	{
+		/** Set false until read back has been performed. */
+		bool bComplete = false;
+
+		/** The data created from readback, null if read back failed. */
+		FPCGTaggedData TaggedData;
+	};
+
+	/** Populates a CPU data object representing the GPU data, performing a readback from GPU->CPU if required. */
+	UE_API FReadbackResult GetCPUData(FPCGContext* InContext) const;
+
+	/** Returns the element count for this data. Does not trigger a GPU->CPU readback. */
+	UE_API FIntVector4 GetElementCount() const;
+
+	/** Returns a description of this data. Does not trigger a GPU->CPU readback. */
+	UE_API bool GetDescription(FPCGDataDesc& OutDescription) const;
+
+	UE_API TSharedPtr<const FPCGProxyForGPUDataCollection> GetGPUInfo() const;
+
+protected:
+	UE_API TSharedPtr<FPCGProxyForGPUDataCollection> GetGPUInfoMutable() const;
+
+protected:
+	UPROPERTY()
+	int32 DataIndexInCollection = INDEX_NONE;
+
+	TSharedPtr<FPCGProxyForGPUDataCollection> DataCollectionOnGPU;
+};
+
+#undef UE_API
diff --git a/Engine/Plugins/PCG/Source/PCG/Public/Compute/DataInterfaces/PCGComputeDataInterface.h b/Engine/Plugins/PCG/Source/PCG/Public/Compute/DataInterfaces/PCGComputeDataInterface.h
--- a/Engine/Plugins/PCG/Source/PCG/Public/Compute/DataInterfaces/PCGComputeDataInterface.h
+++ b/Engine/Plugins/PCG/Source/PCG/Public/Compute/DataInterfaces/PCGComputeDataInterface.h
@@ -12,22 +12,22 @@
 class UPCGComputeKernel;
 class UPCGDataBinding;
 
-UCLASS(Abstract, ClassGroup = (Procedural))
+UCLASS(MinimalAPI, Abstract, ClassGroup = (Procedural))
 class UPCGComputeDataInterface : public UComputeDataInterface
 {
 	GENERATED_BODY()
 
 public:
 	/** Register a downstream pin (and optionally compute graph element pin alias). */
-	virtual void AddDownstreamInputPin(FName InInputPinLabel, const FName* InOptionalInputPinLabelAlias = nullptr);
+	PCG_API virtual void AddDownstreamInputPin(FName InInputPinLabel, const FName* InOptionalInputPinLabelAlias = nullptr);
 
 	/** Set the output pin label and label alias that this data interface is associated with. */
-	virtual void SetOutputPin(FName InOutputPinLabel, const FName* InOptionalOutputPinLabelAlias = nullptr);
+	PCG_API virtual void SetOutputPin(FName InOutputPinLabel, const FName* InOptionalOutputPinLabelAlias = nullptr);
 
 	void SetProducerSettings(const UPCGSettings* InProducerSettings);
 	const UPCGSettings* GetProducerSettings() const;
 
-	static void AddReferencedObjects(UObject* InThis, FReferenceCollector& Collector);
+	PCG_API static void AddReferencedObjects(UObject* InThis, FReferenceCollector& Collector);
 
 	/** Set whether data originates from the CPU and should be uploaded to the GPU. */
 	void SetProducedByCPU(bool bInProducedByCPU) { bProducedByCPU = bInProducedByCPU; }
@@ -60,23 +60,23 @@ protected:
 	mutable TObjectPtr<const UPCGSettings> ResolvedProducerSettings = nullptr;
 };
 
-UCLASS(Abstract, ClassGroup (Procedural))
+UCLASS(MinimalAPI, Abstract, ClassGroup (Procedural))
 class UPCGComputeDataProvider : public UComputeDataProvider
 {
 	GENERATED_BODY()
 
 public:
-	virtual void Initialize(const UComputeDataInterface* InDataInterface, UObject* InBinding, uint64 InInputMask, uint64 InOutputMask) override;
-	virtual void Reset() override;
+	PCG_API virtual void Initialize(const UComputeDataInterface* InDataInterface, UObject* InBinding, uint64 InInputMask, uint64 InOutputMask) override;
+	PCG_API virtual void Reset() override;
 
 	/** Called before PrepareForExecute to allow data providers to do any readbacks and finalize data descriptions. */
-	virtual bool PerformPreExecuteReadbacks_GameThread(UPCGDataBinding* InBinding) { return true; }
+	PCG_API virtual bool PerformPreExecuteReadbacks_GameThread(UPCGDataBinding* InBinding) { return true; }
 
 	/* Called prior to GetRenderProxy and execution. Any processing of data descriptions should be done here (after PerformPreExecuteReadbacks has been called). */
-	virtual bool PrepareForExecute_GameThread(UPCGDataBinding* InBinding) { return true; }
+	PCG_API virtual bool PrepareForExecute_GameThread(UPCGDataBinding* InBinding) { return true; }
 
 	/** Return true when all done. */
-	virtual bool PostExecute(UPCGDataBinding* InBinding) { return true; }
+	PCG_API virtual bool PostExecute(UPCGDataBinding* InBinding) { return true; }
 
 	/** Release any handles to transient resources like GPU buffers. */
 	virtual void ReleaseTransientResources(const TCHAR* InReason = nullptr) {}
@@ -84,7 +84,7 @@ public:
 	const UPCGSettings* GetProducerSettings() const;
 	void SetProducerSettings(const UPCGSettings* InSettings);
 
-	static void AddReferencedObjects(UObject* InThis, FReferenceCollector& Collector);
+	PCG_API static void AddReferencedObjects(UObject* InThis, FReferenceCollector& Collector);
 
 protected:
 #if WITH_EDITOR
@@ -114,7 +114,7 @@ private:
 };
 
 /** Interface for data that can be exported/downloaded from the GPU to the CPU. */
-UCLASS(Abstract, ClassGroup = (Procedural))
+UCLASS(MinimalAPI, Abstract, ClassGroup = (Procedural))
 class UPCGExportableDataInterface : public UPCGComputeDataInterface
 {
 	GENERATED_BODY()
@@ -129,7 +129,7 @@ protected:
 	bool bRequiresExport = false;
 };
 
-UCLASS(Abstract, ClassGroup (Procedural))
+UCLASS(MinimalAPI, Abstract, ClassGroup (Procedural))
 class UPCGExportableDataProvider : public UPCGComputeDataProvider
 {
 	GENERATED_BODY()
@@ -163,7 +163,7 @@ protected:
 };
 
 /** Interface for marshalling settings/params to the GPU. */
-UCLASS(Abstract, ClassGroup = (Procedural))
+UCLASS(MinimalAPI, Abstract, ClassGroup = (Procedural))
 class UPCGKernelParamsDataInterface : public UPCGComputeDataInterface
 {
 	GENERATED_BODY()
@@ -175,7 +175,7 @@ public:
 	//~ End UComputeDataInterface Interface
 };
 
-UCLASS(Abstract, ClassGroup (Procedural))
+UCLASS(MinimalAPI, Abstract, ClassGroup (Procedural))
 class UPCGKernelParamsDataProvider : public UPCGComputeDataProvider
 {
 	GENERATED_BODY()
diff --git a/Engine/Plugins/PCG/Source/PCG/Public/Compute/PCGComputeCommon.h b/Engine/Plugins/PCG/Source/PCG/Public/Compute/PCGComputeCommon.h
new file mode 100644
--- /dev/null
+++ b/Engine/Plugins/PCG/Source/PCG/Public/Compute/PCGComputeCommon.h
@@ -0,0 +1,279 @@
+// Copyright Epic Games, Inc. All Rights Reserved.
+
+#pragma once
+
+#include "PCGCommon.h"
+#include "Data/Registry/PCGDataTypeIdentifier.h"
+
+#include "PCGComputeCommon.generated.h"
+
+class UPCGComputeDataInterface;
+class UPCGComputeKernel;
+class UPCGData;
+class UPCGDataBinding;
+class UPCGNode;
+class UObject;
+class UPCGSettings;
+struct FPCGGPUCompilationContext;
+struct FPCGContext;
+struct FPCGDataCollectionDesc;
+struct FPCGDataDesc;
+struct FPCGPinProperties;
+struct FPCGPinPropertiesGPU;
+struct FPCGTaggedData;
+
+class FRHIShaderResourceView;
+
+#define PCG_KERNEL_LOGGING_ENABLED (!(UE_BUILD_SHIPPING || UE_BUILD_TEST) || USE_LOGGING_IN_SHIPPING)
+
+#if PCG_KERNEL_LOGGING_ENABLED
+#define PCG_KERNEL_VALIDATION_WARN(Context, Settings, ValidationMessage) PCGComputeHelpers::LogKernelWarning(Context, Settings, ValidationMessage);
+#define PCG_KERNEL_VALIDATION_ERR(Context, Settings, ValidationMessage) PCGComputeHelpers::LogKernelError(Context, Settings, ValidationMessage);
+#else
+#define PCG_KERNEL_VALIDATION_WARN(Context, Settings, ValidationMessage) // Log removed
+#define PCG_KERNEL_VALIDATION_ERR(Context, Settings, ValidationMessage) // Log removed
+#endif
+
+/** Modes for exporting the buffer from transient to persistent for downstream consumption. */
+UENUM(meta = (Bitflags))
+enum class EPCGExportMode : uint8
+{
+	/** Buffer is transient and freed after usage. */
+	NoExport = 0,
+	/** Buffer will be exported and a proxy will be output from the compute graph and passed to downstream nodes. */
+	ComputeGraphOutput = 1 << 0,
+	/** Producer node is being inspected, read back data and store in inspection data. */
+	Inspection = 1 << 1,
+	/** Producer node is being debugged, read back data and execute debug visualization. */
+	DebugVisualization = 1 << 2,
+};
+ENUM_CLASS_FLAGS(EPCGExportMode);
+
+/** Dimensionality of a data description element count. */
+UENUM()
+enum class EPCGElementDimension : uint8
+{
+	One,
+	Two,
+	Three,
+	Four
+};
+
+namespace PCGComputeConstants
+{
+	constexpr int MAX_NUM_ATTRS = 128;
+	constexpr int NUM_RESERVED_ATTRS = 32; // Reserved for point properties, spline accessors, etc.
+	constexpr int MAX_NUM_CUSTOM_ATTRS = MAX_NUM_ATTRS - NUM_RESERVED_ATTRS; // Reserved for custom attributes
+
+	constexpr int DATA_COLLECTION_HEADER_SIZE_BYTES = 4; // 4 bytes for NumData
+	constexpr int DATA_HEADER_PREAMBLE_SIZE_BYTES = 12; // 4 bytes for DataType, 4 bytes for NumAttrs, 4 bytes for NumElements
+	constexpr int ATTRIBUTE_HEADER_SIZE_BYTES = 8; // 4 bytes for PackedIdAndStride, 4 bytes for data start address
+	constexpr int DATA_HEADER_SIZE_BYTES = DATA_HEADER_PREAMBLE_SIZE_BYTES + MAX_NUM_ATTRS * ATTRIBUTE_HEADER_SIZE_BYTES;
+
+	constexpr int POINT_DATA_TYPE_ID = 0;
+	constexpr int PARAM_DATA_TYPE_ID = 1;
+
+	constexpr int NUM_POINT_PROPERTIES = 9;
+	constexpr int POINT_POSITION_ATTRIBUTE_ID = 0;
+	constexpr int POINT_ROTATION_ATTRIBUTE_ID = 1;
+	constexpr int POINT_SCALE_ATTRIBUTE_ID = 2;
+	constexpr int POINT_BOUNDS_MIN_ATTRIBUTE_ID = 3;
+	constexpr int POINT_BOUNDS_MAX_ATTRIBUTE_ID = 4;
+	constexpr int POINT_COLOR_ATTRIBUTE_ID = 5;
+	constexpr int POINT_DENSITY_ATTRIBUTE_ID = 6;
+	constexpr int POINT_SEED_ATTRIBUTE_ID = 7;
+	constexpr int POINT_STEEPNESS_ATTRIBUTE_ID = 8;
+
+	constexpr uint32 KernelExecutedFlag = 1 << 31;
+
+	// PackedAttributeInfo bitmasks
+	constexpr uint32 AttributeStrideMask = 0xFF;
+	constexpr uint32 AttributeAllocatedMask = 1 << 31;
+
+	constexpr int32 MAX_PRIMITIVE_COMPONENTS_PER_SPAWNER = 256;
+
+	constexpr uint32 THREAD_GROUP_SIZE = 64;
+
+	// Used to represent invalid/removed points. We use a value slightly less than max float,
+	// as not all platforms support float infinity in shaders.
+	constexpr float INVALID_DENSITY = 3.402823e+38f;
+
+	constexpr TCHAR DataLabelTagPrefix[] = { TEXT("PCG_DATA_LABEL") };
+}
+
+namespace PCGComputeHelpers
+{
+	/** Gets the element count for a given data. E.g. number of points in a PointData, number of metadata entries in a ParamData, etc. */
+	FIntVector4 GetElementCount(const UPCGData* InData);
+
+	/** Gets the element dimension for a given data. E.g. OneD for PointData, TwoD for TextureData, etc. */
+	EPCGElementDimension GetElementDimension(const UPCGData* InData);
+	EPCGElementDimension GetElementDimension(const FPCGDataTypeIdentifier& InDataType);
+
+	/** PCG data types supported in GPU node inputs. */
+	const TArray<FPCGDataTypeIdentifier>& GetAllowedInputTypesList();
+
+	/** PCG data types supported in GPU node outputs. */
+	const TArray<FPCGDataTypeIdentifier>& GetAllowedOutputTypesList();
+
+	/** True if 'Type' is valid on a GPU input pin. */
+	bool IsTypeAllowedAsInput(const FPCGDataTypeIdentifier& Type);
+
+	/** True if 'Type' is valid on a GPU output pin. */
+	bool IsTypeAllowedAsOutput(const FPCGDataTypeIdentifier& Type);
+
+	/** True if 'Type' is valid in a GPU data collection. Some types are only supported as DataInterfaces, and cannot be uploaded in data collections. */
+	bool IsTypeAllowedInDataCollection(const FPCGDataTypeIdentifier& Type);
+
+	/** Whether metadata attributes should be read from the given data and registered for use in GPU graphs. */
+	bool ShouldImportAttributesFromData(const UPCGData* InData);
+
+#if PCG_KERNEL_LOGGING_ENABLED
+	/** Logs a warning on a GPU node in the graph and console. */
+	void LogKernelWarning(const FPCGContext* Context, const UPCGSettings* Settings, const FText& InText);
+
+	/** Logs an error on a GPU node in the graph and console. */
+	void LogKernelError(const FPCGContext* Context, const UPCGSettings* Settings, const FText& InText);
+#endif
+
+	/** Returns true if the given buffer size is dangerously large. Optionally emits error log. */
+	bool IsBufferSizeTooLarge(uint64 InBufferSizeBytes, bool bInLogError = true);
+
+	int32 GetAttributeIdFromMetadataAttributeIndex(int32 InAttributeIndex);
+	int32 GetMetadataAttributeIndexFromAttributeId(int32 InAttributeId);
+
+	/** Produces the data label prefixed with PCGComputeConstants::DataLabelTagPrefix. */
+	FString GetPrefixedDataLabel(const FString& InLabel);
+
+	/** Produces the data interface name of a data label resolver. */
+	FString GetDataLabelResolverName(FName InPinLabel);
+
+#if WITH_EDITOR
+	void ConvertObjectPathToShaderFilePath(FString& InOutPath);
+
+	struct FCreateDataInterfaceParams
+	{
+		FPCGGPUCompilationContext* Context = nullptr;
+		const FPCGPinProperties* PinProperties = nullptr;
+		const UPCGComputeKernel* ProducerKernel = nullptr;
+		UObject* ObjectOuter = nullptr;
+		bool bProducedByCPU = false;
+		bool bRequiresExport = false;
+		const UPCGNode* NodeForDebug = nullptr;
+	};
+
+	UPCGComputeDataInterface* CreateOutputPinDataInterface(const FCreateDataInterfaceParams& InParams);
+
+	void NotifyGPUToCPUReadback(UPCGDataBinding* InBinding, const UPCGComputeKernel* InKernel, const UPCGSettings* InSettings);
+	void NotifyCPUToGPUUpload(UPCGDataBinding* InBinding, const UPCGComputeKernel* InKernel, const UPCGSettings* InSettings);
+#endif
+}
+
+namespace PCGComputeDummies
+{
+	FRHIShaderResourceView* GetDummyFloatBuffer();
+	FRHIShaderResourceView* GetDummyFloat2Buffer();
+	FRHIShaderResourceView* GetDummyFloat4Buffer();
+}
+
+/** A by-label reference to a pin, used for wiring kernels within a node. */
+struct FPCGPinReference
+{
+	/** Reference a pin by label only, used for referencing node pins. */
+	explicit FPCGPinReference(FName InLabel)
+		: Kernel(nullptr)
+		, Label(InLabel)
+	{
+	}
+
+	/** Reference a pin by kernel and label. */
+	explicit FPCGPinReference(UPCGComputeKernel* InKernel, FName InLabel)
+		: Kernel(InKernel)
+		, Label(InLabel)
+	{
+	}
+
+	bool operator==(const FPCGPinReference& Other) const
+	{
+		return Label == Other.Label
+			&& Kernel == Other.Kernel;
+	}
+
+	/** Associated kernel. If null then compiler will look for pin on owning node. */
+	UPCGComputeKernel* Kernel = nullptr;
+
+	/** Pin label. */
+	FName Label;
+};
+
+uint32 GetTypeHash(const FPCGPinReference& In);
+
+/** A connection for wiring kernels within a node. */
+struct FPCGKernelEdge
+{
+	FPCGKernelEdge(const FPCGPinReference& InUpstreamPin, const FPCGPinReference& InDownstreamPin)
+		: UpstreamPin(InUpstreamPin)
+		, DownstreamPin(InDownstreamPin)
+	{
+	}
+
+	bool IsConnectedToNodeInput() const { return UpstreamPin.Kernel == nullptr; }
+	bool IsConnectedToNodeOutput() const { return DownstreamPin.Kernel == nullptr; }
+
+	UPCGComputeKernel* GetUpstreamKernel() const { return UpstreamPin.Kernel; }
+	UPCGComputeKernel* GetDownstreamKernel() const { return DownstreamPin.Kernel; }
+
+	FPCGPinReference UpstreamPin;
+	FPCGPinReference DownstreamPin;
+};
+
+/** An input or output pin of a kernel. Compute graph does not internally have 'pins' so this is useful for mapping between kernel data and PCG pins. */
+USTRUCT()
+struct FPCGKernelPin
+{
+	GENERATED_BODY()
+
+public:
+	FPCGKernelPin() = default;
+
+	explicit FPCGKernelPin(int32 InKernelIndex, FName InPinLabel, bool bInIsInput)
+		: KernelIndex(InKernelIndex)
+		, PinLabel(InPinLabel)
+		, bIsInput(bInIsInput)
+	{
+	}
+
+	bool operator==(const FPCGKernelPin& Other) const = default;
+
+public:
+	UPROPERTY()
+	int32 KernelIndex = INDEX_NONE;
+
+	UPROPERTY()
+	FName PinLabel = NAME_None;
+
+	UPROPERTY()
+	bool bIsInput = false;
+
+	friend uint32 GetTypeHash(const FPCGKernelPin& In);
+};
+
+/** Helper struct for serializing data labels. */
+USTRUCT()
+struct FPCGDataLabels
+{
+	GENERATED_BODY()
+
+	UPROPERTY()
+	TArray<FString> Labels;
+};
+
+/** Helper struct for serializing map of pin name to data labels. */
+USTRUCT()
+struct FPCGPinDataLabels
+{
+	GENERATED_BODY()
+
+	UPROPERTY()
+	TMap</*PinLabel*/FName, FPCGDataLabels> PinToDataLabels;
+};
diff --git a/Engine/Plugins/PCG/Source/PCG/Public/Compute/PCGComputeKernel.h b/Engine/Plugins/PCG/Source/PCG/Public/Compute/PCGComputeKernel.h
new file mode 100644
--- /dev/null
+++ b/Engine/Plugins/PCG/Source/PCG/Public/Compute/PCGComputeKernel.h
@@ -0,0 +1,258 @@
+// Copyright Epic Games, Inc. All Rights Reserved.
+
+#pragma once
+
+#include "PCGSettings.h"
+#include "Compute/PCGComputeCommon.h"
+
+#include "ComputeFramework/ComputeKernel.h"
+#include "UObject/ObjectKey.h"
+
+#include "PCGComputeKernel.generated.h"
+
+class UPCGDataBinding;
+class UPCGSettings;
+struct FPCGContext;
+struct FPCGDataCollectionDesc;
+struct FPCGGPUCompilationContext;
+struct FPCGKernelAttributeKey;
+struct FPCGPinProperties;
+struct FPCGPinPropertiesGPU;
+
+class UComputeDataInterface;
+
+UENUM()
+enum class EPCGComputeKernelFlags : uint32
+{
+	None = 0,
+	/** Errors were detected with node setup at compute graph compilation time. */
+	HasStaticValidationErrors = 1 << 0,
+	/** Dump data descriptions to log. */
+	LogDataDescriptions = 1 << 1,
+	/** Dispatch compute graph containing this kernel every frame to allow profiling. Has no effect unless PCG_GPU_KERNEL_PROFILING defined. */
+	RepeatDispatchForProfiling = 1 << 2,
+};
+ENUM_CLASS_FLAGS(EPCGComputeKernelFlags);
+
+UENUM()
+enum class EPCGKernelLogVerbosity
+{
+	Verbose,
+	Warning,
+	Error
+};
+
+USTRUCT()
+struct FPCGKernelLogEntry
+{
+	GENERATED_BODY()
+
+	FPCGKernelLogEntry() = default;
+
+	explicit FPCGKernelLogEntry(const FText& InMessage, EPCGKernelLogVerbosity InVerbosity)
+		: Message(InMessage)
+		, Verbosity(InVerbosity)
+	{}
+
+	UPROPERTY()
+	FText Message;
+
+	UPROPERTY()
+	EPCGKernelLogVerbosity Verbosity = EPCGKernelLogVerbosity::Verbose;
+};
+
+#if WITH_EDITOR
+struct FPCGComputeKernelParams
+{
+	const UPCGSettings* Settings = nullptr;
+	bool bLogDescriptions = false;
+	bool bRequiresOverridableParams = false;
+};
+#endif
+
+/** Represents a parameter that can be overridden on a kernel. */
+USTRUCT()
+struct FPCGKernelOverridableParam
+{
+	GENERATED_BODY()
+
+	UPROPERTY()
+	FName Label = NAME_None;
+
+	UPROPERTY()
+	TArray<FName> PropertiesNames;
+
+	UPROPERTY()
+	bool bRequiresGPUReadback = false;
+
+	/** True if this param is actually overidden by a pin. */
+	UPROPERTY()
+	bool bIsPropertyOverriddenByPin = false;
+};
+
+UCLASS(MinimalAPI, Abstract)
+class UPCGComputeKernel : public UComputeKernel
+{
+	GENERATED_BODY()
+
+public:
+#if WITH_EDITOR
+	/** Initialize kernel. Editor only as shader compilation only available in editor. */
+	PCG_API void Initialize(const FPCGComputeKernelParams& InParams);
+#endif
+
+	/** Get/set index of kernel in kernels array on owning compute graph. */
+	int32 GetKernelIndex() const { return KernelIndex; }
+	void SetKernelIndex(int32 InKernelIndex) { KernelIndex = InKernelIndex; }
+
+	/** Gets settings for node associated with this kernel, if any. */
+	PCG_API virtual const UPCGSettings* GetSettings() const;
+
+	/** Performs settings validation and returns true if this node is suitable for execution. */
+	PCG_API virtual bool AreKernelSettingsValid(FPCGContext* InContext) const;
+
+	/** Performs data validation and returns true if this node is suitable for deployment to the GPU. */
+	PCG_API virtual bool IsKernelDataValid(FPCGContext* InContext) const;
+
+#if WITH_EDITOR
+	/** Produces the node specific portion of kernel shader source text, including the main entry point. */
+	PCG_API virtual FString GetCookedSource(FPCGGPUCompilationContext& InOutContext) const;
+
+	/** Get the virtual file path of the shader source file for this kernel. */
+	PCG_API virtual const TCHAR* GetSourceFilePath() const { return nullptr; }
+
+	/** Get the name of the main kernel function in the source. This name is also displayed in GPU profile scopes. */
+	PCG_API virtual FString GetEntryPoint() const PURE_VIRTUAL(UPCGComputeKernel::GetEntryPoint, return TEXT("Main"););
+
+	/**
+	 * Gathers additional compute sources referenced by this kernel.
+	 * This is preferred to directly using includes, since it allows us to detect hash diffs and recompile when additional sources are changed externally.
+	 * Note, you should not create new compute sources here, as their object names are not deterministic and will result in a new hash every time.
+	 */
+	PCG_API virtual void GatherAdditionalSources(TArray<TObjectPtr<class UComputeSource>>& OutAdditionalSources) const;
+
+	/** Create additional input data interfaces to marshal any required input data. */
+	PCG_API virtual void CreateAdditionalInputDataInterfaces(FPCGGPUCompilationContext& InOutContext, UObject* InObjectOuter, TArray<TObjectPtr<UComputeDataInterface>>& OutDataInterfaces) const {}
+
+	/** Create additional output data interfaces to marshal any required output data. */
+	PCG_API virtual void CreateAdditionalOutputDataInterfaces(FPCGGPUCompilationContext& InOutContext, UObject* InObjectOuter, TArray<TObjectPtr<UComputeDataInterface>>& OutDataInterfaces) const {}
+
+	/** Whether compute graphs should be split at this input pin. Useful for overridable parameters. */
+	PCG_API virtual bool SplitGraphAtInput(FName InPinLabel) const;
+
+	/** Whether compute graphs should be split at the output of this node. Useful for reading back runtime statistics or diagnostic info before continuing. */
+	PCG_API virtual bool SplitGraphAtOutput() const { return false; }
+
+	/** Mark a pin as being internal to the kernel graph. This means it is not on the CPU/GPU boundary. */
+	void AddInternalPin(FName InPinLabel) { InternalPinLabels.Add(InPinLabel); }
+
+	/** Whether the pin matching the given label should be eligible for inspect/debug. */
+	bool IsPinInternal(FName InPinLabel) const { return InternalPinLabels.Contains(InPinLabel); }
+
+	/** Create a data interface for the given output pin. */
+	PCG_API virtual UPCGComputeDataInterface* CreateOutputPinDataInterface(const PCGComputeHelpers::FCreateDataInterfaceParams& InParams) const;
+#endif
+
+	/** Compute a description of data that will be output from pin InOutputPinLabel. */
+	PCG_API virtual TSharedPtr<const FPCGDataCollectionDesc> ComputeOutputBindingDataDesc(FName InOutputPinLabel, UPCGDataBinding* InBinding) const PURE_VIRTUAL(UPCGComputeKernel::ComputeOutputBindingDataDesc, return nullptr;);
+
+	/** Compute how many threads should be dispatched to execute this node on the GPU. */
+	PCG_API virtual int ComputeThreadCount(const UPCGDataBinding* Binding) const PURE_VIRTUAL(UPCGComputeKernel::ComputeThreadCount, return 0;);
+
+	/** Get any data labels that are used statically by this node. */
+	PCG_API virtual void GetDataLabels(FName InPinLabel, TArray<FString>& OutDataLabels) const {}
+
+	/** Initialize all buffer data to 0. This is not free so should only be used when kernel execution requires 0-initialized data, such as for counters. */
+	PCG_API virtual bool DoesOutputPinRequireZeroInitialization(FName InOutputPinLabel) const { return false; }
+
+	/** Add any strings emitted by this node that are known statically at compile time. */
+	PCG_API virtual void AddStaticCreatedStrings(TArray<FString>& InOutStringTable) const {}
+
+	/** Get all the attributes read or written by this node for which we know the name and type statically. In cases that only the
+	* attribute name used by the node is known (example: user specifies SM Spawner instance data attributes by name only, not type), the attribute
+	* is omitted from this list and must be resolved at execution time in the data provider.
+	*/
+	PCG_API virtual void GetKernelAttributeKeys(TArray<FPCGKernelAttributeKey>& OutKeys) const {}
+
+	/** Get multiplier to be applied to the dispatched thread count. */
+	PCG_API virtual uint32 GetThreadCountMultiplier() const { return 1u; }
+
+	/** Multiplier to be applied to the element count of each output data. */
+	PCG_API virtual uint32 GetElementCountMultiplier(FName InOutputPinLabel) const { return 1u; }
+
+	/** Helper to compute data descriptions when data is entirely determined by pin properties. */
+	void ComputeDataDescFromPinProperties(const FPCGPinPropertiesGPU& OutputPinProps, const TArrayView<const FPCGPinProperties>& InInputPinProps, UPCGDataBinding* InBinding, TSharedPtr<FPCGDataCollectionDesc> OutPinDesc) const;
+
+	PCG_API virtual void GetInputPins(TArray<FPCGPinProperties>& OutPins) const {}
+	PCG_API virtual void GetOutputPins(TArray<FPCGPinPropertiesGPU>& OutPins) const {}
+
+	/** Return the concatenation of GetInputPins() and GetOverridableParamPins(). */
+	void GetInputPinsAndOverridablePins(TArray<FPCGPinProperties>& OutPins) const;
+
+	const TArray<FPCGKernelOverridableParam>& GetCachedOverridableParams() const { return CachedOverridableParams; }
+
+	bool GetLogDataDescriptions() const { return !!(PCGKernelFlags & EPCGComputeKernelFlags::LogDataDescriptions); }
+
+	bool HasStaticValidationErrors() const { return !!(PCGKernelFlags & EPCGComputeKernelFlags::HasStaticValidationErrors); }
+
+#ifdef PCG_GPU_KERNEL_PROFILING
+	void EnableRepeatDispatch() { PCGKernelFlags |= EPCGComputeKernelFlags::RepeatDispatchForProfiling; }
+
+	/** Whether to repeatedly dispatch every frame to help profiling. Either configured on the GPU node or overridden here to avoid data cooks. */
+	virtual bool ShouldRepeatDispatch() const { return !!(PCGKernelFlags & EPCGComputeKernelFlags::RepeatDispatchForProfiling); }
+#endif
+
+	PCG_API static void AddReferencedObjects(UObject* InThis, FReferenceCollector& Collector);
+
+protected:
+#if WITH_EDITOR
+	/** Implement to do any additional initialization specific to your kernel. */
+	PCG_API virtual void InitializeInternal() {}
+
+	/** Performs validation on compile time information (e.g. Node, Settings, etc.). Caches validation errors/warnings. */
+	PCG_API virtual bool PerformStaticValidation();
+
+	/** Do any validation of the PCG node. Returns true if node is valid and GPU execution can proceed. */
+	bool ValidatePCGNode(TArray<FPCGKernelLogEntry>& OutLogEntries) const;
+
+	/** Validate input edges, for example checking connected pin types are compatible. */
+	bool AreInputEdgesValid(TArray<FPCGKernelLogEntry>& OutLogEntries) const;
+#endif
+
+	/** Collects all param pins which are overridable on the GPU. */
+	void GetOverridableParamPins(TArray<FPCGPinProperties>& OutPins) const;
+
+protected:
+	UPROPERTY()
+	EPCGComputeKernelFlags PCGKernelFlags = EPCGComputeKernelFlags::None;
+
+	/** Index into kernels array on owning compute graph. */
+	UPROPERTY()
+	int32 KernelIndex = INDEX_NONE;
+
+	UPROPERTY()
+	TSoftObjectPtr<const UPCGSettings> SettingsPath;
+
+	UPROPERTY()
+	bool bInitialized = false;
+
+	// @todo_pcg: These are a bit of a hack to avoid supporting static/compilation logging in PCGNodeVisualLogs.h
+	/** Log entries created and cached on initialization. Are logged during AreKernelSettingsValid(), unless PCG_KERNEL_LOGGING_ENABLED is false. */
+	UPROPERTY()
+	TArray<FPCGKernelLogEntry> StaticLogEntries;
+
+	/** Overridable params supported by this kernel, cached during compilation. */
+	UPROPERTY()
+	TArray<FPCGKernelOverridableParam> CachedOverridableParams;
+
+#if WITH_EDITOR
+	TSet<FName> InternalPinLabels;
+#endif
+
+	mutable TObjectPtr<const UPCGSettings> ResolvedSettings = nullptr;
+};
+
+namespace PCGComputeKernel
+{
+	PCG_API TConstArrayView<FSoftObjectPath> GetDefaultAdditionalSourcePaths();
+}
diff --git a/Engine/Plugins/PCG/Source/PCG/Public/Compute/PCGDataBinding.h b/Engine/Plugins/PCG/Source/PCG/Public/Compute/PCGDataBinding.h
new file mode 100644
--- /dev/null
+++ b/Engine/Plugins/PCG/Source/PCG/Public/Compute/PCGDataBinding.h
@@ -0,0 +1,212 @@
+// Copyright Epic Games, Inc. All Rights Reserved.
+
+#pragma once
+
+#include "PCGData.h"
+#include "PCGSettings.h"
+#include "Compute/PCGDataForGPU.h"
+#include "Compute/PCGKernelParamsCache.h"
+
+#include "PCGDataBinding.generated.h"
+
+class IPCGRuntimePrimitiveFactory;
+class UComputeKernel;
+class UPCGComponent;
+class UPCGComputeGraph;
+class UPCGComputeKernel;
+class UPCGProxyForGPUData;
+struct FPCGComputeGraphContext;
+struct FPCGContextHandle;
+struct FPCGKernelPin;
+
+USTRUCT()
+struct FPCGSpawnerPrimitives
+{
+	GENERATED_BODY()
+
+	TSharedPtr<IPCGRuntimePrimitiveFactory> PrimitiveFactory;
+
+	TArray<FBox> PrimitiveMeshBounds;
+
+	/** Cumulative distribution function values (one per primitive) to enable choosing a primitive based on a random draw value. */
+	TArray<float> SelectionCDF;
+
+	// Same for all primitives
+	uint32 NumCustomFloats = 0;
+
+	// Same for all primitives
+	TArray<FUintVector4> AttributeIdOffsetStrides;
+
+	int32 SelectorAttributeId = INDEX_NONE;
+
+	TArray<int32> PrimitiveStringKeys;
+};
+
+USTRUCT()
+struct FPCGDataToDebug
+{
+	GENERATED_BODY()
+
+	UPROPERTY()
+	TObjectPtr<UPCGData> Data;
+
+	UPROPERTY()
+	TObjectPtr<UPCGData> DataPendingInit;
+	
+	UPROPERTY()
+	TWeakObjectPtr<const UPCGSettings> ProducerSettings;
+	
+	UPROPERTY()
+	FName PinLabel;
+	
+	UPROPERTY()
+	FName PinLabelAlias;
+
+	// @todo_pcg: This is a big hack to support tags on texture data proxies. Should be replaced with a proper abstraction.
+	/** Allow appending additional tags to the CPU tagged data. */
+	UPROPERTY()
+	TSet<FString> AdditionalTags;
+};
+
+UCLASS(MinimalAPI, Transient, Category = PCG)
+class UPCGDataBinding : public UObject
+{
+	GENERATED_BODY()
+
+public:
+	/** Pre-initialization, set up pointers and store input data. */
+	void Initialize(const UPCGComputeGraph* InComputeGraph, FPCGContext* InContext);
+
+	/** Initialization look up tables from input data. */
+	void InitializeTables(FPCGContext* InContext);
+
+	/** Clear state and release any handles to resources such as GPU buffers. */
+	void ReleaseTransientResources();
+
+	int32 GetAttributeId(const FPCGKernelAttributeKey& InAttribute) const { ensure(bTablesInitialized); return AttributeTable.GetAttributeId(InAttribute); }
+	int32 GetAttributeId(FPCGAttributeIdentifier InIdentifier, EPCGKernelAttributeType InType) const { ensure(bTablesInitialized); return AttributeTable.GetAttributeId(InIdentifier, InType); }
+
+	int32 GetAttributeTableSize() const { return AttributeTable.Num(); }
+
+	/** Get the set of all unique strings entering the compute graph. */
+	const TArray<FString>& GetStringTable() const { ensure(bTablesInitialized); return StringTable; }
+
+	/** Computes a description of data for every pin in the compute graph and caches it. */
+	void PrimeDataDescriptionCache();
+	bool IsDataDescriptionCachePrimed() const { return bIsDataDescriptionCachePrimed; }
+
+	/** Computes a description of data for a single pin in the compute graph and caches it. */
+	PCG_API TSharedPtr<const FPCGDataCollectionDesc> ComputeKernelPinDataDesc(const FPCGKernelPin& InKernelPin);
+	TSharedPtr<const FPCGDataCollectionDesc> ComputeKernelBindingDataDesc(int32 InGraphBindingIndex);
+
+	/** Get description of data produced by a pin. */
+	PCG_API TSharedPtr<const FPCGDataCollectionDesc> GetCachedKernelPinDataDesc(const UPCGComputeKernel* InKernel, FName InPinLabel, bool bIsInput) const;
+	TSharedPtr<const FPCGDataCollectionDesc> GetCachedKernelPinDataDesc(int32 InGraphBindingIndex) const;
+
+	void ReceiveDataFromGPU_GameThread(UPCGData* InData, const UPCGSettings* InProducerSettings, EPCGExportMode InExportMode, FName InPinLabel, FName InPinLabelAlias);
+	void ReceiveDataFromGPU_GameThread(UPCGData* InData, const UPCGSettings* InProducerSettings, EPCGExportMode InExportMode, FName InPinLabel, FName InPinLabelAlias, const TSet<FString>& AdditionalTags);
+
+	using FSourceBufferAndAttributeIndex = TPair<TSharedPtr<const FPCGProxyForGPUDataCollection>, /*Source Metadata Attribute Index*/int32>;
+	using FGraphAttributeIndexAndSourceBuffer = TPair</*Graph Metadata Attribute Index*/int32, TRefCountPtr<FRDGPooledBuffer>>;
+
+	/** Get map from source buffer and source buffer attribute index to this-graph attribute index. */
+	const TMap<FSourceBufferAndAttributeIndex, /*Graph Attribute Index*/int32>& GetSourceBufferAttributeToGraphAttributeIndex() const { return SourceBufferAttributeToGraphAttributeIndex; }
+
+	/** Get map from this-graph attribute index and upstream source buffer to upstream attribute index. */
+	const TMap<FGraphAttributeIndexAndSourceBuffer, /*Source Buffer Attribute Index*/int32>& GetGraphAttributeToSourceBufferAttributeIndex() const { return GraphAttributeToSourceBufferAttributeIndex; }
+
+	using FSourceBufferAndStringKey = TPair<TRefCountPtr<FRDGPooledBuffer>, /*String Key*/int32>;
+
+	/** Map from string key value in graph to the key value for this string in the source buffer. */
+	const TMap<FSourceBufferAndStringKey, /*Graph String Key*/int32>& GetSourceBufferToGraphStringKey() const { return SourceBufferToGraphStringKey; }
+
+	/** Map from string key value in source buffer to the key value usable in this graph. */
+	const TMap<FSourceBufferAndStringKey, /*Source Buffer String Key*/int32>& GetGraphToSourceBufferStringKey() const { return GraphToSourceBufferStringKey; }
+
+	IPCGGraphExecutionSource* GetExecutionSource() const;
+
+	/** Helper to get index in input data collection of the first data item for the given kernel and input pin label. */
+	int32 GetFirstInputDataIndex(const UPCGComputeKernel* InKernel, FName InPinLabel) const;
+
+	/** Returns indices of data in the input data collection arriving on the given kernel pin. */
+	TArray<int32> GetPinInputDataIndices(const UPCGComputeKernel* InKernel, FName InPinLabel) const;
+
+	/** If there is a GPU proxy at the given index of the input data collection, triggers a readback and replaces the data item with CPU data if readback succeeds.
+	* Returns false while readback is in progress.
+	*/
+	bool ReadbackInputDataToCPU(int32 InInputDataIndex);
+
+	void DebugLogDataDescriptions() const;
+
+	/** Initialize the kernel params cache. Performs readbacks of overridable params that require it, and caches all values for quick access from kernels and data providers.
+	* Returns false while readback is in progress.
+	*/
+	bool InitializeKernelParams(FPCGContext* InContext);
+
+	/** Look up the cached parameter struct for a kernel. */
+	const FPCGKernelParams* GetCachedKernelParams(const UPCGComputeKernel* InKernel) const;
+	const FPCGKernelParams* GetCachedKernelParams(int32 InKernelIndex) const;
+
+private:
+	/** Loop over all metadata attributes in all input data ensure all attributes are registered in attribute table. */
+	void AddInputDataAttributesToTable();
+	void AddInputDataStringsToTable();
+	void AddInputDataTagsToTable();
+
+public:
+	/** Compute graph element output data. Data items are labelled with unique virtual output pin labels so that the can be routed correctly by the graph
+	 * executor to downstream nodes. */
+	UPROPERTY()
+	FPCGDataCollection OutputDataCollection;
+
+	UPROPERTY()
+	TObjectPtr<const UPCGComputeGraph> Graph = nullptr;
+
+	/** Data arriving on compute graph element. Since the compute graph is collapsed to a single element, all data crossing from CPU to GPU is in a single collection. */
+	UPROPERTY()
+	FPCGDataCollection InputDataCollection;
+
+	UPROPERTY()
+	TMap<TObjectPtr<const UComputeKernel>, FPCGSpawnerPrimitives> MeshSpawnersToPrimitives;
+
+	UPROPERTY()	
+	TArray<TObjectPtr<const UComputeKernel>> CompletedMeshSpawners;
+
+	UPROPERTY()
+	TArray<FPCGDataToDebug> DataToDebug;
+
+	UPROPERTY()
+	TArray<FPCGDataToDebug> DataToInspect;
+
+	TWeakPtr<FPCGContextHandle> ContextHandle;
+
+protected:
+	/** All attributes present in graph at execution time. Coherent across all graph branches. Seeded using statically known created attributes at compile time, then
+	 * augmented with incoming attributes from input data collection at runtime. */
+	UPROPERTY()
+	FPCGKernelAttributeTable AttributeTable;
+
+	UPROPERTY()
+	TArray<FString> StringTable;
+
+	/** Map from source buffer and source buffer attribute index to this-graph attribute index. */
+	TMap<FSourceBufferAndAttributeIndex, /*Attribute ID*/int32> SourceBufferAttributeToGraphAttributeIndex;
+
+	/** Map from this-graph attribute index and upstream source buffer to upstream attribute index. */
+	TMap<FGraphAttributeIndexAndSourceBuffer, /*Source Buffer Attribute Index*/int32> GraphAttributeToSourceBufferAttributeIndex;
+
+	/** Map from string key value in graph to the key value for this string in the source buffer. */
+	TMap<FSourceBufferAndStringKey, /*Graph String Key*/int32> SourceBufferToGraphStringKey;
+
+	/** Map from string key value in source buffer to the key value usable in this graph. */
+	TMap<FSourceBufferAndStringKey, /*Source Buffer String Key*/int32> GraphToSourceBufferStringKey;
+
+	/** Cache of data descriptions to amortize cost of computing them at runtime. */
+	TMap</*Binding index*/int32, TSharedPtr<const FPCGDataCollectionDesc>> DataDescriptionCache;
+	std::atomic<bool> bIsDataDescriptionCachePrimed = false;
+
+	std::atomic<bool> bTablesInitialized = false;
+
+	/** Cache of overridden settings/values per-kernel, built during compute graph execution. */
+	FPCGKernelParamsCache KernelParamsCache;
+};
diff --git a/Engine/Plugins/PCG/Source/PCG/Public/Compute/PCGDataForGPU.h b/Engine/Plugins/PCG/Source/PCG/Public/Compute/PCGDataForGPU.h
new file mode 100644
--- /dev/null
+++ b/Engine/Plugins/PCG/Source/PCG/Public/Compute/PCGDataForGPU.h
@@ -0,0 +1,330 @@
+// Copyright Epic Games, Inc. All Rights Reserved.
+
+#pragma once
+
+#include "PCGCommon.h"
+#include "PCGData.h"
+#include "Compute/PCGComputeCommon.h"
+
+#include "RenderGraphResources.h"
+#include "RHIGPUReadback.h"
+#include "Misc/SpinLock.h"
+
+#include "PCGDataForGPU.generated.h"
+
+class FReferenceCollector;
+class UPCGDataBinding;
+class UPCGMetadata;
+class UPCGPin;
+class UPCGSettings;
+struct FPCGContext;
+enum class EPCGElementMultiplicity : uint8;
+enum class EPCGMetadataTypes : uint8;
+enum class EPCGPointNativeProperties : uint32;
+
+enum class EPCGUnpackDataCollectionResult
+{
+	Success,
+	DataMismatch,
+	NoData
+};
+
+UENUM()
+enum class EPCGKernelAttributeType : uint8
+{
+	Bool = 1,
+	Int,
+	Float,
+	Float2,
+	Float3,
+	Float4,
+	Rotator,
+	Quat,
+	Transform,
+	StringKey,
+	Name,
+
+	Invalid = std::numeric_limits<uint8>::max() UMETA(Hidden),
+};
+
+/** Attribute name and type which uniquely identify an attribute in a compute graph. */
+USTRUCT()
+struct FPCGKernelAttributeKey
+{
+	GENERATED_BODY()
+
+	FPCGKernelAttributeKey() = default;
+
+	explicit FPCGKernelAttributeKey(FPCGAttributeIdentifier InIdentifier, EPCGKernelAttributeType InType)
+		: Identifier(InIdentifier)
+		, Type(InType)
+	{}
+
+	explicit FPCGKernelAttributeKey(const FPCGAttributePropertySelector& InSelector, EPCGKernelAttributeType InType);
+
+	/** To be called everytime the Selector changes, to update the Identifier. Return true if it has changed. */
+	bool UpdateIdentifierFromSelector();
+	
+	void SetSelector(const FPCGAttributePropertySelector& InSelector);
+
+	/** Cached identifier. Need to be updated if the Selector ever change. */
+	UPROPERTY()
+	FPCGAttributeIdentifier Identifier;
+
+	UPROPERTY(EditAnywhere, Category = "Settings", meta = (DisplayAfter = "Attribute"))
+	EPCGKernelAttributeType Type = EPCGKernelAttributeType::Float;
+
+	bool IsValid() const;
+
+	bool operator==(const FPCGKernelAttributeKey& Other) const;
+
+	friend uint32 GetTypeHash(const FPCGKernelAttributeKey& In);
+
+protected:
+	/** Selector to specify which attribute to create and on which domain. At the moment, only support `@Data` domains or no domain (default one). */
+	UPROPERTY(EditAnywhere, Category = "Settings", meta = (PCG_DiscardPropertySelection, PCG_DiscardExtraSelection))
+	FPCGAttributePropertyOutputNoSourceSelector Name;
+};
+
+/** Table of attributes used in compute graph with helpers to get unique attribute ID used to read/write attribute in data collection buffers. */
+USTRUCT()
+struct FPCGKernelAttributeTable
+{
+	GENERATED_BODY()
+
+	FPCGKernelAttributeTable() = default;
+
+	int32 GetAttributeId(const FPCGKernelAttributeKey& InAttribute) const;
+
+	int32 GetAttributeId(FPCGAttributeIdentifier InIdentifier, EPCGKernelAttributeType InType) const;
+
+	/** Adds an attribute of given name and type. Returns index or INDEX_NONE if add failed (happens if max table size reached). */
+	int32 AddAttribute(const FPCGKernelAttributeKey& Key);
+	int32 AddAttribute(FPCGAttributeIdentifier InIdentifier, EPCGKernelAttributeType InType);
+
+	int32 Num() const { return AttributeTable.Num(); }
+
+#if PCG_KERNEL_LOGGING_ENABLED
+	void DebugLog() const;
+#endif
+
+protected:
+	UPROPERTY()
+	TArray<FPCGKernelAttributeKey> AttributeTable;
+};
+
+/** Data description for a metadata attribute. Stores identifying name and type as well as the unique attribute ID. */
+struct FPCGKernelAttributeDesc
+{
+public:
+	FPCGKernelAttributeDesc() = default;
+
+	explicit FPCGKernelAttributeDesc(int32 InIndex, EPCGKernelAttributeType InType, FPCGAttributeIdentifier InIdentifier)
+		: AttributeId(InIndex)
+		, AttributeKey(InIdentifier, InType)
+	{
+	}
+
+	explicit FPCGKernelAttributeDesc(int32 InIndex, EPCGKernelAttributeType InType, FPCGAttributeIdentifier InIdentifier, const TArray<int32>* InUniqueStringKeys)
+		: AttributeId(InIndex)
+		, AttributeKey(InIdentifier, InType)
+	{
+		if (InUniqueStringKeys)
+		{
+			UniqueStringKeys = *InUniqueStringKeys;
+		}
+	}
+
+	const TArray<int32>& GetUniqueStringKeys() const { return UniqueStringKeys; }
+
+	void AddUniqueStringKeys(const TArray<int32>& InOtherStringKeys);
+
+	void SetStringKeys(const TConstArrayView<int32>& InStringKeys);
+
+	bool operator==(const FPCGKernelAttributeDesc& Other) const;
+
+public:
+	int32 AttributeId = INDEX_NONE;
+
+	FPCGKernelAttributeKey AttributeKey{};
+
+private:
+	/* All possible string keys arriving on this attribute (string keys are indices into the string table in the data binding). */
+	TArray<int32> UniqueStringKeys;
+};
+
+/** Data description for a single data object (UPCGData). */
+struct FPCGDataDesc
+{
+public:
+	FPCGDataDesc() = default;
+	FPCGDataDesc(FPCGDataTypeIdentifier InType, int32 InElementCount);
+	FPCGDataDesc(FPCGDataTypeIdentifier InType, FIntPoint InElementCount);
+	FPCGDataDesc(FPCGDataTypeIdentifier InType, FIntVector3 InElementCount);
+	FPCGDataDesc(FPCGDataTypeIdentifier InType, FIntVector4 InElementCount);
+	FPCGDataDesc(const FPCGTaggedData& InTaggedData, const UPCGDataBinding* InBinding);
+
+	uint64 ComputePackedSizeBytes() const;
+
+	bool HasElementsMetadataDomainAttributes() const;
+
+	bool ContainsAttribute(FPCGAttributeIdentifier InAttributeIdentifier) const;
+
+	bool ContainsAttribute(FPCGAttributeIdentifier InAttributeIdentifier, EPCGKernelAttributeType InAttributeType) const;
+
+	void AddAttribute(FPCGKernelAttributeKey InAttribute, const UPCGDataBinding* InBinding, const TArray<int32>* InOptionalUniqueStringKeys = nullptr);
+
+	const FPCGDataTypeIdentifier& GetType() const { return Type; }
+
+	FIntVector4 GetElementCount() const { return ElementCount; }
+	FIntVector4 GetElementCountForAttribute(const FPCGKernelAttributeDesc& AttributeDesc) const;
+	EPCGElementDimension GetElementDimension() const { return ElementDimension; }
+	int32 ComputeTotalElementCount() const;
+
+	void AddElementCount(int32 InElementCountToAdd);
+	void AddElementCount(FIntPoint InElementCountToAdd);
+	void AddElementCount(FIntVector3 InElementCountToAdd);
+	void AddElementCount(FIntVector4 InElementCountToAdd);
+	void ScaleElementCount(int32 InMultiplier) { ElementCount *= InMultiplier; }
+
+	FPCGDataDesc& CombineElementCount(const FPCGDataDesc& Other, EPCGElementMultiplicity Multiplicity);
+	static FPCGDataDesc CombineElementCount(const FPCGDataDesc& A, const FPCGDataDesc& B, EPCGElementMultiplicity Multiplicity);
+
+	void AllocateProperties(EPCGPointNativeProperties InProperties) { AllocatedPointProperties |= InProperties; }
+	EPCGPointNativeProperties GetAllocatedProperties() const { return AllocatedPointProperties; }
+
+	bool IsAttributeAllocated(const FPCGKernelAttributeDesc& InAttributeDesc) const { return IsAttributeAllocated(InAttributeDesc.AttributeId); }
+	bool IsAttributeAllocated(int32 InAttributeId) const;
+
+private:
+	void InitializeAttributeDescs(const UPCGData* InData, const UPCGDataBinding* InBinding);
+
+public:
+	// @todo_pcg: Encapsulate these members.
+	TArray<FPCGKernelAttributeDesc> AttributeDescs;
+	TArray<int32, TInlineAllocator<4>> TagStringKeys;
+
+private:
+	FPCGDataTypeIdentifier Type = EPCGDataType::Point;
+	FIntVector4 ElementCount = FIntVector4::ZeroValue;
+	EPCGElementDimension ElementDimension = EPCGElementDimension::One;
+
+	/** Properties that should be fully allocated for this data. Only applies to EPCGDataType::Point currently. */
+	EPCGPointNativeProperties AllocatedPointProperties = EPCGPointNativeProperties::None;
+};
+
+/** Data description for a data collection (FPCGDataCollection). */
+struct FPCGDataCollectionDesc
+{
+	/** FPCGDataCollectionDesc should not be created directly, use MakeShared(). */
+	FPCGDataCollectionDesc(const FPCGDataCollectionDesc&) = delete;
+	FPCGDataCollectionDesc& operator=(const FPCGDataCollectionDesc&) = delete;
+	FPCGDataCollectionDesc(FPCGDataCollectionDesc&&) = delete;
+	FPCGDataCollectionDesc& operator=(FPCGDataCollectionDesc&&) = delete;
+
+	/** Create an FPCGDataCollectionDesc. */
+	static PCG_API TSharedPtr<FPCGDataCollectionDesc> MakeShared();
+
+	/** Create an FPCGDataCollectionDesc from another FPCGDataCollectionDesc. */
+	static PCG_API TSharedPtr<FPCGDataCollectionDesc> MakeSharedFrom(const TSharedPtr<const FPCGDataCollectionDesc> InOtherDataDesc);
+
+	/** Computes the size (in bytes) of the header portion of the packed data collection buffer. */
+	uint32 ComputePackedHeaderSizeBytes() const;
+
+	/** Computes the size (in bytes) of the data collection after packing. */
+	uint64 ComputePackedSizeBytes() const;
+
+	void WriteHeader(TArray<uint32>& OutPackedDataCollectionHeader) const;
+
+	/** Pack a data collection into the GPU data format. DataDescs defines which attributes are packed. */
+	void PackDataCollection(const FPCGDataCollection& InDataCollection, FName InPin, const UPCGDataBinding* InDataBinding, TArray<uint32>& OutPackedDataCollection) const;
+
+	/** Unpack a buffer of 8-bit uints to a data collection. */
+	EPCGUnpackDataCollectionResult UnpackDataCollection(FPCGContext* InContext, const TArray<uint8>& InPackedData, FName InPin, const TArray<FString>& InStringTable, FPCGDataCollection& OutDataCollection) const;
+
+	/** Compute total number of processing elements. */
+	PCG_API uint32 ComputeTotalElementCount() const;
+
+	/** Compute total number of processing elements of the given type. */
+	uint32 ComputeTotalElementCount(FPCGDataTypeIdentifier InDataType) const;
+
+	/** Get description of first attribute with matching identifier in input data. Returns true if such an attribute found and also signals whether multiple matching attributes
+	* with conflicting names are present. */
+	bool GetAttributeDesc(FPCGAttributeIdentifier InAttributeIdentifier, FPCGKernelAttributeDesc& OutAttributeDesc, bool& bOutConflictingTypesFound, bool& bOutPresentOnAllData) const;
+
+	bool ContainsAttributeOnAnyData(FPCGAttributeIdentifier InAttributeIdentifier) const;
+
+	/** Makes attribute present on all data. If data has existing attribute with same name then the given type will be applied. */
+	void AddAttributeToAllData(FPCGKernelAttributeKey InAttribute, const UPCGDataBinding* InBinding, const TArray<int32>* InOptionalUniqueStringKeys = nullptr);
+
+	void AllocatePropertiesForAllData(EPCGPointNativeProperties InProperties);
+
+	void GetUniqueStringKeyValues(int32 InAttributeId, TArray<int32>& OutUniqueStringKeys) const;
+
+	int GetNumStringKeyValues(int32 InAttributeId) const;
+
+	/** Description of each data in this data collection. */
+	TArray<FPCGDataDesc> DataDescs;
+
+private:
+	FPCGDataCollectionDesc() = default;
+};
+
+/** A proxy for a data collection residing in a GPU buffer along with functionality to retrieve the data on the CPU. Holds onto GPU memory. */
+struct FPCGProxyForGPUDataCollection : public TSharedFromThis<FPCGProxyForGPUDataCollection>
+{
+public:
+	FPCGProxyForGPUDataCollection(TRefCountPtr<FRDGPooledBuffer> InBuffer, uint32 InBufferSizeBytes, TSharedPtr<const FPCGDataCollectionDesc> InDescription, const TArray<FString>& InStringTable);
+
+	/** Populates a CPU data object representing the GPU data for the given index, performing a readback from GPU->CPU if required. */
+	bool GetCPUData(FPCGContext* InContext, int32 InDataIndex, FPCGTaggedData& OutData);
+
+	const TRefCountPtr<FRDGPooledBuffer>& GetBuffer() const { return Buffer; }
+
+	uint32 GetBufferSizeBytes() const { return BufferSizeBytes; }
+
+	TSharedPtr<const FPCGDataCollectionDesc> GetDescription() const { return Description; }
+
+	const TArray<FString>& GetStringTable() const { return StringTable; }
+
+protected:
+	/** Persistent GPU buffer that can be read back. Buffer will be freed when this ref count is 0. */
+	TRefCountPtr<FRDGPooledBuffer> Buffer;
+
+	uint32 BufferSizeBytes = 0;
+
+	TSharedPtr<const FPCGDataCollectionDesc> Description = nullptr;
+
+	/** Used to comprehend string IDs in buffer. */
+	TArray<FString> StringTable;
+
+	/** Read back data. Populated once upon first readback request. */
+	TArray<const FPCGTaggedData> ReadbackData;
+	TArray<TStrongObjectPtr<const UPCGData>> ReadbackDataRefs;
+
+	TSharedPtr<FRHIGPUBufferReadback> ReadbackRequest;
+
+	TArray<uint8> RawReadbackData;
+
+	std::atomic<bool> bReadbackDataProcessed = false;
+
+	bool bReadbackDataArrived = false;
+
+	UE::FSpinLock ReadbackLock;
+};
+
+namespace PCGDataForGPUHelpers
+{
+	/** Returns GPU type that will be used to represent the given metadata type. */
+	EPCGKernelAttributeType GetAttributeTypeFromMetadataType(EPCGMetadataTypes MetadataType);
+
+	/** Compute how attributes can be packed to custom floats. */
+	void ComputeCustomFloatPacking(
+		FPCGContext* InContext,
+		const UPCGSettings* InSettings,
+		TConstArrayView<FPCGAttributePropertyInputSelector> InAttributeSelectors,
+		const UPCGDataBinding* InBinding,
+		const TSharedPtr<const FPCGDataCollectionDesc> InDataCollectionDescription,
+		uint32& OutCustomFloatCount,
+		TArray<FUint32Vector4>& OutAttributeIdOffsetStrides);
+}
diff --git a/Engine/Plugins/PCG/Source/PCG/Public/Compute/PCGKernelHelpers.h b/Engine/Plugins/PCG/Source/PCG/Public/Compute/PCGKernelHelpers.h
new file mode 100644
--- /dev/null
+++ b/Engine/Plugins/PCG/Source/PCG/Public/Compute/PCGKernelHelpers.h
@@ -0,0 +1,67 @@
+// Copyright Epic Games, Inc. All Rights Reserved.
+
+#pragma once
+
+#include "PCGCommon.h"
+#include "PCGSettings.h"
+#include "Compute/PCGComputeCommon.h"
+#include "Compute/PCGComputeKernel.h"
+#include "Graph/PCGGPUGraphCompilationContext.h"
+
+class UPCGSettings;
+
+namespace PCGKernelHelpers
+{
+#if WITH_EDITOR
+	struct FCreateKernelParams
+	{
+		FCreateKernelParams(UObject* InObjectOuter, const UPCGSettings* InSettings = nullptr)
+			: ObjectOuter(InObjectOuter)
+			, OwnerSettings(InSettings)
+			, bDumpDataDescriptions(InSettings && InSettings->bDumpDataDescriptions)
+		{
+		}
+
+		UObject* ObjectOuter = nullptr;
+
+		bool bRequiresOverridableParams = false;
+
+		const UPCGSettings* OwnerSettings = nullptr;
+
+		bool bDumpDataDescriptions = false;
+
+		/** Create edges from node input pins to kernel inputs, assuming label is same on both. */
+		TArray<FName> NodeInputPinsToWire = { PCGPinConstants::DefaultInputLabel };
+
+		/** Create edges from node output pins to kernel outputs, assuming label is same on both. */
+		TArray<FName> NodeOutputPinsToWire = { PCGPinConstants::DefaultOutputLabel };
+	};
+
+	template<class KernelType>
+	KernelType* CreateKernel(FPCGGPUCompilationContext& InCompilationContext, const FCreateKernelParams& InParams, TArray<UPCGComputeKernel*>& OutKernels, TArray<FPCGKernelEdge>& OutKernelEdges)
+	{
+		FPCGComputeKernelParams KernelParams;
+		KernelParams.Settings = InParams.OwnerSettings;
+		KernelParams.bLogDescriptions = InParams.bDumpDataDescriptions;
+		KernelParams.bRequiresOverridableParams = InParams.bRequiresOverridableParams;
+
+		KernelType* Kernel = InCompilationContext.NewObject_AnyThread<KernelType>(InParams.ObjectOuter);
+
+		Kernel->Initialize(KernelParams);
+		OutKernels.Add(Kernel);
+
+		// Connect node pins to kernel pins
+		for (FName PinLabel : InParams.NodeInputPinsToWire)
+		{
+			OutKernelEdges.Emplace(FPCGPinReference(PinLabel), FPCGPinReference(Kernel, PinLabel));
+		}
+
+		for (FName PinLabel : InParams.NodeOutputPinsToWire)
+		{
+			OutKernelEdges.Emplace(FPCGPinReference(Kernel, PinLabel), FPCGPinReference(PinLabel));
+		}
+
+		return CastChecked<KernelType>(Kernel);
+	}
+#endif // WITH_EDITOR
+}
diff --git a/Engine/Plugins/PCG/Source/PCG/Public/Compute/PCGKernelParamsCache.h b/Engine/Plugins/PCG/Source/PCG/Public/Compute/PCGKernelParamsCache.h
new file mode 100644
--- /dev/null
+++ b/Engine/Plugins/PCG/Source/PCG/Public/Compute/PCGKernelParamsCache.h
@@ -0,0 +1,92 @@
+// Copyright Epic Games, Inc. All Rights Reserved.
+
+#pragma once
+
+#include "Compute/PCGComputeKernel.h"
+
+#include "StructUtils/PropertyBag.h"
+
+class UPCGDataBinding;
+class UPCGSettings;
+struct FPCGContext;
+
+/** Parameters for compute kernel execution. */
+struct FPCGKernelParams
+{
+public:
+	bool GetValueBool(FName InParamName) const;
+	int GetValueInt(FName InParamName) const;
+	FName GetValueName(FName InParamName) const;
+
+	template<typename EnumValueType>
+	EnumValueType GetValueEnum(FName InParamName) const
+	{
+		static_assert(TIsEnum<EnumValueType>::Value, "Should only call this with enum types.");
+
+		EnumValueType Value{};
+
+		const FEnumProperty* Property = nullptr;
+		const void* ValuePtr = nullptr;
+
+		if (const FPropertyAndAddress* PropertyAndAddress = OverriddenProperties.Find(InParamName))
+		{
+			Property = CastField<FEnumProperty>(PropertyAndAddress->Key);
+			ValuePtr = PropertyAndAddress->Value;
+		}
+		else if (const UPCGSettings* Settings = WeakSettings.Get())
+		{
+			Property = CastField<FEnumProperty>(Settings->GetClass()->FindPropertyByName(InParamName));
+			ValuePtr = Property ? Property->ContainerPtrToValuePtr<void>(Settings) : nullptr;
+		}
+
+		const FNumericProperty* NumericProperty = Property ? Property->GetUnderlyingProperty() : nullptr;
+		return static_cast<EnumValueType>(ensure(NumericProperty && ValuePtr) ? NumericProperty->GetSignedIntPropertyValue(ValuePtr) : 0);
+	}
+
+private:
+	/** Optional underlying settings object to get default values that are not overridden. */
+	TWeakObjectPtr<const UPCGSettings> WeakSettings;
+
+	/** Property bag instance that mirrors the WeakSettings. Contains the actual overridden values. */
+	FInstancedPropertyBag OverriddenValues;
+
+	using FPropertyAndAddress = TPair<const FProperty*, const void*>;
+
+	/** Maps param name to its property and an address to read the overridden value. This will only be populated for params which are overidden and read back for CPU logic. */
+	TMap<FName, FPropertyAndAddress> OverriddenProperties;
+
+	friend struct FPCGKernelParamsCache;
+};
+
+/** Helper struct to cache parameters for compute kernels. */
+struct FPCGKernelParamsCache
+{
+public:
+	/** Performs a readback of any kernels with overridable parameters that require readback, and caches them for quick access from kernels and data providers.
+	* Returns false while readback is in progress.
+	*/
+	bool Initialize(FPCGContext* InContext, UPCGDataBinding* InDataBinding);
+
+	/** Look up the cached parameter struct for a kernel. */
+	const FPCGKernelParams* GetCachedKernelParams(int32 InKernelIndex) const;
+
+	/** Clear the cache. */
+	void Reset();
+
+private:
+	/** Cache readback indices for any overridable param data that require readback. */
+	void CacheOverridableParamDataIndices(FPCGContext* InContext, const UPCGDataBinding* InDataBinding);
+
+	/** Poll readback for overridable param data. Returns false while readback is in progress.  */
+	bool ReadbackOverridableParamData(UPCGDataBinding* InDataBinding);
+
+	/** Create and cache a property bag instance to store kernel param values and apply overrides if they exist. */
+	void CreateAndCacheKernelParams(FPCGContext* InContext, const UPCGDataBinding* InDataBinding);
+
+private:
+	/** Maps kernels to their cached kernel parameter values. */
+	TMap</*KernelIndex*/int32, FPCGKernelParams> KernelParams;
+
+	/** Maps kernels to their overridable params and an index into the DataBinding input data collection that holds the corresponding param data. */
+	TMap</*KernelIndex*/int32, TArray<TPair<FPCGKernelOverridableParam, int32>>> OverridableParamDataIndices;
+};
diff --git a/Engine/Plugins/PCG/Source/PCG/Public/Compute/PCGPinPropertiesGPU.h b/Engine/Plugins/PCG/Source/PCG/Public/Compute/PCGPinPropertiesGPU.h
new file mode 100644
--- /dev/null
+++ b/Engine/Plugins/PCG/Source/PCG/Public/Compute/PCGPinPropertiesGPU.h
@@ -0,0 +1,176 @@
+// Copyright Epic Games, Inc. All Rights Reserved.
+
+#pragma once
+
+#include "PCGPin.h"
+#include "Compute/PCGDataForGPU.h"
+
+#include "PCGPinPropertiesGPU.generated.h"
+
+#define UE_API PCG_API
+
+/** Method for initializing an output pin. */
+UENUM()
+enum class EPCGPinInitMode : uint8
+{
+	FromInputPins UMETA(Tooltip = "Initialize from one or more input pins."),
+	Custom UMETA(Tooltip = "Initialize from fixed values."),
+};
+
+/** Method for computing data count. */
+UENUM()
+enum class EPCGDataCountMode : uint8
+{
+	FromInputData,
+	Fixed,
+};
+
+/** Method for computing element count. */
+UENUM()
+enum class EPCGElementCountMode : uint8
+{
+	FromInputData,
+	Fixed,
+};
+
+/** Method for combining two or more data counts. */
+UENUM()
+enum class EPCGDataMultiplicity : uint8
+{
+	Pairwise UMETA(Tooltip = "A data item will be produced for each pair/tuple/etc of input data items across the input pins. Requires all pins to have 1 or N data items."),
+	CartesianProduct UMETA(Tooltip = "If there are two input pins with N and M data items respectively, the output will have NxM data items."),
+};
+
+/** Method for combining two or more element counts. */
+UENUM()
+enum class EPCGElementMultiplicity : uint8
+{
+	Product UMETA(Tooltip = "Element count will be the product of the elements in each pair/tuple/etc of input data."),
+	Sum UMETA(Tooltip = "Element count will be the sum of the elements in each pair/tuple/etc of input data."),
+};
+
+/** Method for inheriting attribute data from input pins. */
+UENUM()
+enum class EPCGAttributeInheritanceMode : uint8
+{
+	None UMETA(Tooltip = "No attributes inherited from initialization pins."),
+	CopyAttributeSetup UMETA(Tooltip = "Take attribute names and types from initialization pins. Pins declared first will take priority during conflicts. Does not copy values."),
+	// TODO: Mode to automatically copy values. Requires code-gen to loop on batches of elements for pins based on the total number of elements and total number of threads.
+	//CopyAttributeSetupAndValues UMETA(Tooltip = "Take attribute names, types, and initial values from initialization pins. Pins declared first will take priority during conflicts."),
+};
+
+/** Helper struct to nest GPU pin properties inside a UI category. */
+USTRUCT(BlueprintType)
+struct FPCGPinPropertiesGPUStruct
+{
+	GENERATED_BODY()
+
+	/** How the output data for this pin will be initialized. */
+	UPROPERTY(BlueprintReadWrite, EditAnywhere, Category = Settings, meta = (EditCondition = "bAllowEditInitMode", EditConditionHides, HideEditConditionToggle))
+	EPCGPinInitMode InitializationMode = EPCGPinInitMode::FromInputPins;
+
+	/** Input pins to initialize this pin's data from. */
+	UPROPERTY(BlueprintReadWrite, EditAnywhere, Category = Settings, meta = (GetOptions = "GetInputPinNames", EditCondition = "bAllowEditInitMode && InitializationMode == EPCGPinInitMode::FromInputPins", EditConditionHides))
+	TArray<FName> PinsToInititalizeFrom;
+
+	/** How the number of data is determined. */
+	UPROPERTY(BlueprintReadWrite, EditAnywhere, Category = Settings, meta = (EditCondition = "bAllowEditInitMode && bAllowEditDataCount && InitializationMode == EPCGPinInitMode::FromInputPins", EditConditionHides))
+	EPCGDataCountMode DataCountMode = EPCGDataCountMode::FromInputData;
+
+	/** How to combine data counts. */
+	UPROPERTY(BlueprintReadWrite, EditAnywhere, Category = Settings, meta = (EditCondition = "bAllowEditInitMode && bAllowEditDataCount && InitializationMode == EPCGPinInitMode::FromInputPins && DataCountMode == EPCGDataCountMode::FromInputData && bMultipleInitPins", EditConditionHides))
+	EPCGDataMultiplicity DataMultiplicity = EPCGDataMultiplicity::Pairwise;
+
+	/** Number of data to create. */
+	UPROPERTY(BlueprintReadWrite, EditAnywhere, Category = Settings, meta = (EditCondition = "bAllowEditInitMode && (InitializationMode == EPCGPinInitMode::Custom || DataCountMode == EPCGDataCountMode::Fixed)", EditConditionHides, ClampMin = 1))
+	int DataCount = 1;
+
+	/** How the number of elements is determined. */
+	UPROPERTY(BlueprintReadWrite, EditAnywhere, Category = Settings, meta = (EditCondition = "bAllowEditInitMode && InitializationMode == EPCGPinInitMode::FromInputPins", EditConditionHides))
+	EPCGElementCountMode ElementCountMode = EPCGElementCountMode::FromInputData;
+
+	/** How to combine element counts. */
+	UPROPERTY(BlueprintReadWrite, EditAnywhere, Category = Settings, meta = (EditCondition = "bAllowEditInitMode && InitializationMode == EPCGPinInitMode::FromInputPins && ElementCountMode == EPCGElementCountMode::FromInputData && bMultipleInitPins", EditConditionHides))
+	EPCGElementMultiplicity ElementMultiplicity = EPCGElementMultiplicity::Product;
+
+	/** Fixed number of elements to create in each output data. */
+	UPROPERTY(BlueprintReadWrite, EditAnywhere, Category = Settings, DisplayName = "Num Elements", meta = (EditCondition = "bAllowEditInitMode && (InitializationMode == EPCGPinInitMode::Custom || ElementCountMode == EPCGElementCountMode::Fixed) && !bShowTexturePinSettings", EditConditionHides, ClampMin = 1))
+	int ElementCount = 1;
+
+	// @todo_pcg: Show texture data type info, e.g. Dimension, Format, etc.
+
+	/** Fixed number of elements to create in each output data. */
+	UPROPERTY(BlueprintReadWrite, EditAnywhere, Category = Settings, DisplayName = "Num Elements", meta = (EditCondition = "bAllowEditInitMode && (InitializationMode == EPCGPinInitMode::Custom || ElementCountMode == EPCGElementCountMode::Fixed) && bShowTexturePinSettings", EditConditionHides, ClampMin = 1))
+	FIntPoint NumElements2D = FIntPoint(64, 64);
+
+	/** Scalar on the fixed number of elements to create in each output data. */
+	UPROPERTY(BlueprintReadWrite, EditAnywhere, Category = Settings, meta = (EditCondition = "bAllowEditInitMode && (InitializationMode == EPCGPinInitMode::Custom || ElementCountMode == EPCGElementCountMode::FromInputData)", EditConditionHides, ClampMin = 1))
+	int ElementCountMultiplier = 1;
+
+	/** How to inherit attribute names, types, and values. */
+	UPROPERTY(BlueprintReadWrite, EditAnywhere, Category = Settings, meta = (EditCondition = "bAllowEditInitMode && InitializationMode == EPCGPinInitMode::FromInputPins && !bShowTexturePinSettings", EditConditionHides))
+	EPCGAttributeInheritanceMode AttributeInheritanceMode = EPCGAttributeInheritanceMode::CopyAttributeSetup;
+
+	/** Add entries to create new attributes on data emitted by this pin. */
+	UPROPERTY(EditAnywhere, DisplayName = "Attributes to Create", Category = Settings, meta = (EditCondition = "!bShowTexturePinSettings", EditConditionHides, HideEditConditionToggle))
+	TArray<FPCGKernelAttributeKey> CreatedKernelAttributeKeys;
+
+#if WITH_EDITORONLY_DATA
+	UPROPERTY(Transient)
+	bool bAllowEditInitMode = false;
+
+	/** Some types force single data. We should not allow data count to be edited for such types. */
+	UPROPERTY(Transient)
+	bool bAllowEditDataCount = false;
+
+	UPROPERTY(Transient)
+	bool bMultipleInitPins = false;
+
+	UPROPERTY(Transient)
+	bool bShowTexturePinSettings = false;
+#endif // WITH_EDITORONLY_DATA
+};
+
+/** An extension of the pin properties that adds hints for GPU specific properties, such as buffer size and data layout. */
+USTRUCT(BlueprintType)
+struct FPCGPinPropertiesGPU : public FPCGPinProperties
+{
+	GENERATED_BODY()
+
+public:
+	FPCGPinPropertiesGPU() = default;
+
+	explicit FPCGPinPropertiesGPU(const FName& InLabel, FPCGDataTypeIdentifier InAllowedTypes)
+		: FPCGPinProperties(InLabel, InAllowedTypes)
+	{}
+
+	explicit FPCGPinPropertiesGPU(const FPCGPinProperties& InPinProps)
+		: FPCGPinProperties(InPinProps.Label, InPinProps.AllowedTypes)
+	{}
+
+	UE_API uint32 GetElementCountMultiplier() const;
+
+#if WITH_EDITOR
+	UE_API bool CanEditChange(const FEditPropertyChain& PropertyChain) const;
+#endif
+
+public:
+	UPROPERTY(BlueprintReadWrite, EditAnywhere, Category = Settings, meta = (DisplayName = "GPU Properties", DisplayAfter = "Tooltip", EditCondition = "bShowPropertiesGPU", EditConditionHides, HideEditConditionToggle))
+	FPCGPinPropertiesGPUStruct PropertiesGPU;
+
+#if WITH_EDITORONLY_DATA
+	UPROPERTY(Transient)
+	bool bShowPropertiesGPU = true;
+#endif // WITH_EDITORONLY_DATA
+};
+
+template<>
+struct TStructOpsTypeTraits<FPCGPinPropertiesGPU> : public TStructOpsTypeTraitsBase2<FPCGPinPropertiesGPU>
+{
+	enum
+	{
+		WithCanEditChange = true,
+	};
+};
+
+#undef UE_API
diff --git a/Engine/Plugins/PCG/Source/PCG/Public/Graph/PCGGPUGraphCompilationContext.h b/Engine/Plugins/PCG/Source/PCG/Public/Graph/PCGGPUGraphCompilationContext.h
new file mode 100644
--- /dev/null
+++ b/Engine/Plugins/PCG/Source/PCG/Public/Graph/PCGGPUGraphCompilationContext.h
@@ -0,0 +1,86 @@
+// Copyright Epic Games, Inc. All Rights Reserved.
+
+#pragma once
+
+#if WITH_EDITOR
+
+#include "PCGCommon.h"
+
+class FPCGGraphCompiler;
+class UPCGComputeGraph;
+struct FPCGKernelAttributeTable;
+
+/** Current context/state used during compilation of GPU graphs. */
+struct FPCGGPUCompilationContext
+{
+	explicit FPCGGPUCompilationContext(FPCGGraphCompiler& InGraphCompiler)
+		: GraphCompiler(InGraphCompiler)
+	{
+	}
+
+	~FPCGGPUCompilationContext()
+	{
+		// Remove Async Flags from objects created during compilation so that they can get tracked by the GarbageCollector properly.
+		for (TObjectPtr<UObject>& AsyncObject : AsyncObjects)
+		{
+			if (ensure(AsyncObject->HasAnyInternalFlags(EInternalObjectFlags::Async)))
+			{
+				AsyncObject->ClearInternalFlags(EInternalObjectFlags::Async);
+				ForEachObjectWithOuter(
+					AsyncObject,
+					[&AsyncObjects = AsyncObjects](UObject* SubObject)
+					{
+						if (AsyncObjects.Contains(SubObject))
+						{
+							return;
+						}
+
+						SubObject->ClearInternalFlags(EInternalObjectFlags::Async);
+					},
+					true);
+			}
+		}
+	}
+
+	/** Thread safe NewObject. If object created off game thread, object added to AsyncObjects. */
+	template<class T, typename... Args>
+	T* NewObject_AnyThread(Args&&... InArgs)
+	{
+		if (!IsInGameThread())
+		{
+			T* Object = nullptr;
+			{
+				FGCScopeGuard Scope;
+				Object = ::NewObject<T>(std::forward<Args>(InArgs)...);
+			}
+
+			check(Object);
+			AsyncObjects.Add(Object);
+			return Object;
+		}
+
+		return ::NewObject<T>(std::forward<Args>(InArgs)...);
+	}
+
+	FPCGGraphCompiler& GetGraphCompiler() { return GraphCompiler; }
+	TArray<TObjectPtr<UPCGComputeGraph>>& GetCompiledComputeGraphs() { return CompiledComputeGraphs; }
+
+	const FPCGKernelAttributeTable* GetStaticAttributeTable() { return StaticAttributeTable; }
+	void SetStaticAttributeTable(const FPCGKernelAttributeTable* InStaticAttributeTable) { StaticAttributeTable = InStaticAttributeTable; }
+
+	/** Returns the index of the newly added compute graph. */
+	int32 AddCompiledComputeGraph(TObjectPtr<UPCGComputeGraph> NewComputeGraph) { return CompiledComputeGraphs.Add(NewComputeGraph); }
+
+private:
+	// List of objects created during compilation, we need to track them so we can remove their Async flags when compilation is done
+	// so that they can be considered as existing on the main thread (and get properly GCed).
+	TSet<TObjectPtr<UObject>> AsyncObjects;
+
+	FPCGGraphCompiler& GraphCompiler;
+	TArray<TObjectPtr<UPCGComputeGraph>> CompiledComputeGraphs;
+
+	/** Attribute table of the currently compiling compute graph. */
+	const FPCGKernelAttributeTable* StaticAttributeTable = nullptr;
+};
+
+#endif // WITH_EDITOR
