[PCG] Fix GPU TransformPoints ignoring source rotation

Also remove transpose of transform attribute since transform attributes are now transposed in pack/unpack

#rnx

diff --git a/Engine/Plugins/PCG/Shaders/Private/Elements/PCGTransformPoints.usf b/Engine/Plugins/PCG/Shaders/Private/Elements/PCGTransformPoints.usf
--- a/Engine/Plugins/PCG/Shaders/Private/Elements/PCGTransformPoints.usf
+++ b/Engine/Plugins/PCG/Shaders/Private/Elements/PCGTransformPoints.usf
@@ -53,8 +53,7 @@ void PCGTransformPointsCS(uint3 GroupId : SV_GroupID, uint GroupIndex : SV_Group
 
 	if (bApplyToAttribute)
 	{
-		// TODO: If the transform is transposed in pack, remove this transpose.
-		float4x4 InTransform = transpose(In_GetTransform(InDataIndex, InElemIndex, AttributeId));
+		const float4x4 InTransform = In_GetTransform(InDataIndex, InElemIndex, AttributeId);
 
 		InScale = float3(length(InTransform._m00_m10_m20), length(InTransform._m01_m11_m21), length(InTransform._m02_m12_m22));
 
@@ -79,10 +78,17 @@ void PCGTransformPointsCS(uint3 GroupId : SV_GroupID, uint GroupIndex : SV_Group
 	uint Seed = ComputeSeed(GetSeed(), InSeed);
 
 	// Position
-	float3 OutPosition = bAbsoluteOffset ? float3(0, 0, 0) : InPosition;
-	OutPosition.x += lerp(OffsetMin.x, OffsetMax.x, FRand(Seed));
-	OutPosition.y += lerp(OffsetMin.y, OffsetMax.y, FRand(Seed));
-	OutPosition.z += lerp(OffsetMin.z, OffsetMax.z, FRand(Seed));
+	float3 PositionOffset = float3(
+		lerp(OffsetMin.x, OffsetMax.x, FRand(Seed)),
+		lerp(OffsetMin.y, OffsetMax.y, FRand(Seed)),
+		lerp(OffsetMin.z, OffsetMax.z, FRand(Seed)));
+
+	if (!bAbsoluteOffset)
+	{
+		PositionOffset = Quat_RotateVector(InRotation, PositionOffset);
+	}
+
+	const float3 OutPosition = InPosition + PositionOffset;
 
 	// Rotation
 	FQuat OutRotation = bAbsoluteRotation ? FQuat(0, 0, 0, 1) : InRotation;
@@ -126,8 +132,7 @@ void PCGTransformPointsCS(uint3 GroupId : SV_GroupID, uint GroupIndex : SV_Group
 			Axis0.z, Axis1.z, Axis2.z, OutPosition.z,
 			(float3)0.0, 1.0);
 
-		// TODO: If the transform is transposed in unpack, remove this transpose.
-		Out_SetTransform(OutDataIndex, OutElemIndex, AttributeId, transpose(OutTransform));
+		Out_SetTransform(OutDataIndex, OutElemIndex, AttributeId, OutTransform);
 
 		// Copy original
 		Out_SetPosition(OutDataIndex, OutElemIndex, In_GetPosition(InDataIndex, InElemIndex));
