[PCG Break out pcg.Cache.Enabled and pcg.Cache.MemoryBudgetMB into two cvars for editor vs game.

[RN] [PCG] Added specific graph cache enabled and budget CVars for editor worlds vs game worlds.

diff --git a/Engine/Plugins/PCG/Source/PCG/Private/Graph/PCGGraphCache.cpp b/Engine/Plugins/PCG/Source/PCG/Private/Graph/PCGGraphCache.cpp
--- a/Engine/Plugins/PCG/Source/PCG/Private/Graph/PCGGraphCache.cpp
+++ b/Engine/Plugins/PCG/Source/PCG/Private/Graph/PCGGraphCache.cpp
@@ -9,20 +9,30 @@
 #include "HAL/IConsoleManager.h"
 #include "Misc/ScopeRWLock.h"
 
-static TAutoConsoleVariable<bool> CVarCacheEnabled(
-	TEXT("pcg.Cache.Enabled"),
+static TAutoConsoleVariable<bool> CVarCacheEnabledEditor(
+	TEXT("pcg.Cache.Editor.Enabled"),
 	true,
-	TEXT("Enables the cache system."));
+	TEXT("Enables the cache system in editor worlds."));
+
+static TAutoConsoleVariable<bool> CVarCacheEnabledRuntime(
+	TEXT("pcg.Cache.Runtime.Enabled"),
+	false,
+	TEXT("Enables the cache system in runtime game worlds."));
 
 static TAutoConsoleVariable<bool> CVarCacheDebugging(
 	TEXT("pcg.Cache.EnableDebugging"),
 	false,
 	TEXT("Enable various features for debugging the graph cache system."));
 
-static TAutoConsoleVariable<int32> CVarCacheMemoryBudgetMB(
-	TEXT("pcg.Cache.MemoryBudgetMB"),
+static TAutoConsoleVariable<int32> CVarCacheMemoryBudgetEditorMB(
+	TEXT("pcg.Cache.Editor.MemoryBudgetMB"),
 	6144,
-	TEXT("Memory budget for data in cache (MB)."));
+	TEXT("Memory budget for data in cache in editor worlds (MB)."));
+
+static TAutoConsoleVariable<int32> CVarCacheMemoryBudgetRuntimeMB(
+	TEXT("pcg.Cache.Runtime.MemoryBudgetMB"),
+	128,
+	TEXT("Memory budget for data in cache in game worlds (MB)."));
 
 static TAutoConsoleVariable<float> CVarCacheMemoryCleanupRatio(
 	TEXT("pcg.Cache.MemoryCleanupRatio"),
@@ -39,12 +49,44 @@ static TAutoConsoleVariable<bool> CVarValidateElementToCacheEntryKeys(
 	false,
 	TEXT("Validate ElementToCacheEntryKeys acceleration table (debug)."));
 
+// Aliases for old deprecated settings.
+static TAutoConsoleVariable<bool> CVarCacheEnabled_DEPRECATED(
+	TEXT("pcg.Cache.Enabled"),
+	true,
+	TEXT("DEPRECATED (5.7): use pcg.Cache.Editor.Enabled and pcg.Cache.Runtime.Enabled instead."),
+	FConsoleVariableDelegate::CreateLambda([](IConsoleVariable* InVariable)
+	{
+		// Called when user sets cvar from console, or when this cvar is set when launching standalone build. Pushes value to new cvars.
+		UE_LOG(LogPCG, Warning, TEXT("pcg.Cache.Enabled is deprecated in 5.7. use pcg.Cache.Editor.Enabled and pcg.Cache.Runtime.Enabled instead."));
+		check(InVariable);
+		bool bNewValue = true;
+		InVariable->GetValue(bNewValue);
+		CVarCacheEnabledRuntime->Set(bNewValue, ECVF_SetByCode);
+		CVarCacheEnabledEditor->Set(bNewValue, ECVF_SetByCode);
+	}));
+static TAutoConsoleVariable<int32> CVarCacheMemoryBudgetMB_DEPRECATED(
+	TEXT("pcg.Cache.MemoryBudgetMB"),
+	6144,
+	TEXT("DEPRECATED (5.7): use pcg.Cache.Editor.MemoryBudgetMB and pcg.Cache.Runtime.MemoryBudgetMB instead."),
+	FConsoleVariableDelegate::CreateLambda([](IConsoleVariable* InVariable)
+	{
+		// Called when user sets cvar from console, or when this cvar is set when launching standalone build. Pushes value to new cvars.
+		UE_LOG(LogPCG, Warning, TEXT("pcg.Cache.MemoryBudgetMB is deprecated in 5.7. use pcg.Cache.Editor.MemoryBudgetMB and pcg.Cache.Runtime.MemoryBudgetMB instead."));
+		check(InVariable);
+		int32 NewValue = 6144;
+		InVariable->GetValue(NewValue);
+		CVarCacheMemoryBudgetEditorMB->Set(NewValue, ECVF_SetByCode);
+		CVarCacheMemoryBudgetRuntimeMB->Set(NewValue, ECVF_SetByCode);
+	}));
+
 // Initial max number of entries graph cache
 static const int32 GPCGGraphCacheInitialCapacity = 65536;
 
-FPCGGraphCache::FPCGGraphCache()
-	: CacheData(CVarCacheEnabled.GetValueOnGameThread() ? GPCGGraphCacheInitialCapacity : 0)
+FPCGGraphCache::FPCGGraphCache(bool bGameWorld)
+	: bIsGameWorld(bGameWorld)
 {
+	const int32 InitialMaxCapacity = IsEnabled() ? GPCGGraphCacheInitialCapacity : 0;
+	CacheData.Empty(InitialMaxCapacity);
 }
 
 FPCGGraphCache::~FPCGGraphCache()
@@ -52,9 +94,19 @@ FPCGGraphCache::~FPCGGraphCache()
 	ClearCache();
 }
 
+bool FPCGGraphCache::IsEnabled() const
+{
+	return bIsGameWorld ? CVarCacheEnabledRuntime.GetValueOnAnyThread() : CVarCacheEnabledEditor.GetValueOnAnyThread();
+}
+
+int32 FPCGGraphCache::GetMemoryBudgetMB() const
+{
+	return bIsGameWorld ? CVarCacheMemoryBudgetRuntimeMB.GetValueOnAnyThread() : CVarCacheMemoryBudgetEditorMB.GetValueOnAnyThread();
+}
+
 bool FPCGGraphCache::GetFromCache(const FPCGGetFromCacheParams& Params, FPCGDataCollection& OutOutput) const
 {
-	if (!CVarCacheEnabled.GetValueOnAnyThread())
+	if (!IsEnabled())
 	{
 		return false;
 	}
@@ -103,7 +155,7 @@ bool FPCGGraphCache::GetFromCache(const FPCGGetFromCacheParams& Params, FPCGData
 
 void FPCGGraphCache::StoreInCache(const FPCGStoreInCacheParams& Params, const FPCGDataCollection& InOutput)
 {
-	if (!CVarCacheEnabled.GetValueOnAnyThread())
+	if (!IsEnabled())
 	{
 		return;
 	}
@@ -147,7 +199,7 @@ void FPCGGraphCache::ClearCache()
 bool FPCGGraphCache::EnforceMemoryBudget()
 {
 	TRACE_CPUPROFILER_EVENT_SCOPE(FPCGGraphCache::FPCGGraphCache::EnforceMemoryBudget);
-	if (!CVarCacheEnabled.GetValueOnAnyThread())
+	if (!IsEnabled())
 	{
 		return false;
 	}
@@ -157,7 +209,7 @@ bool FPCGGraphCache::EnforceMemoryBudget()
 		return false;
 	}
 
-	const uint64 MemoryBudget = static_cast<uint64>(CVarCacheMemoryBudgetMB.GetValueOnAnyThread()) * 1024 * 1024;
+	const uint64 MemoryBudget = static_cast<uint64>(GetMemoryBudgetMB()) * 1024 * 1024;
 	if (TotalMemoryUsed <= MemoryBudget)
 	{
 		return false;
@@ -399,3 +451,57 @@ bool FPCGGraphCache::IsDebuggingEnabled() const
 {
 	return CVarCacheDebugging.GetValueOnAnyThread();
 }
+
+#if WITH_EDITOR
+// Cvar deprecation. One-shot migration: runs after ini loading / user input via a cvar sink. Only required with editor, in standalone
+// it appears the pcg.Cache.Enabled change delegate fires which handles driving the new cvars.
+static void MigrateCVarsIfNeeded()
+{
+	static bool bDone = false;
+	if (bDone)
+	{
+		return;
+	}
+	bDone = true;
+
+	auto SetFromIni = [](uint32 InFlags, uint32 InSetBy)
+	{
+		// Treat any of these as "user set in ini" (covers typical places users put CVars)
+		return InSetBy == ECVF_SetBySystemSettingsIni || // [SystemSettings] sections in *Engine.ini
+			InSetBy == ECVF_SetByConsoleVariablesIni || // ConsoleVariables.ini
+			InSetBy == ECVF_SetByDeviceProfile || // DeviceProfiles.ini
+			(InFlags & ECVF_CreatedFromIni) != 0; // Created by ini before registration
+	};
+
+	const uint32 CacheEnabledFlags = CVarCacheEnabled_DEPRECATED->GetFlags();
+	const uint32 CacheEnabledSetBy = (CacheEnabledFlags & ECVF_SetByMask);
+
+	if (SetFromIni(CacheEnabledFlags, CacheEnabledSetBy))
+	{
+		const bool bValue = CVarCacheEnabled_DEPRECATED->GetBool();
+
+		// Copy with the same SetBy priority so behavior matches the user's source of truth.
+		CVarCacheEnabledRuntime->Set(bValue, (EConsoleVariableFlags)CacheEnabledSetBy);
+		CVarCacheEnabledEditor->Set(bValue, (EConsoleVariableFlags)CacheEnabledSetBy);
+
+		UE_LOG(LogPCG, Warning, TEXT("Cvar 'pcg.Cache.Enabled' was set from ini but is deprecated in 5.7. Migrated value to pcg.Cache.Runtime.Enabled=%d, pcg.Cache.Editor.Enabled=%d (SetBy=0x%x)."), bValue, bValue, CacheEnabledSetBy);
+	}
+
+	const uint32 CacheBudgetFlags = CVarCacheMemoryBudgetMB_DEPRECATED->GetFlags();
+	const uint32 CacheBudgetSetBy = (CacheBudgetFlags & ECVF_SetByMask);
+
+	if (SetFromIni(CacheBudgetFlags, CacheBudgetSetBy))
+	{
+		const int32 Value = CVarCacheMemoryBudgetMB_DEPRECATED->GetInt();
+
+		// Copy with the same SetBy priority so behavior matches the user's source of truth.
+		CVarCacheMemoryBudgetRuntimeMB->Set(Value, (EConsoleVariableFlags)CacheBudgetSetBy);
+		CVarCacheMemoryBudgetEditorMB->Set(Value, (EConsoleVariableFlags)CacheBudgetSetBy);
+
+		UE_LOG(LogPCG, Warning, TEXT("Cvar 'pcg.Cache.MemoryBudgetMB' was set from ini but is deprecated in 5.7. Migrated value to pcg.Cache.Runtime.MemoryBudgetMB=%d, pcg.Cache.Editor.MemoryBudgetMB=%d (SetBy=0x%x)."), Value, Value, CacheBudgetSetBy);
+	}
+}
+
+// Register a sink so this runs at the right time (after ini load / on first cvar change tick).
+static FAutoConsoleVariableSink GCVarMigrationSink(FConsoleCommandDelegate::CreateStatic(&MigrateCVarsIfNeeded));
+#endif // WITH_EDITOR
diff --git a/Engine/Plugins/PCG/Source/PCG/Private/Graph/PCGGraphCache.h b/Engine/Plugins/PCG/Source/PCG/Private/Graph/PCGGraphCache.h
--- a/Engine/Plugins/PCG/Source/PCG/Private/Graph/PCGGraphCache.h
+++ b/Engine/Plugins/PCG/Source/PCG/Private/Graph/PCGGraphCache.h
@@ -65,9 +65,13 @@ private:
 class FPCGGraphCache : public FGCObject, public IPCGGraphCache
 {
 public:
-	FPCGGraphCache();
+	FPCGGraphCache(bool bGameWorld);
 	~FPCGGraphCache();
 
+	bool IsEnabled() const;
+
+	int32 GetMemoryBudgetMB() const;
+
 	/** Returns true if data was found from the cache, in which case the outputs are written in OutOutput. InNode is optional and for logging only. */
 	virtual bool GetFromCache(const FPCGGetFromCacheParams& Params, FPCGDataCollection& OutCollection) const override;
 
@@ -122,5 +126,7 @@ private:
 	/** Total memory usage by all data objects in cache. */
 	uint64 TotalMemoryUsed = 0;
 
+	bool bIsGameWorld = true;
+
 	mutable UE::FSpinLock CacheLock;
 };
diff --git a/Engine/Plugins/PCG/Source/PCG/Private/Graph/PCGGraphExecutor.cpp b/Engine/Plugins/PCG/Source/PCG/Private/Graph/PCGGraphExecutor.cpp
--- a/Engine/Plugins/PCG/Source/PCG/Private/Graph/PCGGraphExecutor.cpp
+++ b/Engine/Plugins/PCG/Source/PCG/Private/Graph/PCGGraphExecutor.cpp
@@ -215,7 +215,8 @@ FPCGGraphExecutor::FPCGGraphExecutor()
 }
 
 FPCGGraphExecutor::FPCGGraphExecutor(UWorld* InWorld)
-	: World(InWorld)
+	: GraphCache(InWorld && InWorld->IsGameWorld())
+	, World(InWorld)
 {
 	GameThreadHandler = MakeShared<FGameThreadHandler>(this);
 }
