[PCG] Restore "Attribute Set" instead of "Param" for allowed types

* Was "Attribute Set" before the refactor and slip back to "Param"

#rnx

diff --git a/Engine/Plugins/PCG/Source/PCG/Public/PCGParamData.h b/Engine/Plugins/PCG/Source/PCG/Public/PCGParamData.h
--- a/Engine/Plugins/PCG/Source/PCG/Public/PCGParamData.h
+++ b/Engine/Plugins/PCG/Source/PCG/Public/PCGParamData.h
@@ -14,7 +14,7 @@ namespace PCGParamDataConstants
 	const FName ElementsDomainName = "Elements";
 }
 
-USTRUCT()
+USTRUCT(meta = (PCG_DataTypeDisplayName="Attribute Set"))
 struct FPCGDataTypeInfoParam : public FPCGDataTypeInfo
 {
 	GENERATED_BODY()
