PCG Paint Mode:
Added Brush Rotation on Mousewheel Up/Down
Added cube brush options for improved WYSISWYG, default still sphere
Painted Points now use the brushes' actor rotation for their transform
Tool Settings now persist between sessions

#virtualized

diff --git a/Engine/Plugins/PCG/Source/PCGEditor/PCGEditor.Build.cs b/Engine/Plugins/PCG/Source/PCGEditor/PCGEditor.Build.cs
--- a/Engine/Plugins/PCG/Source/PCGEditor/PCGEditor.Build.cs
+++ b/Engine/Plugins/PCG/Source/PCGEditor/PCGEditor.Build.cs
@@ -72,7 +72,9 @@ namespace UnrealBuildTool.Rules
 					"WidgetRegistration", 
 					"ModelingComponents",
 					"ModelingComponentsEditorOnly",
-					"Landscape"
+					"Landscape",
+					"GeometryCore",
+					"GeometryFramework"
 				});
 			
 			// @todo_pcg
diff --git a/Engine/Plugins/PCG/Source/PCGEditor/Private/EditorMode/Tools/Gizmos/PCGGizmos.cpp b/Engine/Plugins/PCG/Source/PCGEditor/Private/EditorMode/Tools/Gizmos/PCGGizmos.cpp
new file mode 100644
--- /dev/null
+++ b/Engine/Plugins/PCG/Source/PCGEditor/Private/EditorMode/Tools/Gizmos/PCGGizmos.cpp
@@ -0,0 +1,201 @@
+﻿// Copyright Epic Games, Inc. All Rights Reserved.
+
+#include "EditorMode/Tools/Gizmos/PCGGizmos.h"
+
+#include "InteractiveGizmoManager.h"
+#include "InteractiveToolManager.h"
+#include "PreviewMesh.h"
+#include "ToolDataVisualizer.h"
+#include "Materials/Material.h"
+#include "Materials/MaterialInstanceDynamic.h"
+#include "Generators/MinimalBoxMeshGenerator.h"
+#include "Generators/SphereGenerator.h"
+
+namespace UE::PCG::EditorMode
+{
+	namespace Gizmos
+	{
+		UMaterialInstanceDynamic* GetDefaultVolumeBrushMaterial(UInteractiveToolManager* ToolManager)
+		{
+			if (UMaterial* Material = LoadObject<UMaterial>(nullptr, TEXT("/PCG/Materials/BrushIndicatorMaterial")))
+			{
+				UMaterialInstanceDynamic* MaterialInstance = UMaterialInstanceDynamic::Create(Material, ToolManager);
+				return MaterialInstance;
+			}
+			
+			return nullptr;
+		}
+	}
+}
+
+UInteractiveGizmo* UVolumetricBoxBrushStampIndicatorBuilder::BuildGizmo(const FToolBuilderState& SceneState) const
+{
+	return NewObject<UPCGVolumetricBrushStampIndicator>(SceneState.GizmoManager);
+}
+
+void UPCGVolumetricBrushStampIndicator::Setup()
+{
+	bDrawIndicatorLines = true;
+	bDrawRadiusCircle = true;
+}
+
+void UPCGVolumetricBrushStampIndicator::Render(IToolsContextRenderAPI* RenderAPI)
+{
+	if (bVisible == false)
+	{
+		return;
+	}
+	
+	if (BrushMode == EPCGBrushMode::Sphere)
+	{
+		Super::Render(RenderAPI);		
+	}
+	else if (BrushMode == EPCGBrushMode::Box)
+	{
+		if (bDrawIndicatorLines)
+		{
+			FToolDataVisualizer Draw;
+			Draw.BeginFrame(RenderAPI);
+
+			FRotator ActorRotation = GetPreviewMesh()->GetActor()->GetActorRotation();
+			if (bDrawSecondaryLines)
+			{
+				const bool bCirclesCoincident = bDrawRadiusCircle && FMath::IsNearlyEqual(BrushFalloff, 1.0f);
+				if (!bCirclesCoincident)
+				{
+					float SideLength = BrushRadius * 2 * BrushFalloff;
+					FVector SideA = ActorRotation.RotateVector(FVector(SideLength, 0, 0));
+					FVector SideB = ActorRotation.RotateVector(FVector(0, SideLength, 0));
+					Draw.DrawSquare(BrushPosition, SideA, SideB, LineColor, LineThickness, bDepthTested);
+				}
+			}
+		
+			if (bDrawExtents)
+			{
+				float SideLength = BrushRadius * 2;
+				FVector SideA = ActorRotation.RotateVector(FVector(SideLength, 0, 0));
+				FVector SideB = ActorRotation.RotateVector(FVector(0, SideLength, 0));
+				Draw.DrawSquare(BrushPosition, SideA, SideB, LineColor, LineThickness, bDepthTested);
+			}
+		
+			Draw.EndFrame();
+		}
+	}	
+}
+
+void UPCGVolumetricBrushStampIndicator::Shutdown()
+{
+	if (PreviewMesh)
+	{
+		PreviewMesh->Disconnect();
+	}
+	
+	PreviewMesh = nullptr;
+}
+
+void UPCGVolumetricBrushStampIndicator::SetVolumetricBrushVisible(bool bInVisible)
+{
+	if (PreviewMesh)
+	{
+		PreviewMesh->SetVisible(bInVisible);
+	}
+}
+
+bool UPCGVolumetricBrushStampIndicator::IsVolumetricBrushVisible() const
+{
+	if (PreviewMesh)
+	{
+		return PreviewMesh->IsVisible();
+	}
+	
+	return false;
+}
+
+void UPCGVolumetricBrushStampIndicator::MakeBrushIndicatorMesh(UInteractiveTool* ParentTool, UWorld* World)
+{
+	if (ensure(ParentTool))
+	{
+		PreviewMesh = NewObject<UPreviewMesh>(ParentTool);
+		if (PreviewMesh)
+		{
+			PreviewMesh->CreateInWorld(World, FTransform::Identity);
+			
+			GenerateMesh(BrushMode);
+			
+			if (UMaterialInstanceDynamic* BrushIndicatorMaterial = UE::PCG::EditorMode::Gizmos::GetDefaultVolumeBrushMaterial(ParentTool->GetToolManager()))
+			{
+				PreviewMesh->SetMaterial(BrushIndicatorMaterial);
+			}
+
+			// make sure raytracing is disabled on the brush indicator
+			if (UDynamicMeshComponent* DynamicComponent = Cast<UDynamicMeshComponent>(PreviewMesh->GetRootComponent()))
+			{
+				DynamicComponent->SetEnableRaytracing(false);
+				PreviewMesh->SetShadowsEnabled(false);
+			}
+			
+			AttachedComponent = PreviewMesh->GetRootComponent();
+		}
+	}
+}
+
+UPreviewMesh* UPCGVolumetricBrushStampIndicator::GetPreviewMesh() const
+{
+	if (PreviewMesh)
+	{
+		return PreviewMesh;
+	}
+
+	return nullptr;
+}
+
+void UPCGVolumetricBrushStampIndicator::AddYawRotation(float AddYawRotation)
+{
+	if (PreviewMesh)
+	{
+		FRotator AdditionalRotator(0.f, AddYawRotation, 0.f);
+		PreviewMesh->GetActor()->AddActorLocalRotation(AdditionalRotator);
+	}
+}
+
+void UPCGVolumetricBrushStampIndicator::SetBrushMode(EPCGBrushMode InBrushMode)
+{
+	if(BrushMode == InBrushMode)
+	{
+		return;
+	}
+
+	BrushMode = InBrushMode;
+
+	GenerateMesh(BrushMode);
+}
+
+void UPCGVolumetricBrushStampIndicator::GenerateMesh(EPCGBrushMode InBrushMode)
+{
+	if(PreviewMesh == nullptr)
+	{
+		return;
+	}
+	
+	FDynamicMesh3 Mesh;
+
+	if( BrushMode == EPCGBrushMode::Sphere)
+	{
+		UE::Geometry::FSphereGenerator SphereGenerator;
+		SphereGenerator.NumPhi = SphereGenerator.NumTheta = 32;
+		SphereGenerator.Generate();
+		Mesh = FDynamicMesh3(&SphereGenerator);
+	}
+	else if (BrushMode == EPCGBrushMode::Box)
+	{
+		UE::Geometry::FMinimalBoxMeshGenerator BoxGenerator;
+		BoxGenerator.Box = UE::Geometry::FOrientedBox3d(FVector3d::Zero(), FVector3d::One());
+		BoxGenerator.Generate();
+		Mesh = FDynamicMesh3(&BoxGenerator);
+	}
+
+	if(Mesh.CheckValidity(FDynamicMesh3::FValidityOptions(), UE::Geometry::EValidityCheckFailMode::ReturnOnly))
+	{
+		PreviewMesh->UpdatePreview(&Mesh);
+	}
+}
diff --git a/Engine/Plugins/PCG/Source/PCGEditor/Private/EditorMode/Tools/Line/PCGDrawSplineTool.cpp b/Engine/Plugins/PCG/Source/PCGEditor/Private/EditorMode/Tools/Line/PCGDrawSplineTool.cpp
--- a/Engine/Plugins/PCG/Source/PCGEditor/Private/EditorMode/Tools/Line/PCGDrawSplineTool.cpp
+++ b/Engine/Plugins/PCG/Source/PCGEditor/Private/EditorMode/Tools/Line/PCGDrawSplineTool.cpp
@@ -202,7 +202,12 @@ namespace PCGDrawSplineTool
 UPCGInteractiveToolSettings_SplineBase::UPCGInteractiveToolSettings_SplineBase()
 	: SplineName(GetDefault<UPCGEditorModeSettings>()->DefaultNewSplineComponentName)
 {
-	
+	TInstancedStruct<FPCGRaycastFilterRule_Landscape> LandscapeRule = TInstancedStruct<FPCGRaycastFilterRule_Landscape>::Make();
+	TInstancedStruct<FPCGRaycastFilterRule_Meshes> MeshesRule = TInstancedStruct<FPCGRaycastFilterRule_Meshes>::Make();
+	TInstancedStruct<FPCGRaycastFilterRule_IgnorePCGGeneratedComponents> IgnorePCGGeneratedComponentsRule = TInstancedStruct<FPCGRaycastFilterRule_IgnorePCGGeneratedComponents>::Make();
+	RaycastRuleCollection.RaycastRules.Add(LandscapeRule);
+	RaycastRuleCollection.RaycastRules.Add(MeshesRule);
+	RaycastRuleCollection.RaycastRules.Add(IgnorePCGGeneratedComponentsRule);
 }
 
 USplineComponent* UPCGInteractiveToolSettings_SplineBase::GetWorkingSplineComponent(FName DataInstanceIdentifier) const
@@ -287,13 +292,6 @@ void UPCGDrawSplineToolBase::Setup()
 	ClickOrDragBehavior = NewObject<USingleClickOrDragInputBehavior>();
 	ClickOrDragBehavior->Initialize(this, this);
 	AddInputBehavior(ClickOrDragBehavior);
-
-	TInstancedStruct<FPCGRaycastFilterRule_Landscape> LandscapeRule = TInstancedStruct<FPCGRaycastFilterRule_Landscape>::Make();
-	TInstancedStruct<FPCGRaycastFilterRule_Meshes> MeshesRule = TInstancedStruct<FPCGRaycastFilterRule_Meshes>::Make();
-	TInstancedStruct<FPCGRaycastFilterRule_IgnorePCGGeneratedComponents> IgnorePCGGeneratedComponentsRule = TInstancedStruct<FPCGRaycastFilterRule_IgnorePCGGeneratedComponents>::Make();
-	ToolSettings->RaycastRuleCollection.RaycastRules.Add(LandscapeRule);
-	ToolSettings->RaycastRuleCollection.RaycastRules.Add(MeshesRule);
-	ToolSettings->RaycastRuleCollection.RaycastRules.Add(IgnorePCGGeneratedComponentsRule);
 	
 	// @todo_pcg: Make adjustments to this and check when properties change.
 	CachedRayParams =
diff --git a/Engine/Plugins/PCG/Source/PCGEditor/Private/EditorMode/Tools/PCGInteractiveToolBuilder.cpp b/Engine/Plugins/PCG/Source/PCGEditor/Private/EditorMode/Tools/PCGInteractiveToolBuilder.cpp
--- a/Engine/Plugins/PCG/Source/PCGEditor/Private/EditorMode/Tools/PCGInteractiveToolBuilder.cpp
+++ b/Engine/Plugins/PCG/Source/PCGEditor/Private/EditorMode/Tools/PCGInteractiveToolBuilder.cpp
@@ -53,6 +53,7 @@ bool UE::PCG::EditorMode::Tool::BuildTool(UInteractiveTool* InTool)
 	}
 	
 	SettingsProperty->SetObjectPropertyValue(const_cast<void*>(PropertyAddress), Settings);
+	Settings->RestoreProperties(InTool);
 	return true;
 }
 
diff --git a/Engine/Plugins/PCG/Source/PCGEditor/Private/EditorMode/Tools/PCGInteractiveToolSettings.cpp b/Engine/Plugins/PCG/Source/PCGEditor/Private/EditorMode/Tools/PCGInteractiveToolSettings.cpp
--- a/Engine/Plugins/PCG/Source/PCGEditor/Private/EditorMode/Tools/PCGInteractiveToolSettings.cpp
+++ b/Engine/Plugins/PCG/Source/PCGEditor/Private/EditorMode/Tools/PCGInteractiveToolSettings.cpp
@@ -30,6 +30,8 @@ void UE::PCG::EditorMode::Tool::Shutdown(UInteractiveTool* InTool, EToolShutdown
 	{
 		if(UPCGInteractiveToolSettings* PCGSettings = Cast<UPCGInteractiveToolSettings>(ToolPropertyObject))
 		{
+			PCGSettings->SaveProperties(InTool);
+			
 			if(ShutdownType == EToolShutdownType::Cancel)
 			{
 				PCGSettings->Cancel(InTool);
diff --git a/Engine/Plugins/PCG/Source/PCGEditor/Private/EditorMode/Tools/Paint/PCGPaintTool.cpp b/Engine/Plugins/PCG/Source/PCGEditor/Private/EditorMode/Tools/Paint/PCGPaintTool.cpp
--- a/Engine/Plugins/PCG/Source/PCGEditor/Private/EditorMode/Tools/Paint/PCGPaintTool.cpp
+++ b/Engine/Plugins/PCG/Source/PCGEditor/Private/EditorMode/Tools/Paint/PCGPaintTool.cpp
@@ -1,7 +1,9 @@
 // Copyright Epic Games, Inc. All Rights Reserved.
 
 #include "EditorMode/Tools/Paint/PCGPaintTool.h"
+#include "Data/PCGPointArrayData.h"
 #include "EditorMode/PCGEdModeStyle.h"
+#include "EditorMode/Tools/Gizmos/PCGGizmos.h"
 #include "Helpers/PCGHelpers.h"
 
 #include "InteractiveGizmoManager.h"
@@ -12,10 +14,10 @@
 #include "ToolTargets/PrimitiveComponentToolTarget.h"
 #include "BaseBehaviors/TwoAxisPropertyEditBehavior.h"
 #include "PreviewMesh.h"
-#include "BaseGizmos/VolumetricBrushStampIndicator.h"
 #include "Materials/MaterialInstanceDynamic.h"
 #include "MaterialEditingLibrary.h"
-#include "Data/PCGPointArrayData.h"
+#include "BaseBehaviors/MouseWheelBehavior.h"
+#include "Generators/MinimalBoxMeshGenerator.h"
 
 #include UE_INLINE_GENERATED_CPP_BY_NAME(PCGPaintTool)
 
@@ -30,6 +32,12 @@ namespace PCGPaintTool::Constants
 
 UPCGInteractiveToolSettings_PaintTool::UPCGInteractiveToolSettings_PaintTool()
 {
+	TInstancedStruct<FPCGRaycastFilterRule_Landscape> LandscapeRule = TInstancedStruct<FPCGRaycastFilterRule_Landscape>::Make();
+	TInstancedStruct<FPCGRaycastFilterRule_Meshes> MeshesRule = TInstancedStruct<FPCGRaycastFilterRule_Meshes>::Make();
+	TInstancedStruct<FPCGRaycastFilterRule_IgnorePCGGeneratedComponents> IgnorePCGGeneratedComponentsRule = TInstancedStruct<FPCGRaycastFilterRule_IgnorePCGGeneratedComponents>::Make();
+	RaycastRuleCollection.RaycastRules.Add(LandscapeRule);
+	RaycastRuleCollection.RaycastRules.Add(MeshesRule);
+	RaycastRuleCollection.RaycastRules.Add(IgnorePCGGeneratedComponentsRule);
 } 
 
 FName UPCGInteractiveToolSettings_PaintTool::GetToolTag() const
@@ -46,6 +54,14 @@ void UPCGPaintTool::Setup()
 {
 	PropertyClass = UPCGEdModeBrushBase::StaticClass();
 	
+	ToolSettings->WatchProperty(ToolSettings->BrushMode, [this](EPCGBrushMode InBrushMode)
+	{
+		if (UPCGVolumetricBrushStampIndicator* TypedBrushStampIndicator = Cast<UPCGVolumetricBrushStampIndicator>(BrushStampIndicator))
+		{
+			TypedBrushStampIndicator->SetBrushMode(InBrushMode);
+		}
+	});
+	
 	AddToolPropertySource(ToolSettings);
 
 	AActor* TargetActor = Cast<UPrimitiveComponentToolTarget>(GetTarget())->GetOwnerActor();
@@ -54,19 +70,16 @@ void UPCGPaintTool::Setup()
 	{
 		ToolSettings->GeneratedResources.SpawnedActor->SetActorLocation(TargetActor->GetActorLocation());
 	}
-
-	TInstancedStruct<FPCGRaycastFilterRule_Landscape> LandscapeRule = TInstancedStruct<FPCGRaycastFilterRule_Landscape>::Make();
-	TInstancedStruct<FPCGRaycastFilterRule_Meshes> MeshesRule = TInstancedStruct<FPCGRaycastFilterRule_Meshes>::Make();
-	TInstancedStruct<FPCGRaycastFilterRule_IgnorePCGGeneratedComponents> IgnorePCGGeneratedComponentsRule = TInstancedStruct<FPCGRaycastFilterRule_IgnorePCGGeneratedComponents>::Make();
-	ToolSettings->RaycastRuleCollection.RaycastRules.Add(LandscapeRule);
-	ToolSettings->RaycastRuleCollection.RaycastRules.Add(MeshesRule);
-	ToolSettings->RaycastRuleCollection.RaycastRules.Add(IgnorePCGGeneratedComponentsRule);
 	
 	CachedRayParams =
 	{
 		.World = GetWorld(),
 		.FilterRuleCollection = &ToolSettings->RaycastRuleCollection
 	};
+
+	MouseWheelInputBehavior = NewObject<UMouseWheelInputBehavior>();
+	MouseWheelInputBehavior->Initialize(this);
+	AddInputBehavior(MouseWheelInputBehavior);
 	
 	Super::Setup();	
 }
@@ -116,7 +129,7 @@ void UPCGPaintTool::OnTick(float DeltaTime)
 
 	if(bShiftPressedChanged)
 	{
-		if(UVolumetricBrushStampIndicator* VolumetricStampIndicator = Cast<UVolumetricBrushStampIndicator>(BrushStampIndicator))
+		if(UPCGVolumetricBrushStampIndicator* VolumetricStampIndicator = Cast<UPCGVolumetricBrushStampIndicator>(BrushStampIndicator))
 		{
 			if(UMaterialInstanceDynamic* Instance = Cast<UMaterialInstanceDynamic>(VolumetricStampIndicator->GetPreviewMesh()->GetActiveMaterial()))
 			{
@@ -163,17 +176,39 @@ void UPCGPaintTool::OnUpdateDrag(const FRay& Ray)
 	bPaintOnTick = true;
 }
 
+FInputRayHit UPCGPaintTool::ShouldRespondToMouseWheel(const FInputDeviceRay& CurrentPos)
+{
+	FInputRayHit Result;
+	Result.bHit = true;
+
+	return Result;
+}
+
+void UPCGPaintTool::OnMouseWheelScrollUp(const FInputDeviceRay& CurrentPos)
+{
+	if (UPCGVolumetricBrushStampIndicator* TypedBrushStampIndicator = Cast<UPCGVolumetricBrushStampIndicator>(BrushStampIndicator))
+	{
+		TypedBrushStampIndicator->AddYawRotation(ToolSettings->RotateYawStepSize);
+	}
+}
+
+void UPCGPaintTool::OnMouseWheelScrollDown(const FInputDeviceRay& CurrentPos)
+{
+	if (UPCGVolumetricBrushStampIndicator* TypedBrushStampIndicator = Cast<UPCGVolumetricBrushStampIndicator>(BrushStampIndicator))
+	{
+		TypedBrushStampIndicator->AddYawRotation(-ToolSettings->RotateYawStepSize);
+	}
+}
+
 void UPCGPaintTool::SetupBrushStampIndicator()
 {
 	if (!BrushStampIndicator)
 	{
-		// See USkinWeightsPaintTool::SetupBrushStampIndicator
-		FString VolumetricBrushIndicatorGizmoType = TEXT("VolumetricBrushIndicatorGizmoType");
-		// register and spawn brush indicator gizmo
-		GetToolManager()->GetPairedGizmoManager()->RegisterGizmoType(VolumetricBrushIndicatorGizmoType, NewObject<UVolumetricBrushStampIndicatorBuilder>());
-		BrushStampIndicator = GetToolManager()->GetPairedGizmoManager()->CreateGizmo<UVolumetricBrushStampIndicator>(VolumetricBrushIndicatorGizmoType, FString(), this);
-		if (UVolumetricBrushStampIndicator* VolumetricBrushIndicator = Cast<UVolumetricBrushStampIndicator>(BrushStampIndicator))
+		GetToolManager()->GetPairedGizmoManager()->RegisterGizmoType(UE::PCG::EditorMode::Gizmos::Constants::BrushIdentifier_VolumetricBox, NewObject<UVolumetricBoxBrushStampIndicatorBuilder>());
+		BrushStampIndicator = GetToolManager()->GetPairedGizmoManager()->CreateGizmo<UPCGVolumetricBrushStampIndicator>(UE::PCG::EditorMode::Gizmos::Constants::BrushIdentifier_VolumetricBox, FString(), this);
+		if (UPCGVolumetricBrushStampIndicator* VolumetricBrushIndicator = Cast<UPCGVolumetricBrushStampIndicator>(BrushStampIndicator))
 		{
+			VolumetricBrushIndicator->SetBrushMode(ToolSettings->BrushMode);
 			VolumetricBrushIndicator->MakeBrushIndicatorMesh(this, GetTargetWorld());
 		}
 	}
@@ -212,15 +247,19 @@ void UPCGPaintTool::UpdateBrushStampIndicator()
 void UPCGPaintTool::ShutdownBrushStampIndicator()
 {
 	if (BrushStampIndicator)
-	{
-		FString VolumetricBrushIndicatorGizmoType = TEXT("VolumetricBrushIndicatorGizmoType");
-		
+	{		
 		GetToolManager()->GetPairedGizmoManager()->DestroyGizmo(BrushStampIndicator);
 		BrushStampIndicator = nullptr;
-		GetToolManager()->GetPairedGizmoManager()->DeregisterGizmoType(VolumetricBrushIndicatorGizmoType);
+		GetToolManager()->GetPairedGizmoManager()->DeregisterGizmoType(UE::PCG::EditorMode::Gizmos::Constants::BrushIdentifier_VolumetricBox);
 	}
 }
 
+void UPCGPaintTool::RestartBrushStampIndicator()
+{
+	ShutdownBrushStampIndicator();
+	SetupBrushStampIndicator();
+}
+
 double UPCGPaintTool::EstimateMaximumTargetDimension()
 {
 	// We return the target actor's max dimension
@@ -233,10 +272,16 @@ void UPCGPaintTool::AddSinglePointFromBrush()
 	
 	const int32 PointIndex = PointData->GetNumPoints();
 	PointData->SetNumPoints(PointIndex + 1);
+
+	TOptional<FRotator> BrushRotator;
+	if (UPCGVolumetricBrushStampIndicator* TypedBrushStampIndicator = Cast<UPCGVolumetricBrushStampIndicator>(BrushStampIndicator))
+	{
+		BrushRotator = TypedBrushStampIndicator->GetPreviewMesh()->GetActor()->GetActorRotation();
+	}
 	
 	// Values to assign
 	const float PointSteepness = 1.0f;
-	const FTransform PointTransform = FTransform(LastBrushStamp.WorldPosition);
+	const FTransform PointTransform = FTransform(BrushRotator.Get(FRotator::ZeroRotator), LastBrushStamp.WorldPosition);
 	const FVector& Position = PointTransform.GetLocation();
 	const int32 PointSeed = PCGHelpers::ComputeSeed((int)Position.X, (int)Position.Y, (int)Position.Z);
 	const FBox LocalBounds = FBox(FVector(-CurrentBrushRadius), FVector(CurrentBrushRadius));
diff --git a/Engine/Plugins/PCG/Source/PCGEditor/Public/EditorMode/Tools/Gizmos/PCGGizmos.h b/Engine/Plugins/PCG/Source/PCGEditor/Public/EditorMode/Tools/Gizmos/PCGGizmos.h
new file mode 100644
--- /dev/null
+++ b/Engine/Plugins/PCG/Source/PCGEditor/Public/EditorMode/Tools/Gizmos/PCGGizmos.h
@@ -0,0 +1,70 @@
+﻿// Copyright Epic Games, Inc. All Rights Reserved.
+
+#pragma once
+
+#include "BaseGizmos/BrushStampIndicator.h"
+#include "InteractiveTool.h"
+#include "InteractiveGizmoBuilder.h"
+#include "PCGGizmos.generated.h"
+
+class UPreviewMesh;
+class UMaterialInstanceDynamic;
+
+namespace UE::PCG::EditorMode::Gizmos
+{
+	namespace Constants
+	{
+		const FString BrushIdentifier_VolumetricBox = "VolumetricBoxBrushGizmoType";
+	}
+	
+	UMaterialInstanceDynamic* GetDefaultVolumeBrushMaterial(UInteractiveToolManager* ToolManager);
+}
+
+UENUM()
+enum class EPCGBrushMode
+{
+	Sphere,
+	Box
+};
+
+UCLASS(MinimalAPI)
+class UVolumetricBoxBrushStampIndicatorBuilder : public UInteractiveGizmoBuilder
+{
+	GENERATED_BODY()
+public:
+	virtual UInteractiveGizmo* BuildGizmo(const FToolBuilderState& SceneState) const override;
+};
+
+UCLASS(Transient, MinimalAPI)
+class UPCGVolumetricBrushStampIndicator : public UBrushStampIndicator
+{
+	GENERATED_BODY()
+
+public:
+	virtual void Setup() override;
+	virtual void Render(IToolsContextRenderAPI* RenderAPI) override;
+	virtual void Shutdown() override;
+
+	void SetVolumetricBrushVisible(bool bInVisible);
+
+	bool IsVolumetricBrushVisible() const;
+
+	void MakeBrushIndicatorMesh(UInteractiveTool* ParentTool, UWorld* World);
+
+	UPreviewMesh* GetPreviewMesh() const;
+
+	void AddYawRotation(float AddYawRotation);
+
+	void SetBrushMode(EPCGBrushMode InBrushMode);
+protected:
+	void GenerateMesh(EPCGBrushMode InBrushMode);
+private:
+	UPROPERTY()
+	TObjectPtr<UPreviewMesh> PreviewMesh;
+
+	UPROPERTY()
+	bool bDrawExtents = true;
+
+	UPROPERTY()
+	EPCGBrushMode BrushMode = EPCGBrushMode::Sphere;
+};
\ No newline at end of file
diff --git a/Engine/Plugins/PCG/Source/PCGEditor/Public/EditorMode/Tools/Paint/PCGPaintTool.h b/Engine/Plugins/PCG/Source/PCGEditor/Public/EditorMode/Tools/Paint/PCGPaintTool.h
--- a/Engine/Plugins/PCG/Source/PCGEditor/Public/EditorMode/Tools/Paint/PCGPaintTool.h
+++ b/Engine/Plugins/PCG/Source/PCGEditor/Public/EditorMode/Tools/Paint/PCGPaintTool.h
@@ -6,9 +6,12 @@
 #include "BaseBehaviors/KeyAsModifierInputBehavior.h"
 #include "EditorMode/Tools/PCGInteractiveToolBuilder.h"
 #include "EditorMode/Tools/PCGInteractiveToolSettings.h"
+#include "EditorMode/Tools/Gizmos/PCGGizmos.h"
 #include "ToolTargets/PrimitiveComponentToolTarget.h"
 #include "PCGPaintTool.generated.h"
 
+class UMouseWheelInputBehavior;
+
 UCLASS()
 class PCGEDITOR_API UPCGInteractiveToolSettings_PaintTool : public UPCGInteractiveToolSettings
 {
@@ -22,16 +25,24 @@ public:
 	virtual const UScriptStruct* GetWorkingDataType() const override;
 	// ~End UPCGInteractiveToolSettings interface
 
+	/** The brush shape to use. */
+	UPROPERTY(EditAnywhere, BlueprintReadOnly, Category = "Paint")
+	EPCGBrushMode BrushMode = EPCGBrushMode::Sphere;
+	
 	/** If false, will erase points on all paint data instances. */
-	UPROPERTY(BlueprintReadOnly, EditAnywhere, Category = "Paint")
+	UPROPERTY(EditAnywhere, BlueprintReadOnly, Category = "Paint")
 	bool bErasePointsOfSelectedDataInstance = true;
 
+	/** The step size for rotation of the brush. Moving the mouse wheel up or down will add or subtract from the brush rotation. */
+	UPROPERTY(EditAnywhere, BlueprintReadOnly, Category = "Paint")
+	float RotateYawStepSize = 10.f;
+
 	UPROPERTY(EditAnywhere, Category= "Raycast")
 	FPCGRaycastFilterRuleCollection RaycastRuleCollection;
 };
 
 UCLASS(Transient)
-class UPCGPaintTool : public UBaseBrushTool
+class UPCGPaintTool : public UBaseBrushTool, public IMouseWheelBehaviorTarget
 {
 	GENERATED_BODY()
 
@@ -50,10 +61,17 @@ public:
 protected:
 	virtual void OnClickPress(const FInputDeviceRay& PressPos) override;
 	virtual void OnUpdateDrag(const FRay& Ray) override;
+
+	/** IMouseWheelBehaviorTarget */
+	virtual FInputRayHit ShouldRespondToMouseWheel(const FInputDeviceRay& CurrentPos) override;
+	virtual void OnMouseWheelScrollUp(const FInputDeviceRay& CurrentPos) override;
+	virtual void OnMouseWheelScrollDown(const FInputDeviceRay& CurrentPos) override;
 	
 	virtual void SetupBrushStampIndicator() override;
 	virtual void UpdateBrushStampIndicator() override;
 	virtual void ShutdownBrushStampIndicator() override;
+
+	void RestartBrushStampIndicator();
 	
 	/** This will be used as the brush base, if not specifying a radius directly */
 	virtual double EstimateMaximumTargetDimension() override;
@@ -66,6 +84,9 @@ private:
 	UPROPERTY(Transient)
 	TObjectPtr<UPCGInteractiveToolSettings_PaintTool> ToolSettings;
 
+	UPROPERTY(Transient)
+	TObjectPtr<UMouseWheelInputBehavior> MouseWheelInputBehavior;
+	
 	UPROPERTY(Transient)
 	bool bPaintOnTick = false;
 
