[PCG] QoL improvements for InstancedSkinnedMesh spawner

* Remove Bank/SequenceNumber as they are not used
* Clearer error message

#rnx

diff --git a/Engine/Plugins/PCG/Source/PCG/Private/MeshSelectors/PCGSkinnedMeshSelector.cpp b/Engine/Plugins/PCG/Source/PCG/Private/MeshSelectors/PCGSkinnedMeshSelector.cpp
--- a/Engine/Plugins/PCG/Source/PCG/Private/MeshSelectors/PCGSkinnedMeshSelector.cpp
+++ b/Engine/Plugins/PCG/Source/PCG/Private/MeshSelectors/PCGSkinnedMeshSelector.cpp
@@ -73,61 +73,29 @@ bool UPCGSkinnedMeshSelector::SelectInstances(
 
 	if (!InPointData)
 	{
-		PCGE_LOG_C(Error, GraphAndLog, &Context, LOCTEXT("InputMissingData", "Missing input data"));
+		PCGLog::InputOutput::LogInvalidInputDataError(&Context);
 		return true;
 	}
 
 	if (!InPointData->Metadata)
 	{
-		PCGE_LOG_C(Error, GraphAndLog, &Context, LOCTEXT("InputMissingMetadata", "Unable to get metadata from input"));
-		return true;
-	}
-
-	if (!InPointData->Metadata->HasAttribute(BankAttribute.GetAttributeName()))
-	{
-		PCGE_LOG_C(Error, GraphAndLog, &Context, FText::Format(LOCTEXT("AttributeNotInMetadata", "Attribute '{0}' is not in the metadata"), BankAttribute.GetDisplayText()));
-		return true;
-	}
-
-	if (!InPointData->Metadata->HasAttribute(SequenceIndexAttribute.GetAttributeName()))
-	{
-		PCGE_LOG_C(Error, GraphAndLog, &Context, FText::Format(LOCTEXT("AttributeNotInMetadata", "Attribute '{0}' is not in the metadata"), SequenceIndexAttribute.GetDisplayText()));
+		PCGLog::Metadata::LogInvalidMetadata(&Context);
 		return true;
 	}
 
 	if (!InPointData->Metadata->HasAttribute(MeshAttribute.GetAttributeName()))
 	{
-		PCGE_LOG_C(Error, GraphAndLog, &Context, FText::Format(LOCTEXT("AttributeNotInMetadata", "Attribute '{0}' is not in the metadata"), MeshAttribute.GetDisplayText()));
+		PCGE_LOG_C(Error, GraphAndLog, &Context, FText::Format(LOCTEXT("MeshAttributeNotInMetadata", "Mesh Attribute '{0}' is not in the metadata."), MeshAttribute.GetDisplayText()));
 		return true;
 	}
 
-	const FPCGMetadataAttributeBase* BankAttributeBase = InPointData->Metadata->GetConstAttribute(BankAttribute.GetAttributeName());
-	check(BankAttributeBase);
-
-	const FPCGMetadataAttributeBase* SequenceIndexAttributeBase = InPointData->Metadata->GetConstAttribute(SequenceIndexAttribute.GetAttributeName());
-	check(SequenceIndexAttributeBase);
-
 	const FPCGMetadataAttributeBase* MeshAttributeBase = InPointData->Metadata->GetConstAttribute(MeshAttribute.GetAttributeName());
 	check(MeshAttributeBase);
 
-	// Validate that the "bank" attribute is of the right type
-	if (!PCG::Private::IsOfTypes<FSoftObjectPath, FString>(BankAttributeBase->GetTypeId()))
-	{
-		PCGE_LOG_C(Error, GraphAndLog, &Context, FText::Format(LOCTEXT("BankAttributeInvalidType", "Attribute '{0}' is not of valid type (must be FString or FSoftObjectPath)"), BankAttribute.GetDisplayText()));
-		return true;
-	}
-
-	// Validate that the "sequence index" attribute is of the right type
-	if (!PCG::Private::IsOfTypes<int32>(SequenceIndexAttributeBase->GetTypeId()))
-	{
-		PCGE_LOG_C(Error, GraphAndLog, &Context, FText::Format(LOCTEXT("SeqIndexAttributeInvalidType", "Attribute '{0}' is not of valid type (must be int32)"), SequenceIndexAttribute.GetDisplayText()));
-		return true;
-	}
-
 	// Validate that the "mesh" attribute is of the right type
 	if (!PCG::Private::IsOfTypes<FSoftObjectPath, FString>(MeshAttributeBase->GetTypeId()))
 	{
-		PCGE_LOG_C(Error, GraphAndLog, &Context, FText::Format(LOCTEXT("MeshAttributeInvalidType", "Attribute '{0}' is not of valid type (must be FString or FSoftObjectPath)"), MeshAttribute.GetDisplayText()));
+		PCGE_LOG_C(Error, GraphAndLog, &Context, FText::Format(LOCTEXT("MeshAttributeInvalidType", "Mesh Attribute '{0}' is not of valid type (must be FString or FSoftObjectPath)."), MeshAttribute.GetDisplayText()));
 		return true;
 	}
 
@@ -227,8 +195,9 @@ bool UPCGSkinnedMeshSelector::SelectInstances(
 
 		if (Context.PointProviderValues.IsEmpty() || Context.PointAnimationValues.IsEmpty())
 		{
-			PCGAttributeAccessorHelpers::ExtractAllValues(InPointData, BankAttribute, Context.PointProviderValues, &Context);
-			PCGAttributeAccessorHelpers::ExtractAllValues(InPointData, SequenceIndexAttribute, Context.PointAnimationValues, &Context);
+			// GW-TODO: 
+			// PCGAttributeAccessorHelpers::ExtractAllValues(InPointData, BankAttribute, Context.PointProviderValues, &Context);
+			// PCGAttributeAccessorHelpers::ExtractAllValues(InPointData, SequenceIndexAttribute, Context.PointAnimationValues, &Context);
 		}
 
 		// TODO: Revisit this when attribute partitioning is returned in a more optimized form
diff --git a/Engine/Plugins/PCG/Source/PCG/Public/MeshSelectors/PCGSkinnedMeshSelector.h b/Engine/Plugins/PCG/Source/PCG/Public/MeshSelectors/PCGSkinnedMeshSelector.h
--- a/Engine/Plugins/PCG/Source/PCG/Public/MeshSelectors/PCGSkinnedMeshSelector.h
+++ b/Engine/Plugins/PCG/Source/PCG/Public/MeshSelectors/PCGSkinnedMeshSelector.h
@@ -77,12 +77,7 @@ public:
 		UPCGPointData* OutPointData) const;
 
 public:
-	UPROPERTY(BlueprintReadWrite, EditAnywhere, Category = MeshSelector)
-	FPCGAttributePropertyInputSelector BankAttribute;
-
-	UPROPERTY(BlueprintReadWrite, EditAnywhere, Category = MeshSelector)
-	FPCGAttributePropertyInputSelector SequenceIndexAttribute;
-
+	/** Attribute in the input data that contains the path to the skinned mesh to spawn. */
 	UPROPERTY(BlueprintReadWrite, EditAnywhere, Category = MeshSelector)
 	FPCGAttributePropertyInputSelector MeshAttribute;
 
