[PCG] Be more resilient with data domain when we manually copy attributes over

* Metadata->AddAttributes can lose the data domain values (difference between default value and first value).
* Add a helper that copy the data domain but only add the attributes on the element domain

#rnx

diff --git a/Engine/Plugins/PCG/Source/PCG/Private/Elements/PCGAttributeFilter.cpp b/Engine/Plugins/PCG/Source/PCG/Private/Elements/PCGAttributeFilter.cpp
--- a/Engine/Plugins/PCG/Source/PCG/Private/Elements/PCGAttributeFilter.cpp
+++ b/Engine/Plugins/PCG/Source/PCG/Private/Elements/PCGAttributeFilter.cpp
@@ -487,7 +487,7 @@ bool FPCGAttributeFilterElementBase::DoFiltering(FPCGContext* Context, EPCGAttri
 		UPCGBasePointData* OutFilterPointData = nullptr;
 		TBitArray<TInlineAllocator<2048>> FilterBitArray;
 
-		const UPCGMetadata* OriginalMetadata = nullptr;
+		const UPCGData* OriginalData = nullptr;
 		UPCGMetadata* InFilterMetadata = nullptr;
 		UPCGMetadata* OutFilterMetadata = nullptr;
 
@@ -631,6 +631,8 @@ bool FPCGAttributeFilterElementBase::DoFiltering(FPCGContext* Context, EPCGAttri
 			continue;
 		}
 
+		OperationData.OriginalData = OriginalData;
+
 		if (OperationData.bIsInputPointData)
 		{
 			const UPCGBasePointData* OriginalPointData = CastChecked<UPCGBasePointData>(OriginalData);
@@ -645,8 +647,6 @@ bool FPCGAttributeFilterElementBase::DoFiltering(FPCGContext* Context, EPCGAttri
 			// Param data
 			const UPCGParamData* OriginalParamData = CastChecked<UPCGParamData>(OriginalData);
 
-			OperationData.OriginalMetadata = OriginalParamData->Metadata;
-
 			// Will be set individually in batches
 			OperationData.FilterBitArray.SetNumUninitialized(OriginalParamData->Metadata ? OriginalParamData->Metadata->GetItemCountForChild() : 0);
 		}
@@ -819,17 +819,18 @@ bool FPCGAttributeFilterElementBase::DoFiltering(FPCGContext* Context, EPCGAttri
 			}
 			else
 			{
-				check(OperationData.OriginalMetadata);
+				check(OperationData.OriginalData);
 
 				if (!OperationData.FilterBitArray.IsEmpty())
 				{
-					check(OperationData.FilterBitArray.Num() == OperationData.OriginalMetadata->GetItemCountForChild())
+					const UPCGMetadata* OriginalMetadata = OperationData.OriginalData->ConstMetadata();
+					check(OperationData.FilterBitArray.Num() == OriginalMetadata->GetItemCountForChild())
 
 					TArray<PCGMetadataEntryKey, TInlineAllocator<256>> InEntryKeys;
 					TArray<PCGMetadataEntryKey, TInlineAllocator<256>> OutEntryKeys;
 
 					const int32 NumInFilterMetadata = OperationData.FilterBitArray.CountSetBits();
-					const int32 NumOutFilterMetadata = OperationData.OriginalMetadata->GetItemCountForChild() - NumInFilterMetadata;
+					const int32 NumOutFilterMetadata = OriginalMetadata->GetItemCountForChild() - NumInFilterMetadata;
 
 					InEntryKeys.Reserve(NumInFilterMetadata);
 					OutEntryKeys.Reserve(NumOutFilterMetadata);
@@ -851,8 +852,9 @@ bool FPCGAttributeFilterElementBase::DoFiltering(FPCGContext* Context, EPCGAttri
 					{
 						UPCGParamData* InFilterParamData = FPCGContext::NewObject_AnyThread<UPCGParamData>(Context);
 						UPCGMetadata* InFilterMetadata = InFilterParamData->Metadata;
-						InFilterMetadata->AddAttributesFiltered(OperationData.OriginalMetadata, TSet<FName>(), EPCGMetadataFilterMode::ExcludeAttributes);
-						InFilterMetadata->SetAttributes(InEntryKeys, OperationData.OriginalMetadata, /*InOutOptionalKeys=*/nullptr, Context);
+						PCGMetadataHelpers::InitializeMetadataWithDataDomainCopyAndElementsNoCopy(OperationData.OriginalData, InFilterParamData);
+						InFilterMetadata->AddAttributesFiltered(OriginalMetadata, TSet<FName>(), EPCGMetadataFilterMode::ExcludeAttributes);
+						InFilterMetadata->SetAttributes(InEntryKeys, OriginalMetadata, /*InOutOptionalKeys=*/nullptr, Context);
 						InFilterData = InFilterParamData;
 					}
 
@@ -860,8 +862,9 @@ bool FPCGAttributeFilterElementBase::DoFiltering(FPCGContext* Context, EPCGAttri
 					{
 						UPCGParamData* OutFilterParamData = FPCGContext::NewObject_AnyThread<UPCGParamData>(Context);
 						UPCGMetadata* OutFilterMetadata = OutFilterParamData->Metadata;
-						OutFilterMetadata->AddAttributesFiltered(OperationData.OriginalMetadata, TSet<FName>(), EPCGMetadataFilterMode::ExcludeAttributes);
-						OutFilterMetadata->SetAttributes(OutEntryKeys, OperationData.OriginalMetadata, /*InOutOptionalKeys=*/nullptr, Context);
+						PCGMetadataHelpers::InitializeMetadataWithDataDomainCopyAndElementsNoCopy(OperationData.OriginalData, OutFilterParamData);
+						OutFilterMetadata->AddAttributesFiltered(OriginalMetadata, TSet<FName>(), EPCGMetadataFilterMode::ExcludeAttributes);
+						OutFilterMetadata->SetAttributes(OutEntryKeys, OriginalMetadata, /*InOutOptionalKeys=*/nullptr, Context);
 						OutFilterData = OutFilterParamData;
 					}
 				}
diff --git a/Engine/Plugins/PCG/Source/PCG/Private/Elements/PCGAttributeGetFromIndexElement.cpp b/Engine/Plugins/PCG/Source/PCG/Private/Elements/PCGAttributeGetFromIndexElement.cpp
--- a/Engine/Plugins/PCG/Source/PCG/Private/Elements/PCGAttributeGetFromIndexElement.cpp
+++ b/Engine/Plugins/PCG/Source/PCG/Private/Elements/PCGAttributeGetFromIndexElement.cpp
@@ -83,7 +83,7 @@ bool FPCGAttributeGetFromIndexElement::ExecuteInternal(FPCGContext* Context) con
 
 		UPCGParamData* SubParam = FPCGContext::NewObject_AnyThread<UPCGParamData>(Context);
 		check(SubParam->Metadata);
-		SubParam->Metadata->AddAttributes(ParamMetadata);
+		PCGMetadataHelpers::InitializeMetadataWithDataDomainCopyAndElementsNoCopy(Param, SubParam);
 
 		const PCGMetadataEntryKey OriginalEntryKey = Index;
 		PCGMetadataEntryKey SingleEntryKey = SubParam->Metadata->AddEntry();
diff --git a/Engine/Plugins/PCG/Source/PCG/Private/Elements/PCGMatchAndSetAttributes.cpp b/Engine/Plugins/PCG/Source/PCG/Private/Elements/PCGMatchAndSetAttributes.cpp
--- a/Engine/Plugins/PCG/Source/PCG/Private/Elements/PCGMatchAndSetAttributes.cpp
+++ b/Engine/Plugins/PCG/Source/PCG/Private/Elements/PCGMatchAndSetAttributes.cpp
@@ -1120,7 +1120,7 @@ bool FPCGMatchAndSetAttributesElement::PrepareDataInternal(FPCGContext* InContex
 		else if (const UPCGParamData* InParamData = Cast<const UPCGParamData>(OutState.InData))
 		{
 			UPCGParamData* OutParamData = FPCGContext::NewObject_AnyThread<UPCGParamData>(InContext);
-			OutParamData->Metadata->AddAttributes(InParamData->Metadata);
+			PCGMetadataHelpers::InitializeMetadataWithDataDomainCopyAndElementsNoCopy(InParamData, OutParamData);
 
 			OutState.OutData = OutParamData;
 		}
diff --git a/Engine/Plugins/PCG/Source/PCG/Private/Helpers/PCGMetadataHelpers.cpp b/Engine/Plugins/PCG/Source/PCG/Private/Helpers/PCGMetadataHelpers.cpp
--- a/Engine/Plugins/PCG/Source/PCG/Private/Helpers/PCGMetadataHelpers.cpp
+++ b/Engine/Plugins/PCG/Source/PCG/Private/Helpers/PCGMetadataHelpers.cpp
@@ -47,6 +47,25 @@ namespace PCGMetadataHelpers
 		return Metadata->GetParent();
 	}
 
+	void InitializeMetadataWithDataDomainCopyAndElementsNoCopy(const UPCGData* InData, UPCGData* OutData)
+	{
+		check(InData && OutData && InData->ConstMetadata() && OutData->MutableMetadata())
+
+		// Initialize the data domain by copy
+		const FPCGMetadataDomain* InDataDomain = InData->ConstMetadata()->GetConstMetadataDomain(PCGMetadataDomainID::Data);
+		FPCGMetadataDomain* OutDataDomain = OutData->MutableMetadata()->GetMetadataDomain(PCGMetadataDomainID::Data);
+
+		if (InDataDomain && OutDataDomain)
+		{
+			OutDataDomain->InitializeAsCopy(InDataDomain);
+		}
+
+		// And just add the attributes for the element domain since they will be copied manually.
+		const FPCGMetadataDomain* InElementDomain = InData->ConstMetadata()->GetConstMetadataDomain(PCGMetadataDomainID::Elements);
+		FPCGMetadataDomain* OutElementDomain = OutData->MutableMetadata()->GetMetadataDomain(PCGMetadataDomainID::Elements);
+		OutElementDomain->AddAttributes(InElementDomain);
+	}
+
 	const UPCGMetadata* GetConstMetadata(const UPCGData* InData)
 	{
 		return InData ? InData->ConstMetadata() : nullptr;
diff --git a/Engine/Plugins/PCG/Source/PCG/Private/Metadata/PCGMetadataDomain.cpp b/Engine/Plugins/PCG/Source/PCG/Private/Metadata/PCGMetadataDomain.cpp
--- a/Engine/Plugins/PCG/Source/PCG/Private/Metadata/PCGMetadataDomain.cpp
+++ b/Engine/Plugins/PCG/Source/PCG/Private/Metadata/PCGMetadataDomain.cpp
@@ -408,6 +408,12 @@ void FPCGMetadataDomain::InitializeAsCopy(const FPCGMetadataDomainInitializePara
 	}
 }
 
+bool FPCGMetadataDomain::AddAttributes(const FPCGMetadataDomain* InOther)
+{
+	FPCGMetadataDomainInitializeParams Params{InOther};
+	return AddAttributes(Params);
+}
+
 bool FPCGMetadataDomain::AddAttributes(const FPCGMetadataDomainInitializeParams& InParams)
 {
 	if (!InParams.Parent)
diff --git a/Engine/Plugins/PCG/Source/PCG/Private/Metadata/PCGMetadataPartitionCommon.cpp b/Engine/Plugins/PCG/Source/PCG/Private/Metadata/PCGMetadataPartitionCommon.cpp
--- a/Engine/Plugins/PCG/Source/PCG/Private/Metadata/PCGMetadataPartitionCommon.cpp
+++ b/Engine/Plugins/PCG/Source/PCG/Private/Metadata/PCGMetadataPartitionCommon.cpp
@@ -531,7 +531,7 @@ namespace PCGMetadataPartitionCommon
 			else
 			{
 				UPCGParamData* NewData = FPCGContext::NewObject_AnyThread<UPCGParamData>(InOptionalContext);
-				NewData->Metadata->AddAttributes(OriginalMetadata);
+				PCGMetadataHelpers::InitializeMetadataWithDataDomainCopyAndElementsNoCopy(InData, NewData);
 				NewMetadata = NewData->Metadata;
 				PartitionedData.Add(NewData);
 			}
@@ -603,7 +603,7 @@ namespace PCGMetadataPartitionCommon
 				else
 				{
 					UPCGParamData* NewData = FPCGContext::NewObject_AnyThread<UPCGParamData>(InOptionalContext);
-					NewData->Metadata->AddAttributes(OriginalMetadata);
+					PCGMetadataHelpers::InitializeMetadataWithDataDomainCopyAndElementsNoCopy(InData, NewData);
 					OutputData = NewData;
 					OutMetadata = NewData->Metadata;
 				}
diff --git a/Engine/Plugins/PCG/Source/PCG/Public/Helpers/PCGMetadataHelpers.h b/Engine/Plugins/PCG/Source/PCG/Public/Helpers/PCGMetadataHelpers.h
--- a/Engine/Plugins/PCG/Source/PCG/Public/Helpers/PCGMetadataHelpers.h
+++ b/Engine/Plugins/PCG/Source/PCG/Public/Helpers/PCGMetadataHelpers.h
@@ -23,6 +23,11 @@ namespace PCGMetadataHelpers
 	PCG_API const UPCGMetadata* GetParentMetadata(const UPCGMetadata* Metadata);
 	PCG_API const FPCGMetadataDomain* GetParentMetadata(const FPCGMetadataDomain* Metadata);
 
+	/** Useful for data that doesn't support metadata parenting (like ParamData), where we need to copy the data domain,
+	 * but only want to add the attributes for the elements domain, because we will copy the values manually.
+	 */
+	PCG_API void InitializeMetadataWithDataDomainCopyAndElementsNoCopy(const UPCGData* InData, UPCGData* OutData);
+
 	// Utility function to break circular dependency for UPCGMetadata deprecation in PCGMetadataAttributeTpl.
 	PCG_API FPCGMetadataDomain* GetDefaultMetadataDomain(UPCGMetadata* InMetadata);
 
