[PCG] Ignore supported but not allocated domains in Delete Attributes

* When targetting a specific domain that is supported, if it is not allocated there is nothing to do.
* If the domain is not supported it will still error out

#rnx

diff --git a/Engine/Plugins/PCG/Source/PCG/Private/Elements/PCGDeleteAttributesElement.cpp b/Engine/Plugins/PCG/Source/PCG/Private/Elements/PCGDeleteAttributesElement.cpp
--- a/Engine/Plugins/PCG/Source/PCG/Private/Elements/PCGDeleteAttributesElement.cpp
+++ b/Engine/Plugins/PCG/Source/PCG/Private/Elements/PCGDeleteAttributesElement.cpp
@@ -171,11 +171,17 @@ bool FPCGDeleteAttributesElement::ExecuteInternal(FPCGContext* Context) const
 
 		const FPCGAttributePropertySelector DomainSelector = FPCGAttributePropertySelector::CreateAttributeSelector(NAME_None, Settings->MetadataDomain);
 		const FPCGMetadataDomainID SelectedDomainID = InputData->GetMetadataDomainIDFromSelector(DomainSelector);
+		if (!SelectedDomainID.IsValid())
+		{
+			PCGLog::Metadata::LogInvalidMetadataDomain(DomainSelector, Context);
+			continue;
+		}
+
 		const FPCGMetadataDomain* InputMetadataDomain = InputMetadata->GetConstMetadataDomain(SelectedDomainID);
 
 		if (!InputMetadataDomain)
 		{
-			PCGLog::Metadata::LogInvalidMetadataDomain(DomainSelector, Context);
+			// Domain is supported but not allocated, nothing to do.
 			continue;
 		}
 
