[PCG] Fix deprecation when both ends of edge changes types, that should still be compatible

* UpdatePins was updating the pin types sequentially, it means that one update pin could see the other pin on the edge not updated yet, and potentially incompatible, even if the updated type would have been compatible
* Add the pin type update in the ApplyDeprecationBeforeUpdatePin
* Keep track if the pin type change to break the edges in UpdatePins if needed

diff --git a/Engine/Plugins/PCG/Source/PCG/Private/PCGNode.cpp b/Engine/Plugins/PCG/Source/PCG/Private/PCGNode.cpp
--- a/Engine/Plugins/PCG/Source/PCG/Private/PCGNode.cpp
+++ b/Engine/Plugins/PCG/Source/PCG/Private/PCGNode.cpp
@@ -87,8 +87,36 @@ void UPCGNode::ApplyDeprecationBeforeUpdatePins()
 	{
 		const int VersionBefore = Settings->DataVersion;
 
+		// Give the chance to the settings to rename the pins
 		Settings->ApplyDeprecationBeforeUpdatePins(this, InputPins, OutputPins);
 
+		// Then make sure the pin properties are up to date on their types. We won't break edges here, it will be done in the update pins.
+		// But all the nodes need to have their pins types up-to-date to know if an edge should be broken or not. Otherwise, the first node gets
+		// its pin updated, check the other side of the edge, finds that the (old) type is not matching and break the edge, even if the new type would have been
+		// compatible.
+		auto UpdatePinTypes = [](const TArray<FPCGPinProperties>& PinProperties, UPCGPin* Pin)
+		{
+			if (const FPCGPinProperties* MatchingProperties = PinProperties.FindByPredicate([Pin](const FPCGPinProperties& Prop) { return Prop.Label == Pin->Properties.Label; }))
+			{
+				if (!(Pin->Properties == *MatchingProperties))
+				{
+					Pin->Modify();
+					Pin->Properties = *MatchingProperties;
+					Pin->bPropertiesHaveChangedOnDeprecation = true;
+				}
+			}
+		};
+
+		for (UPCGPin* InputPin : InputPins)
+		{
+			UpdatePinTypes(Settings->AllInputPinProperties(), InputPin);
+		}
+
+		for (UPCGPin* OutputPin : OutputPins)
+		{
+			UpdatePinTypes(Settings->AllOutputPinProperties(), OutputPin);
+		}
+
 		// Version number should not be bumped in this before-update-pins migration, if it does
 		// we risk deprecation code not running in post-update-pins because version number is latest.
 		ensure(VersionBefore == Settings->DataVersion);
@@ -980,6 +1008,13 @@ EPCGChangeType UPCGNode::UpdatePins(TFunctionRef<UPCGPin*(UPCGNode*)> PinAllocat
 					bChangedTooltips = true;
 				}
 #endif // WITH_EDITOR
+				
+				// Make sure that if the pin type changed in the ApplyDeprecationBeforeUpdatePins to break incompatible edges.
+				if (Pin->bPropertiesHaveChangedOnDeprecation)
+				{
+					bAppliedEdgeChanges |= Pin->BreakAllIncompatibleEdges(&TouchedNodes);
+					Pin->bPropertiesHaveChangedOnDeprecation = false;
+				}
 			}
 			else
 			{
diff --git a/Engine/Plugins/PCG/Source/PCG/Public/PCGPin.h b/Engine/Plugins/PCG/Source/PCG/Public/PCGPin.h
--- a/Engine/Plugins/PCG/Source/PCG/Public/PCGPin.h
+++ b/Engine/Plugins/PCG/Source/PCG/Public/PCGPin.h
@@ -316,6 +316,12 @@ public:
 	PRAGMA_DISABLE_DEPRECATION_WARNINGS
 	EPCGTypeConversion GetRequiredTypeConversion(const UPCGPin* InOtherPin) const { return EPCGTypeConversion::Failed; };
 	PRAGMA_ENABLE_DEPRECATION_WARNINGS
+
+private:
+	friend UPCGNode;
+
+	// Keep track if a pin had their properties changed while loading, so UpdatePins will try to break incompatible edges.
+	bool bPropertiesHaveChangedOnDeprecation = false;
 };
 
 #if WITH_EDITOR
