[PCG] Put more restrictions on data type identifiers to avoid overlapping types.

* PCG Data type identifier should behave like a bitflag, Surface | Landscape should yield Surface
* By not enforcing that, the `&` (intersection) was not giving the right result
* Refactored all places where we add a new Id to the list to go through the same code path
* Add a type check to avoid adding the same type or child type
* Also remove child types if we add parent types
* Add units tests for that

#rnx

diff --git a/Engine/Plugins/PCG/Source/PCG/Private/Data/Registry/PCGDataTypeIdentifier.cpp b/Engine/Plugins/PCG/Source/PCG/Private/Data/Registry/PCGDataTypeIdentifier.cpp
--- a/Engine/Plugins/PCG/Source/PCG/Private/Data/Registry/PCGDataTypeIdentifier.cpp
+++ b/Engine/Plugins/PCG/Source/PCG/Private/Data/Registry/PCGDataTypeIdentifier.cpp
@@ -51,7 +51,7 @@ void FPCGDataTypeIdentifier::ComposeFromLegacyType(EPCGDataType LegacyType)
 	FPCGDataTypeBaseId BaseId = FPCGDataTypeBaseId::MakeFromLegacyType(LegacyType);
 	if (BaseId.IsValid())
 	{
-		Ids.Emplace(BaseId);
+		ComposeFromBaseId(BaseId);
 		return;
 	}
 
@@ -63,7 +63,7 @@ void FPCGDataTypeIdentifier::ComposeFromLegacyType(EPCGDataType LegacyType)
 			BaseId = FPCGDataTypeBaseId::MakeFromLegacyType(Type);
 			if (BaseId.IsValid())
 			{
-				Ids.Add(BaseId);
+				ComposeFromBaseId(BaseId);
 				LegacyType &= ~Type;
 			}
 		}
@@ -79,18 +79,61 @@ void FPCGDataTypeIdentifier::ComposeFromPCGDataSubclass(TSubclassOf<UPCGData> Su
 {
 	if (Subclass)
 	{
-		if (FPCGDataTypeBaseId BaseId = Subclass->GetDefaultObject<UPCGData>()->GetDataTypeId(); BaseId.IsValid())
+		ComposeFromBaseId(Subclass->GetDefaultObject<UPCGData>()->GetDataTypeId());
+	}
+}
+
+bool FPCGDataTypeIdentifier::PrepareForCompose(const FPCGDataTypeBaseId& BaseId)
+{
+	if (!BaseId.IsValid())
+	{
+		return false;
+	}
+
+	// By definition, adding Any will remove everything
+	if (BaseId == FPCGDataTypeInfo::AsId())
+	{
+		Ids.Empty(1);
+		return true;
+	}
+
+	// Make sure to never add Ids that are already covered by a wider type (or the same one), and replace any type that is the child of the added type.
+	// Order is not important, so we remove at swap.
+	// By construction, we can never have a type that would be a child and a parent of a type in Ids (either it is a child or a parent)
+	// i.e. if we have the hierarchy A -> B -> C, an identifier C | A is not constructible, so we could never try to add B to it.
+	for (int32 i = Ids.Num() - 1; i >= 0; --i)
+	{
+		const FPCGDataTypeBaseId& ThisId = Ids[i];
+
+		// This also catch BaseId == ThisId, which mean we do not remove it if it is the same (but won't add it either).
+		if (BaseId.IsChildOf(ThisId))
+		{
+			return false;
+		}
+		else if (ThisId.IsChildOf(BaseId))
 		{
-			Ids.AddUnique(BaseId);
+			Ids.RemoveAtSwap(i);
 		}
 	}
+
+	return true;
 }
 
 void FPCGDataTypeIdentifier::ComposeFromBaseId(const FPCGDataTypeBaseId& BaseId)
 {
-	if (BaseId.IsValid())
+	if (PrepareForCompose(BaseId))
+	{
+		// No need for add unique, as the unicity was checked in PrepareForCompose
+		Ids.Add(BaseId);
+	}
+}
+
+void FPCGDataTypeIdentifier::ComposeFromBaseId(FPCGDataTypeBaseId&& BaseId)
+{
+	if (PrepareForCompose(BaseId))
 	{
-		Ids.AddUnique(BaseId);
+		// No need for add unique, as the unicity was checked in PrepareForCompose
+		Ids.Add(MoveTemp(BaseId));
 	}
 }
 
@@ -117,10 +160,7 @@ FPCGDataTypeIdentifier FPCGDataTypeIdentifier::Construct(TConstArrayView<FPCGDat
 	{
 		for (const FPCGDataTypeBaseId& BaseId : Id.GetIds())
 		{
-			if (BaseId.IsValid())
-			{
-				Result.Ids.AddUnique(BaseId);
-			}
+			Result.ComposeFromBaseId(BaseId);
 		}
 	}
 
@@ -167,10 +207,7 @@ FPCGDataTypeIdentifier& FPCGDataTypeIdentifier::operator|=(const FPCGDataTypeIde
 	{
 		for (const FPCGDataTypeBaseId& Id : Other.GetIds())
 		{
-			if (Id.IsValid())
-			{
-				Ids.AddUnique(Id);
-			}
+			ComposeFromBaseId(Id);
 		}
 	}
 
@@ -186,13 +223,9 @@ FPCGDataTypeIdentifier& FPCGDataTypeIdentifier::operator|=(FPCGDataTypeIdentifie
 	}
 	else
 	{
-		for (const FPCGDataTypeBaseId& Id : Other.GetIds())
+		for (FPCGDataTypeBaseId& Id : Other.Ids)
 		{
-			if (Id.IsValid())
-			{
-				// Safe to const_cast
-				Ids.AddUnique(MoveTemp(const_cast<FPCGDataTypeBaseId&>(Id)));
-			}
+			ComposeFromBaseId(MoveTemp(Id));
 		}
 	}
 	
diff --git a/Engine/Plugins/PCG/Source/PCG/Private/Tests/Registry/PCGDataTypeRegistryTests.cpp b/Engine/Plugins/PCG/Source/PCG/Private/Tests/Registry/PCGDataTypeRegistryTests.cpp
--- a/Engine/Plugins/PCG/Source/PCG/Private/Tests/Registry/PCGDataTypeRegistryTests.cpp
+++ b/Engine/Plugins/PCG/Source/PCG/Private/Tests/Registry/PCGDataTypeRegistryTests.cpp
@@ -142,7 +142,7 @@ bool FPCGByDefaultTypeTest::RunTest(const FString& Parameters)
 
 bool FPCGCompatibilityTest::RunTest(const FString& Parameters)
 {
-	TTuple<EPCGDataType, EPCGDataType, EPCGDataTypeCompatibilityResult> TestValues[] =
+	TTuple<FPCGDataTypeIdentifier, FPCGDataTypeIdentifier, EPCGDataTypeCompatibilityResult> TestValues[] =
 	{
 		{EPCGDataType::Any, EPCGDataType::None, EPCGDataTypeCompatibilityResult::UnknownType},
 		{EPCGDataType::None, EPCGDataType::None, EPCGDataTypeCompatibilityResult::UnknownType},
@@ -181,6 +181,12 @@ bool FPCGCompatibilityTest::RunTest(const FString& Parameters)
 		{EPCGDataType::Landscape, EPCGDataType::Texture, EPCGDataTypeCompatibilityResult::NotCompatible},
 		{EPCGDataType::Landscape, EPCGDataType::Surface, EPCGDataTypeCompatibilityResult::Compatible},
 		{EPCGDataType::Surface, EPCGDataType::Landscape, EPCGDataTypeCompatibilityResult::RequireFilter},
+
+		// Explicit construction of overlapping types
+		{EPCGDataType::Surface, FPCGDataTypeIdentifier::Construct(EPCGDataType::Surface, EPCGDataType::Landscape), EPCGDataTypeCompatibilityResult::Compatible},
+		{EPCGDataType::Landscape, FPCGDataTypeIdentifier::Construct(EPCGDataType::Surface, EPCGDataType::Landscape), EPCGDataTypeCompatibilityResult::Compatible},
+		{EPCGDataType::Surface, FPCGDataTypeIdentifier::Construct(EPCGDataType::Landscape, EPCGDataType::Surface), EPCGDataTypeCompatibilityResult::Compatible},
+		{EPCGDataType::Landscape, FPCGDataTypeIdentifier::Construct(EPCGDataType::Landscape, EPCGDataType::Surface), EPCGDataTypeCompatibilityResult::Compatible},
 		
 		{EPCGDataType::RenderTarget, EPCGDataType::Texture, EPCGDataTypeCompatibilityResult::NotCompatible},
 		{EPCGDataType::BaseTexture, EPCGDataType::Texture, EPCGDataTypeCompatibilityResult::RequireFilter},
@@ -198,7 +204,7 @@ bool FPCGCompatibilityTest::RunTest(const FString& Parameters)
 		{EPCGDataType::Point, EPCGDataType::PointOrParam, EPCGDataTypeCompatibilityResult::Compatible},
 		{EPCGDataType::Param, EPCGDataType::Param, EPCGDataTypeCompatibilityResult::Compatible},
 
-		
+
 		{EPCGDataType::Spline | EPCGDataType::Point, EPCGDataType::Spline, EPCGDataTypeCompatibilityResult::RequireFilter},
 		{EPCGDataType::Spline | EPCGDataType::Point, EPCGDataType::Point, EPCGDataTypeCompatibilityResult::RequireFilter},
 		{EPCGDataType::Spline | EPCGDataType::Point, EPCGDataType::Param, EPCGDataTypeCompatibilityResult::NotCompatible},
@@ -212,22 +218,17 @@ bool FPCGCompatibilityTest::RunTest(const FString& Parameters)
 
 	const FPCGDataTypeRegistry& Registry = FPCGModule::GetConstDataTypeRegistry();
 
-	const UEnum* Enum = StaticEnum<EPCGDataType>();
-	check(Enum);
-
 	for (const auto& [TypeA, TypeB, ExpectedCompatibility] : TestValues)
 	{
-		const FName TypeAName = Enum->GetNameByValue(static_cast<int64>(TypeA));
-		const FName TypeBName = Enum->GetNameByValue(static_cast<int64>(TypeB));
+		const FString TypeAName = TypeA.ToString();
+		const FString TypeBName = TypeB.ToString();
 
-		const FPCGDataTypeIdentifier TypeAIdentifier{TypeA};
-		const FPCGDataTypeIdentifier TypeBIdentifier{TypeB};
 		
 		UTEST_EQUAL(
 			FString::Printf(TEXT("Compatiblity between %s and %s is expected."),
-				*TypeAName.ToString(),
-				*TypeBName.ToString()),
-			static_cast<int64>(Registry.IsCompatible(TypeAIdentifier, TypeBIdentifier)),
+				*TypeAName,
+				*TypeBName),
+			static_cast<int64>(Registry.IsCompatible(TypeA, TypeB)),
 			static_cast<int64>(ExpectedCompatibility)
 		);
 	}
@@ -450,32 +451,40 @@ bool FPCGCompositionTest::RunTest(const FString& Parameters)
 bool FPCGUnicityTest::RunTest(const FString& Parameters)
 {
 	const FPCGDataTypeIdentifier SameTypeWithClassTpl = FPCGDataTypeIdentifier::Construct<UPCGData, UPCGData>();
-	UTEST_EQUAL("Construct with same class in template gives just 1 ID", SameTypeWithClassTpl.GetIds().Num(), 1);
+	UTEST_EQUAL("Construct with same class in template gives just 1 ID", SameTypeWithClassTpl.GetIds().Num(), 1)
+
+	const FPCGDataTypeIdentifier ChildTypesWithClassTpl = FPCGDataTypeIdentifier::Construct<UPCGData, UPCGSpatialData, UPCGPointData>();
+	UTEST_EQUAL("Construct with child classes in template gives just 1 ID", SameTypeWithClassTpl.GetIds().Num(), 1)
+	UTEST_EQUAL("Construct with child classes in template gives the wider one", SameTypeWithClassTpl.GetId(), FPCGDataTypeInfo::AsId());
 
 	const FPCGDataTypeIdentifier SameTypeWithClass = FPCGDataTypeIdentifier::Construct(UPCGPointData::StaticClass(), UPCGPointData::StaticClass());
-	UTEST_EQUAL("Construct with same class gives just 1 ID", SameTypeWithClass.GetIds().Num(), 1);
+	UTEST_EQUAL("Construct with same class gives just 1 ID", SameTypeWithClass.GetIds().Num(), 1)
 
-	const FPCGDataTypeIdentifier SameTypeWithBaseIds = FPCGDataTypeIdentifier::Construct(FPCGDataTypeBaseId::Construct<FPCGDataTypeInfoLandscape>(), FPCGDataTypeBaseId::Construct<FPCGDataTypeInfoLandscape>());
-	UTEST_EQUAL("Construct with same base ids gives just 1 ID", SameTypeWithBaseIds.GetIds().Num(), 1);
+	const FPCGDataTypeIdentifier SameTypeWithBaseIds = FPCGDataTypeIdentifier::Construct(FPCGDataTypeInfoLandscape::AsId(), FPCGDataTypeInfoLandscape::AsId());
+	UTEST_EQUAL("Construct with same base ids gives just 1 ID", SameTypeWithBaseIds.GetIds().Num(), 1)
+	
+	const FPCGDataTypeIdentifier ChildTypeWithBaseIds = FPCGDataTypeIdentifier::Construct(FPCGDataTypeInfoLandscape::AsId(), FPCGDataTypeInfoSurface::AsId(), FPCGDataTypeInfoLandscape::AsId());
+	UTEST_EQUAL("Construct with child base ids gives just 1 ID", ChildTypeWithBaseIds.GetIds().Num(), 1)
+	UTEST_EQUAL("Construct with child base ids gives the wider one", ChildTypeWithBaseIds.GetId(), FPCGDataTypeInfoSurface::AsId())
 
 	const FPCGDataTypeIdentifier SameTypeWithIds = FPCGDataTypeIdentifier::Construct(SameTypeWithClassTpl, SameTypeWithClassTpl, SameTypeWithClassTpl);
-	UTEST_EQUAL("Construct with same ids gives just 1 ID", SameTypeWithIds.GetIds().Num(), 1);
+	UTEST_EQUAL("Construct with same ids gives just 1 ID", SameTypeWithIds.GetIds().Num(), 1)
 
 	const FPCGDataTypeIdentifier SameTypeWithDefaultTypeOr = SameTypeWithClass | EPCGDataType::Point;
-	UTEST_EQUAL("Bitwise Or with same default type gives just 1 ID", SameTypeWithDefaultTypeOr.GetIds().Num(), 1);
+	UTEST_EQUAL("Bitwise Or with same default type gives just 1 ID", SameTypeWithDefaultTypeOr.GetIds().Num(), 1)
 
 	const FPCGDataTypeIdentifier SameTypeWithOr = SameTypeWithClass | SameTypeWithClass;
-	UTEST_EQUAL("Bitwise Or with same ids gives just 1 ID", SameTypeWithOr.GetIds().Num(), 1);
+	UTEST_EQUAL("Bitwise Or with same ids gives just 1 ID", SameTypeWithOr.GetIds().Num(), 1)
 
 	FPCGDataTypeIdentifier SameTypeOr = SameTypeWithClass;
 	SameTypeOr |= SameTypeWithClass;
-	UTEST_EQUAL("Bitwise Or with self gives just 1 ID", SameTypeOr.GetIds().Num(), 1);
+	UTEST_EQUAL("Bitwise Or with self gives just 1 ID", SameTypeOr.GetIds().Num(), 1)
 
 	SameTypeOr |= EPCGDataType::Point;
-	UTEST_EQUAL("Bitwise Or on self with default type gives just 1 ID", SameTypeOr.GetIds().Num(), 1);
+	UTEST_EQUAL("Bitwise Or on self with default type gives just 1 ID", SameTypeOr.GetIds().Num(), 1)
 
 	const FPCGDataTypeIdentifier SameTypeCompose = SameTypeWithClass.Compose({SameTypeWithClass, SameTypeWithClass});
-	UTEST_EQUAL("Compose with same type gives just 1 ID", SameTypeCompose.GetIds().Num(), 1);
+	UTEST_EQUAL("Compose with same type gives just 1 ID", SameTypeCompose.GetIds().Num(), 1)
 
 	return true;
 }
diff --git a/Engine/Plugins/PCG/Source/PCG/Public/Data/Registry/PCGDataTypeIdentifier.h b/Engine/Plugins/PCG/Source/PCG/Public/Data/Registry/PCGDataTypeIdentifier.h
--- a/Engine/Plugins/PCG/Source/PCG/Public/Data/Registry/PCGDataTypeIdentifier.h
+++ b/Engine/Plugins/PCG/Source/PCG/Public/Data/Registry/PCGDataTypeIdentifier.h
@@ -203,6 +203,9 @@ private:
 	PCG_API void ComposeFromLegacyType(EPCGDataType LegacyType);
 	PCG_API void ComposeFromPCGDataSubclass(TSubclassOf<UPCGData> Subclass);
 	PCG_API void ComposeFromBaseId(const FPCGDataTypeBaseId& BaseId);
+	PCG_API void ComposeFromBaseId(FPCGDataTypeBaseId&& BaseId);
+	
+	bool PrepareForCompose(const FPCGDataTypeBaseId& BaseId);
 
 private:
 	UPROPERTY(BlueprintReadOnly, VisibleAnywhere, Category = "PCG", meta = (AllowPrivateAccess))
