[PCG] Fix ensure when calling RefreshRuntimeGenComponent with a null or invalid PCG Component

#rnx

diff --git a/Engine/Plugins/PCG/Source/PCG/Private/Subsystems/PCGSubsystem.cpp b/Engine/Plugins/PCG/Source/PCG/Private/Subsystems/PCGSubsystem.cpp
--- a/Engine/Plugins/PCG/Source/PCG/Private/Subsystems/PCGSubsystem.cpp
+++ b/Engine/Plugins/PCG/Source/PCG/Private/Subsystems/PCGSubsystem.cpp
@@ -937,12 +937,7 @@ void UPCGSubsystem::CancelGenerationInternal(IPCGGraphExecutionSource* Source, b
 
 void UPCGSubsystem::RefreshRuntimeGenComponent(UPCGComponent* RuntimeComponent, EPCGChangeType ChangeType)
 {
-	if (!ensure(RuntimeComponent && RuntimeComponent->IsManagedByRuntimeGenSystem()))
-	{
-		return;
-	}
-
-	if (ensure(RuntimeGenScheduler))
+	if (RuntimeComponent && RuntimeComponent->IsManagedByRuntimeGenSystem() && RuntimeGenScheduler)
 	{
 		// Only need to remove PAs if the grid sizes have changed.
 		const bool bRemovePartitionActors = !!(ChangeType & EPCGChangeType::GenerationGrid);
