[PCG] Editor toolmode robustness fixes: crash on cancel w/ painting tool, crash on painting after removing a layer.

#rnx

diff --git a/Engine/Plugins/PCG/Source/PCG/Private/Data/PCGToolData.cpp b/Engine/Plugins/PCG/Source/PCG/Private/Data/PCGToolData.cpp
--- a/Engine/Plugins/PCG/Source/PCG/Private/Data/PCGToolData.cpp
+++ b/Engine/Plugins/PCG/Source/PCG/Private/Data/PCGToolData.cpp
@@ -121,7 +121,8 @@ void FPCGInteractiveToolWorkingData_PointArrayData::OnToolCancel(const FPCGInter
 {
 	if(OnToolStartPointArray != nullptr)
 	{
-		UEngine::CopyPropertiesForUnrelatedObjects(OnToolStartPointArray, GeneratedPointData);
+		// We'd expect no relationship between the generated point data on tool cancel, but we'll make sure this is correct here.
+		GeneratedPointData = Cast<UPCGPointArrayData>(StaticDuplicateObject(OnToolStartPointArray, Context.OwningPCGComponent.Get()));
 	}
 	
 	FPCGInteractiveToolWorkingData::OnToolCancel(Context);
diff --git a/Engine/Plugins/PCG/Source/PCGEditor/Private/EditorMode/Tools/Paint/PCGPaintTool.cpp b/Engine/Plugins/PCG/Source/PCGEditor/Private/EditorMode/Tools/Paint/PCGPaintTool.cpp
--- a/Engine/Plugins/PCG/Source/PCGEditor/Private/EditorMode/Tools/Paint/PCGPaintTool.cpp
+++ b/Engine/Plugins/PCG/Source/PCGEditor/Private/EditorMode/Tools/Paint/PCGPaintTool.cpp
@@ -269,6 +269,13 @@ double UPCGPaintTool::EstimateMaximumTargetDimension()
 void UPCGPaintTool::AddSinglePointFromBrush()
 {
 	UPCGPointArrayData* PointData = ToolSettings->GetTypedWorkingData<FPCGInteractiveToolWorkingData_PointArrayData>(ToolSettings->DataInstance)->GetPointArrayData();
+
+	// If the current layer doesn't exist anymore (from the graph), it's possible the working data has been changed/removed.
+	// @todo_pcg: Preferably this should get caught when the graph changes.
+	if (!PointData)
+	{
+		return;
+	}
 	
 	const int32 PointIndex = PointData->GetNumPoints();
 	PointData->SetNumPoints(PointIndex + 1);
