[PCG] Move FastGeo primitive factory helpers out of StaticMeshDataInterface so that it can be made private again

#rnx
#FailedValidations Preflight

diff --git a/Engine/Plugins/Experimental/PCGInterops/PCGFastGeoInterop/Source/PCGFastGeoInterop/Private/Compute/PrimitiveFactories/PCGPrimitiveFactoryFastGeoPISMC.h b/Engine/Plugins/Experimental/PCGInterops/PCGFastGeoInterop/Source/PCGFastGeoInterop/Private/Compute/PrimitiveFactories/PCGPrimitiveFactoryFastGeoPISMC.h
--- a/Engine/Plugins/Experimental/PCGInterops/PCGFastGeoInterop/Source/PCGFastGeoInterop/Private/Compute/PrimitiveFactories/PCGPrimitiveFactoryFastGeoPISMC.h
+++ b/Engine/Plugins/Experimental/PCGInterops/PCGFastGeoInterop/Source/PCGFastGeoInterop/Private/Compute/PrimitiveFactories/PCGPrimitiveFactoryFastGeoPISMC.h
@@ -2,6 +2,7 @@
 
 #pragma once
 
+#include "Components/PCGProceduralISMComponentDescriptor.h"
 #include "Compute/PrimitiveFactories/PCGPrimitiveFactoryISMBase.h"
 
 #include "FastGeoWeakElement.h"
@@ -10,7 +11,6 @@
 struct FPCGContext;
 class FPrimitiveSceneProxy;
 class UPrimitiveComponent;
-struct FPCGProceduralISMComponentDescriptor;
 
 class FPCGPrimitiveFactoryFastGeoPISMC : public IPCGPrimitiveFactoryISMBase
 {
diff --git a/Engine/Plugins/Experimental/PCGInterops/PCGFastGeoInterop/Source/PCGFastGeoInterop/Private/PCGFastGeoInteropModule.cpp b/Engine/Plugins/Experimental/PCGInterops/PCGFastGeoInterop/Source/PCGFastGeoInterop/Private/PCGFastGeoInteropModule.cpp
--- a/Engine/Plugins/Experimental/PCGInterops/PCGFastGeoInterop/Source/PCGFastGeoInterop/Private/PCGFastGeoInteropModule.cpp
+++ b/Engine/Plugins/Experimental/PCGInterops/PCGFastGeoInterop/Source/PCGFastGeoInterop/Private/PCGFastGeoInteropModule.cpp
@@ -4,14 +4,13 @@
 
 #include "Compute/PrimitiveFactories/PCGPrimitiveFactoryFastGeoPISMC.h"
 
-#include "Compute/DataInterfaces/Elements/PCGStaticMeshSpawnerDataInterface.h"
 #include "Modules/ModuleManager.h"
 
 IMPLEMENT_MODULE(FPCGFastGeoInteropModule, PCGFastGeoInterop);
 
 void FPCGFastGeoInteropModule::StartupModule()
 {
-	PCGStaticMeshSpawnerDataInterface::Private::SetupFastGeoPrimitiveFactory([]()
+	PCGPrimitiveFactoryHelpers::Private::SetupFastGeoPrimitiveFactory([]()
 	{
 		return MakeShared<FPCGPrimitiveFactoryFastGeoPISMC>();
 	});
@@ -19,5 +18,5 @@ void FPCGFastGeoInteropModule::StartupModule()
 
 void FPCGFastGeoInteropModule::ShutdownModule()
 {
-	PCGStaticMeshSpawnerDataInterface::Private::ResetFastGeoPrimitiveFactory();
+	PCGPrimitiveFactoryHelpers::Private::ResetFastGeoPrimitiveFactory();
 }
diff --git a/Engine/Plugins/PCG/Source/PCG/Private/Compute/DataInterfaces/Elements/PCGStaticMeshSpawnerDataInterface.cpp b/Engine/Plugins/PCG/Source/PCG/Private/Compute/DataInterfaces/Elements/PCGStaticMeshSpawnerDataInterface.cpp
--- a/Engine/Plugins/PCG/Source/PCG/Private/Compute/DataInterfaces/Elements/PCGStaticMeshSpawnerDataInterface.cpp
+++ b/Engine/Plugins/PCG/Source/PCG/Private/Compute/DataInterfaces/Elements/PCGStaticMeshSpawnerDataInterface.cpp
@@ -47,23 +47,6 @@ namespace PCGStaticMeshSpawnerDataInterface
 		TEXT("pcg.RuntimeGeneration.ISM.ComponentlessPrimitives"),
 		false,
 		TEXT("Uses a component-less path for creating primitives in game worlds (PIE and standalone)."));
-
-	static TFunction<TSharedPtr<IPCGPrimitiveFactoryISMBase>()> FastGeoPrimitiveFactoryGetter;
-
-	namespace Private
-	{
-		void SetupFastGeoPrimitiveFactory(TFunction<TSharedPtr<IPCGPrimitiveFactoryISMBase>()>&& Getter)
-		{
-			ensure(!FastGeoPrimitiveFactoryGetter); // Trivial set and reset patterns expected.
-			FastGeoPrimitiveFactoryGetter = MoveTemp(Getter);
-		}
-
-		void ResetFastGeoPrimitiveFactory()
-		{
-			ensure(FastGeoPrimitiveFactoryGetter); // Trivial set and reset patterns expected.
-			FastGeoPrimitiveFactoryGetter.Reset();
-		}
-	}
 }
 
 void UPCGStaticMeshSpawnerDataInterface::GetSupportedInputs(TArray<FShaderFunctionDefinition>& OutFunctions) const
@@ -799,7 +782,7 @@ bool UPCGStaticMeshSpawnerDataProvider::SetupPrimitives(FPCGContext* InContext,
 	{
 		bool bCreateComponentless = false;
 
-		if (PCGStaticMeshSpawnerDataInterface::CVarCreatePrimitivesComponentless.GetValueOnGameThread() && PCGStaticMeshSpawnerDataInterface::FastGeoPrimitiveFactoryGetter)
+		if (PCGStaticMeshSpawnerDataInterface::CVarCreatePrimitivesComponentless.GetValueOnGameThread() && PCGPrimitiveFactoryHelpers::GetFastGeoPrimitiveFactory())
 		{
 			if (UWorld* World = InContext->ExecutionSource.IsValid() ? InContext->ExecutionSource->GetExecutionState().GetWorld() : nullptr)
 			{
@@ -809,7 +792,7 @@ bool UPCGStaticMeshSpawnerDataProvider::SetupPrimitives(FPCGContext* InContext,
 
 		if (bCreateComponentless)
 		{
-			Factory = PCGStaticMeshSpawnerDataInterface::FastGeoPrimitiveFactoryGetter();
+			Factory = PCGPrimitiveFactoryHelpers::GetFastGeoPrimitiveFactory();
 		}
 		else
 		{
diff --git a/Engine/Plugins/PCG/Source/PCG/Private/Compute/DataInterfaces/Elements/PCGStaticMeshSpawnerDataInterface.h b/Engine/Plugins/PCG/Source/PCG/Private/Compute/DataInterfaces/Elements/PCGStaticMeshSpawnerDataInterface.h
new file mode 100644
--- /dev/null
+++ b/Engine/Plugins/PCG/Source/PCG/Private/Compute/DataInterfaces/Elements/PCGStaticMeshSpawnerDataInterface.h
@@ -0,0 +1,136 @@
+// Copyright Epic Games, Inc. All Rights Reserved.
+
+#pragma once
+
+#include "Compute/DataInterfaces/PCGComputeDataInterface.h"
+#include "ComputeFramework/ComputeDataProvider.h"
+
+#include "Components/PCGProceduralISMComponentDescriptor.h"
+
+#include "PCGStaticMeshSpawnerDataInterface.generated.h"
+
+class FPCGStaticMeshSpawnerDataInterfaceParameters;
+struct FPCGContext;
+struct FStreamableHandle;
+
+/** Data Interface to marshal Static Mesh Spawner settings to the GPU. */
+UCLASS(ClassGroup = (Procedural))
+class UPCGStaticMeshSpawnerDataInterface : public UPCGComputeDataInterface
+{
+	GENERATED_BODY()
+
+public:
+	//~ Begin UComputeDataInterface Interface
+	TCHAR const* GetClassName() const override { return TEXT("PCGStaticMeshSpawner"); }
+	void GetSupportedInputs(TArray<FShaderFunctionDefinition>& OutFunctions) const override;
+	void GetShaderParameters(TCHAR const* UID, FShaderParametersMetadataBuilder& InOutBuilder, FShaderParametersMetadataAllocations& InOutAllocations) const override;
+	void GetHLSL(FString& OutHLSL, FString const& InDataInterfaceName) const override;
+	UComputeDataProvider* CreateDataProvider() const override;
+	//~ End UComputeDataInterface Interface
+
+	static constexpr uint32 MAX_ATTRIBUTES = 64;
+};
+
+UCLASS()
+class UPCGStaticMeshSpawnerDataProvider : public UPCGComputeDataProvider
+{
+	GENERATED_BODY()
+
+public:
+	//~ Begin UComputeDataProvider Interface
+	virtual bool PerformPreExecuteReadbacks_GameThread(UPCGDataBinding* InBinding) override;
+	virtual bool PrepareForExecute_GameThread(UPCGDataBinding* InBinding) override;
+	FComputeDataProviderRenderProxy* GetRenderProxy() override;
+	virtual void Reset() override;
+	//~ End UComputeDataProvider Interface
+
+	void CreatePrimitiveDescriptors(FPCGContext* InContext, UPCGDataBinding* InBinding);
+
+	/** Setup primitives. Returns true if at least one component was set up. */
+	virtual bool SetupPrimitives(FPCGContext* InContext, UPCGDataBinding* InBinding);
+
+public:
+	/** Attributes to use for writing per-instance custom floats. */
+	UPROPERTY()
+	TArray<FUintVector4> AttributeIdOffsetStrides;
+
+	UPROPERTY()
+	TArray<int32> PrimitiveStringKeys;
+
+	UPROPERTY()
+	TArray<FBox> PrimitiveMeshBounds;
+
+	UPROPERTY()
+	TArray<float> PrimitiveSelectionCDF;
+
+	/** Attribute Id for mesh selector. */
+	UPROPERTY()
+	int32 SelectorAttributeId = INDEX_NONE;
+
+	UPROPERTY()
+	int32 NumInputPoints = 0;
+
+	/** Output attribute Id for the output selected mesh. */
+	UPROPERTY()
+	int32 SelectedMeshAttributeId = INDEX_NONE;
+
+	/** The number of instances per string key value, used for by-attribute spawning. */
+	TMap<int32, uint32> StringKeyToInstanceCount;
+
+	int32 AnalysisDataIndex = INDEX_NONE;
+
+	bool bPrimitiveDescriptorsCreated = false;
+
+	TArray<FPCGProceduralISMComponentDescriptor> PrimitiveDescriptors;
+
+	uint32 CustomFloatCount = 0;
+
+	bool bRegisteredPrimitives = false;
+
+	bool bStaticMeshesLoaded = false;
+
+	bool bPrimitivesSetUp = false;
+
+	/** Number of primitives created during primitive setup. */
+	int32 NumPrimitivesSetup = 0;
+
+	/** Keeps loaded objects alive. */
+	TSharedPtr<FStreamableHandle> LoadHandle;
+};
+
+class FPCGStaticMeshSpawnerDataProviderProxy : public FComputeDataProviderRenderProxy
+{
+public:
+	FPCGStaticMeshSpawnerDataProviderProxy(
+		const TArray<FUintVector4>& InAttributeIdOffsetStrides,
+		int32 InSelectorAttributeId,
+		const TArray<int32>& InPrimitiveStringKeys,
+		TArray<float> InSelectionCDF,
+		int32 InSelectedMeshAttributeId,
+		const TArray<FBox>& InPrimitiveMeshBounds);
+
+	//~ Begin FComputeDataProviderRenderProxy Interface
+	bool IsValid(FValidationData const& InValidationData) const override;
+	void AllocateResources(FRDGBuilder& GraphBuilder, FAllocationData const& InAllocationData) override;
+	void GatherDispatchData(FDispatchData const& InDispatchData) override;
+	//~ End FComputeDataProviderRenderProxy Interface
+
+protected:
+	using FParameters = FPCGStaticMeshSpawnerDataInterfaceParameters;
+
+	TArray<FUintVector4> AttributeIdOffsetStrides;
+	TArray<float> SelectionCDF;
+
+	int32 SelectorAttributeId = INDEX_NONE;
+
+	TArray<int32> PrimitiveStringKeys;
+	FRDGBufferSRVRef PrimitiveStringKeysBufferSRV = nullptr;
+
+	TArray<FVector4f> PrimitiveMeshBoundsMin;
+	FRDGBufferSRVRef PrimitiveMeshBoundsMinBufferSRV = nullptr;
+
+	TArray<FVector4f> PrimitiveMeshBoundsMax;
+	FRDGBufferSRVRef PrimitiveMeshBoundsMaxBufferSRV = nullptr;
+
+	int32 SelectedMeshAttributeId = INDEX_NONE;
+};
diff --git a/Engine/Plugins/PCG/Source/PCG/Private/Compute/PrimitiveFactories/PCGPrimitiveFactoryISMBase.cpp b/Engine/Plugins/PCG/Source/PCG/Private/Compute/PrimitiveFactories/PCGPrimitiveFactoryISMBase.cpp
new file mode 100644
--- /dev/null
+++ b/Engine/Plugins/PCG/Source/PCG/Private/Compute/PrimitiveFactories/PCGPrimitiveFactoryISMBase.cpp
@@ -0,0 +1,28 @@
+// Copyright Epic Games, Inc. All Rights Reserved.
+
+#include "Compute/PrimitiveFactories/PCGPrimitiveFactoryISMBase.h"
+
+namespace PCGPrimitiveFactoryHelpers
+{
+	static TFunction<TSharedPtr<IPCGPrimitiveFactoryISMBase>()> FastGeoPrimitiveFactoryGetter;
+
+	TSharedPtr<IPCGPrimitiveFactoryISMBase> GetFastGeoPrimitiveFactory()
+	{
+		return FastGeoPrimitiveFactoryGetter ? FastGeoPrimitiveFactoryGetter() : nullptr;
+	}
+
+	namespace Private
+	{
+		void SetupFastGeoPrimitiveFactory(TFunction<TSharedPtr<IPCGPrimitiveFactoryISMBase>()>&& Getter)
+		{
+			ensure(!FastGeoPrimitiveFactoryGetter); // Trivial set and reset patterns expected.
+			FastGeoPrimitiveFactoryGetter = MoveTemp(Getter);
+		}
+
+		void ResetFastGeoPrimitiveFactory()
+		{
+			ensure(FastGeoPrimitiveFactoryGetter); // Trivial set and reset patterns expected.
+			FastGeoPrimitiveFactoryGetter.Reset();
+		}
+	}
+}
diff --git a/Engine/Plugins/PCG/Source/PCG/Public/Compute/DataInterfaces/Elements/PCGStaticMeshSpawnerDataInterface.h b/Engine/Plugins/PCG/Source/PCG/Public/Compute/DataInterfaces/Elements/PCGStaticMeshSpawnerDataInterface.h
deleted file mode 100644
diff --git a/Engine/Plugins/PCG/Source/PCG/Public/Compute/PrimitiveFactories/PCGPrimitiveFactoryISMBase.h b/Engine/Plugins/PCG/Source/PCG/Public/Compute/PrimitiveFactories/PCGPrimitiveFactoryISMBase.h
--- a/Engine/Plugins/PCG/Source/PCG/Public/Compute/PrimitiveFactories/PCGPrimitiveFactoryISMBase.h
+++ b/Engine/Plugins/PCG/Source/PCG/Public/Compute/PrimitiveFactories/PCGPrimitiveFactoryISMBase.h
@@ -23,3 +23,18 @@ public:
 	virtual bool Create(FPCGContext* InContext) = 0;
 	virtual FBox GetMeshBounds(int32 InPrimitiveIndex) const = 0;
 };
+
+namespace PCGPrimitiveFactoryHelpers
+{
+	TSharedPtr<IPCGPrimitiveFactoryISMBase> GetFastGeoPrimitiveFactory();
+
+	// todo_pcg: Remove and integrate FastGeo component spawning directly into PCG once FastGeo exits experimental.
+	namespace Private
+	{
+		UE_EXPERIMENTAL(5.7, "This API is internal and not intended to be used.")
+		PCG_API void SetupFastGeoPrimitiveFactory(TFunction<TSharedPtr<IPCGPrimitiveFactoryISMBase>()>&& Getter);
+
+		UE_EXPERIMENTAL(5.7, "This API is internal and not intended to be used.")
+		PCG_API void ResetFastGeoPrimitiveFactory();
+	}
+}
