[PCG] Move Landscape Cache Entry to shared ptr

* With PCG multithreading, accessing the Landscape Cache entry was not threadsafe
* We could hold a pointer on a thread that would get deleted in the main thread
* CacheEntry is now shared ptr, that will make the entry persist while executing, even if the entry gets deleted.

#rnx

diff --git a/Engine/Plugins/PCG/Source/PCG/Private/Data/PCGLandscapeData.cpp b/Engine/Plugins/PCG/Source/PCG/Private/Data/PCGLandscapeData.cpp
--- a/Engine/Plugins/PCG/Source/PCG/Private/Data/PCGLandscapeData.cpp
+++ b/Engine/Plugins/PCG/Source/PCG/Private/Data/PCGLandscapeData.cpp
@@ -343,7 +343,7 @@ bool UPCGLandscapeData::ProjectPoint(const FTransform& InTransform, const FBox&
 	const FVector LocalPoint = LandscapeTransform.InverseTransformPosition(InTransform.GetLocation());
 	const FIntPoint ComponentMapKey(FMath::FloorToInt(LocalPoint.X / LandscapeInfo->ComponentSizeQuads), FMath::FloorToInt(LocalPoint.Y / LandscapeInfo->ComponentSizeQuads));
 
-	const FPCGLandscapeCacheEntry* LandscapeCacheEntry = LandscapeCache->GetCacheEntry(LandscapeInfo, ComponentMapKey);
+	TSharedPtr<const FPCGLandscapeCacheEntry> LandscapeCacheEntry = LandscapeCache->GetCacheEntry(LandscapeInfo, ComponentMapKey);
 
 	if (!LandscapeCacheEntry)
 	{
@@ -361,6 +361,9 @@ bool UPCGLandscapeData::ProjectPoint(const FTransform& InTransform, const FBox&
 		LandscapeCacheEntry->GetInterpolatedPoint(ComponentLocalPoint, OutPoint, DataProps.bGetLayerWeights ? OutMetadata : nullptr);
 	}
 
+	// Not needed anymore
+	LandscapeCacheEntry.Reset();
+
 	ULandscapeHeightfieldCollisionComponent* LandscapeCollisionComponent = LandscapeInfo->XYtoCollisionComponentMap.FindRef(ComponentMapKey);
 
 	if (DataProps.bGetActorReference && OutMetadata && LandscapeCollisionComponent)
@@ -529,7 +532,7 @@ const UPCGBasePointData* UPCGLandscapeData::CreateBasePointData(FPCGContext* Con
 			{
 				FIntPoint ComponentMapKey(ComponentX, ComponentY);
 
-				const FPCGLandscapeCacheEntry* LandscapeCacheEntry = LandscapeCache->GetCacheEntry(LandscapeInfo, ComponentMapKey);
+				const TSharedPtr<const FPCGLandscapeCacheEntry> LandscapeCacheEntry = LandscapeCache->GetCacheEntry(LandscapeInfo, ComponentMapKey);
 
 				if (!LandscapeCacheEntry)
 				{
diff --git a/Engine/Plugins/PCG/Source/PCG/Private/Grid/PCGLandscapeCache.cpp b/Engine/Plugins/PCG/Source/PCG/Private/Grid/PCGLandscapeCache.cpp
--- a/Engine/Plugins/PCG/Source/PCG/Private/Grid/PCGLandscapeCache.cpp
+++ b/Engine/Plugins/PCG/Source/PCG/Private/Grid/PCGLandscapeCache.cpp
@@ -67,7 +67,7 @@ namespace PCGLandscapeCache
 }
 
 #if WITH_EDITOR
-FPCGLandscapeCacheEntry* FPCGLandscapeCacheEntry::CreateCacheEntry(ULandscapeInfo* LandscapeInfo, ULandscapeComponent* InComponent)
+TSharedPtr<FPCGLandscapeCacheEntry> FPCGLandscapeCacheEntry::CreateCacheEntry(ULandscapeInfo* LandscapeInfo, ULandscapeComponent* InComponent)
 {
 	TRACE_CPUPROFILER_EVENT_SCOPE(FPCGLandscapeCacheEntry::CreateCacheEntry);
 
@@ -92,7 +92,7 @@ FPCGLandscapeCacheEntry* FPCGLandscapeCacheEntry::CreateCacheEntry(ULandscapeInf
 		return nullptr;
 	}
 
-	FPCGLandscapeCacheEntry *Result = new FPCGLandscapeCacheEntry();
+	TSharedPtr<FPCGLandscapeCacheEntry> Result = MakeShared<FPCGLandscapeCacheEntry>();
 
 	Result->PointHalfSize = InComponent->GetComponentTransform().GetScale3D() * 0.5;
 	Result->Stride = Stride;
@@ -541,12 +541,12 @@ void UPCGLandscapeCache::Serialize(FArchive& Archive)
 			TPair<FGuid, FIntPoint> Key;
 			Archive << Key;
 
-			FPCGLandscapeCacheEntry* Entry = new FPCGLandscapeCacheEntry();
+			TSharedPtr<FPCGLandscapeCacheEntry> Entry = MakeShared<FPCGLandscapeCacheEntry>();
 			Entry->OwningCache = this;
 			Entry->Serialize(Archive, this, EntryIndex, SerializedContents);
 
 			CacheMapKey MapKey(Key.Key, Key.Value, nullptr);
-			CachedData.Add(MapKey, Entry);
+			CachedData.Add(MapKey, MoveTemp(Entry));
 		}
 
 #if WITH_EDITOR
@@ -589,10 +589,10 @@ void UPCGLandscapeCache::UpdateCacheWorldKeys()
 		return;
 	}
 
-	TMap<CacheMapKey, FPCGLandscapeCacheEntry*> UpdatedCachedData;
+	TMap<CacheMapKey, TSharedPtr<FPCGLandscapeCacheEntry>> UpdatedCachedData;
 	UpdatedCachedData.Reserve(CachedData.Num());
 
-	for (const TPair<CacheMapKey, FPCGLandscapeCacheEntry*>& CacheEntry : CachedData)
+	for (const TPair<CacheMapKey, TSharedPtr<FPCGLandscapeCacheEntry>>& CacheEntry : CachedData)
 	{
 		CacheMapKey Key = CacheEntry.Key;
 		if (Key.WorldKey == FObjectKey())
@@ -627,8 +627,8 @@ void UPCGLandscapeCache::Tick(float DeltaSeconds)
 
 		if (CacheMemorySize > MemoryThresholdInBytes)
 		{
-			TArray<FPCGLandscapeCacheEntry*> LoadedEntries;
-			for(TPair<CacheMapKey, FPCGLandscapeCacheEntry*>& CacheEntry : CachedData)
+			TArray<TSharedPtr<FPCGLandscapeCacheEntry>> LoadedEntries;
+			for(TPair<CacheMapKey, TSharedPtr<FPCGLandscapeCacheEntry>>& CacheEntry : CachedData)
 			{
 				if (CacheEntry.Value->bDataLoaded)
 				{
@@ -636,7 +636,7 @@ void UPCGLandscapeCache::Tick(float DeltaSeconds)
 				}
 			}
 
-			Algo::Sort(LoadedEntries, [](const FPCGLandscapeCacheEntry* A, const FPCGLandscapeCacheEntry* B)
+			Algo::Sort(LoadedEntries, [](const TSharedPtr<FPCGLandscapeCacheEntry>& A, const TSharedPtr<FPCGLandscapeCacheEntry>& B)
 			{
 				return A->Touch < B->Touch;
 			});
@@ -709,7 +709,7 @@ void UPCGLandscapeCache::PrimeCache()
 	}
 
 	// Then create the landscape cache entries
-	TArray<FPCGLandscapeCacheEntry*> NewEntries;
+	TArray<TSharedPtr<FPCGLandscapeCacheEntry>> NewEntries;
 	NewEntries.SetNum(CacheEntriesToBuild.Num());
 
 	ParallelFor(CacheEntriesToBuild.Num(), [this, &CacheEntriesToBuild, &NewEntries](int32 EntryIndex)
@@ -731,7 +731,7 @@ void UPCGLandscapeCache::PrimeCache()
 	{
 		if (NewEntries[EntryIndex])
 		{
-			CachedData.Add(CacheEntriesToBuild[EntryIndex].Get<2>(), NewEntries[EntryIndex]);
+			CachedData.Add(CacheEntriesToBuild[EntryIndex].Get<2>(), MoveTemp(NewEntries[EntryIndex]));
 		}
 	}
 
@@ -747,12 +747,6 @@ void UPCGLandscapeCache::ClearCache()
 		Modify();
 	}
 
-	for(TPair<CacheMapKey, FPCGLandscapeCacheEntry*>& CacheEntry : CachedData)
-	{
-		delete CacheEntry.Value;
-		CacheEntry.Value = nullptr;
-	}
-
 	CachedData.Reset();
 	CachedLayerNames.Reset();
 #if WITH_EDITOR
@@ -770,12 +764,13 @@ void UPCGLandscapeCache::TakeOwnership(UPCGLandscapeCache* InLandscapeCache)
 	InLandscapeCache->UpdateCacheWorldKeys();
 	bool bShouldDirty = false;
 
-	for(TPair<CacheMapKey, FPCGLandscapeCacheEntry*>& CacheEntryPair : InLandscapeCache->CachedData)
+	for (TPair<CacheMapKey, TSharedPtr<FPCGLandscapeCacheEntry>>& CacheEntryPair : InLandscapeCache->CachedData)
 	{
-		FPCGLandscapeCacheEntry* CacheEntry = CachedData.FindOrAdd(CacheEntryPair.Key, CacheEntryPair.Value);
-		if (CacheEntry == CacheEntryPair.Value)
+		if (CacheEntryPair.Value && !CachedData.Contains(CacheEntryPair.Key))
 		{
-			CacheEntryPair.Value = nullptr;
+			CacheEntryPair.Value->OwningCache = this;
+			CachedData.Emplace(CacheEntryPair.Key, MoveTemp(CacheEntryPair.Value));
+			CacheEntryPair.Value = TSharedPtr<FPCGLandscapeCacheEntry>{};
 			bShouldDirty = true;
 		}
 	}
@@ -799,7 +794,7 @@ void UPCGLandscapeCache::TakeOwnership(UPCGLandscapeCache* InLandscapeCache)
 	}
 }
 
-const bool UPCGLandscapeCache::AreCacheEntriesReady(const ULandscapeInfo* InLandscapeInfo, const FIntPoint& InMinComponentKey, const FIntPoint& InMaxComponentKey)
+bool UPCGLandscapeCache::AreCacheEntriesReady(const ULandscapeInfo* InLandscapeInfo, const FIntPoint& InMinComponentKey, const FIntPoint& InMaxComponentKey)
 {
 	ALandscapeProxy* LandscapeProxy = InLandscapeInfo->GetLandscapeProxy();
 	for (int32 ComponentKeyX = InMinComponentKey.X; ComponentKeyX <= InMaxComponentKey.X; ++ComponentKeyX)
@@ -809,7 +804,7 @@ const bool UPCGLandscapeCache::AreCacheEntriesReady(const ULandscapeInfo* InLand
 			const FIntPoint ComponentKey(ComponentKeyX, ComponentKeyY);
 
 			bool bOutIsValidEntry = false;
-			const FPCGLandscapeCacheEntry* CacheEntry = GetCacheEntryInternal(InLandscapeInfo, ComponentKey, LandscapeProxy, /*bAllowLoadOrCreate=*/false, bOutIsValidEntry);
+			const TSharedPtr<const FPCGLandscapeCacheEntry> CacheEntry = GetCacheEntryInternal(InLandscapeInfo, ComponentKey, LandscapeProxy, /*bAllowLoadOrCreate=*/false, bOutIsValidEntry);
 
 			if (!bOutIsValidEntry)
 			{
@@ -839,7 +834,7 @@ void UPCGLandscapeCache::PrimeCache(const ULandscapeInfo* InLandscapeInfo, const
 	}
 }
 
-const FPCGLandscapeCacheEntry* UPCGLandscapeCache::GetCacheEntry(const ULandscapeInfo* InLandscapeInfo, const FIntPoint& InComponentKey, ALandscapeProxy* InLandscapeProxy, bool bInAllowLoadOrCreate)
+TSharedPtr<const FPCGLandscapeCacheEntry> UPCGLandscapeCache::GetCacheEntry(const ULandscapeInfo* InLandscapeInfo, const FIntPoint& InComponentKey, ALandscapeProxy* InLandscapeProxy, bool bInAllowLoadOrCreate)
 {
 	bool bOutIsValidEntry = false;
 
@@ -851,7 +846,7 @@ const FPCGLandscapeCacheEntry* UPCGLandscapeCache::GetCacheEntry(const ULandscap
 	return GetCacheEntryInternal(InLandscapeInfo, InComponentKey, InLandscapeProxy, bInAllowLoadOrCreate, bOutIsValidEntry);
 }
 
-const FPCGLandscapeCacheEntry* UPCGLandscapeCache::GetCacheEntryInternal(const ULandscapeInfo* InLandscapeInfo, const FIntPoint& InComponentKey, ALandscapeProxy* InLandscapeProxy, bool bInAllowLoadOrCreate, bool& bOutIsValidEntry)
+TSharedPtr<const FPCGLandscapeCacheEntry> UPCGLandscapeCache::GetCacheEntryInternal(const ULandscapeInfo* InLandscapeInfo, const FIntPoint& InComponentKey, ALandscapeProxy* InLandscapeProxy, bool bInAllowLoadOrCreate, bool& bOutIsValidEntry)
 {
 	// If entry can't exist this will be false (at runtime it means we don't have the data, in editor it means component key is invalid)
 	bOutIsValidEntry = false;
@@ -862,10 +857,10 @@ const FPCGLandscapeCacheEntry* UPCGLandscapeCache::GetCacheEntryInternal(const U
 #else
 	ALandscapeProxy* LandscapeProxy = InLandscapeProxy ? InLandscapeProxy : InLandscapeInfo->GetLandscapeProxy();
 	const FGuid LandscapeGuid = LandscapeProxy ? LandscapeProxy->GetOriginalLandscapeGuid() : FGuid();
-	const FPCGLandscapeCacheEntry* CacheEntry = GetCacheEntryInternal(LandscapeProxy, LandscapeGuid, InComponentKey, bInAllowLoadOrCreate);
+	TSharedPtr<const FPCGLandscapeCacheEntry> CacheEntry = GetCacheEntryInternal(LandscapeProxy, LandscapeGuid, InComponentKey, bInAllowLoadOrCreate);
 
 	// Entry is valid if it exists (it might not be loaded yet)
-	bOutIsValidEntry = !!CacheEntry;
+	bOutIsValidEntry = CacheEntry.IsValid();
 
 	return CacheEntry;
 #endif
@@ -873,7 +868,7 @@ const FPCGLandscapeCacheEntry* UPCGLandscapeCache::GetCacheEntryInternal(const U
 
 
 #if WITH_EDITOR
-const FPCGLandscapeCacheEntry* UPCGLandscapeCache::GetOrCreateCacheEntryInternal(const ULandscapeInfo* InLandscapeInfo, const FIntPoint& InComponentKey, bool bInAllowLoadOrCreate, bool& bOutIsValidEntry)
+TSharedPtr<const FPCGLandscapeCacheEntry> UPCGLandscapeCache::GetOrCreateCacheEntryInternal(const ULandscapeInfo* InLandscapeInfo, const FIntPoint& InComponentKey, bool bInAllowLoadOrCreate, bool& bOutIsValidEntry)
 {
 	bOutIsValidEntry = false;
 	
@@ -898,7 +893,7 @@ const FPCGLandscapeCacheEntry* UPCGLandscapeCache::GetOrCreateCacheEntryInternal
 	bOutIsValidEntry = true;
 
 	const FGuid LandscapeGuid = (LandscapeComponent->GetLandscapeProxy() ? LandscapeComponent->GetLandscapeProxy()->GetOriginalLandscapeGuid() : FGuid());
-	const FPCGLandscapeCacheEntry* CacheEntry = GetCacheEntryInternal(LandscapeComponent->GetLandscapeProxy(), LandscapeGuid, InComponentKey, bInAllowLoadOrCreate);
+	TSharedPtr<const FPCGLandscapeCacheEntry> CacheEntry = GetCacheEntryInternal(LandscapeComponent->GetLandscapeProxy(), LandscapeGuid, InComponentKey, bInAllowLoadOrCreate);
 	if (CacheEntry || !bInAllowLoadOrCreate)
 	{
 		return CacheEntry;
@@ -906,14 +901,14 @@ const FPCGLandscapeCacheEntry* UPCGLandscapeCache::GetOrCreateCacheEntryInternal
 		
 	FWriteScopeLock ScopeLock(CacheLock);
 	CacheMapKey ComponentKey(LandscapeGuid, InComponentKey, LandscapeComponent->GetOwner());
-	if (FPCGLandscapeCacheEntry** FoundEntry = CachedData.Find(ComponentKey))
+	if (const TSharedPtr<FPCGLandscapeCacheEntry>* FoundEntry = CachedData.Find(ComponentKey))
 	{
 		CacheEntry = *FoundEntry;
 	}
 	else
 	{
 		check(LandscapeComponent->SectionBaseX / LandscapeComponent->ComponentSizeQuads == ComponentKey.Coordinate.X && LandscapeComponent->SectionBaseY / LandscapeComponent->ComponentSizeQuads == ComponentKey.Coordinate.Y);
-		if (FPCGLandscapeCacheEntry* NewEntry = FPCGLandscapeCacheEntry::CreateCacheEntry(LandscapeComponent->GetLandscapeInfo(), LandscapeComponent))
+		if (TSharedPtr<FPCGLandscapeCacheEntry> NewEntry = FPCGLandscapeCacheEntry::CreateCacheEntry(LandscapeComponent->GetLandscapeInfo(), LandscapeComponent))
 		{
 			NewEntry->OwningCache = this;
 			CacheEntry = NewEntry;
@@ -938,9 +933,9 @@ const FPCGLandscapeCacheEntry* UPCGLandscapeCache::GetOrCreateCacheEntryInternal
 }
 #endif
 
-const FPCGLandscapeCacheEntry* UPCGLandscapeCache::GetCacheEntryInternal(AActor* HintActor, const FGuid& LandscapeGuid, const FIntPoint& ComponentCoordinate, bool bInAllowLoad)
+TSharedPtr<const FPCGLandscapeCacheEntry> UPCGLandscapeCache::GetCacheEntryInternal(AActor* HintActor, const FGuid& LandscapeGuid, const FIntPoint& ComponentCoordinate, bool bInAllowLoad)
 {
-	const FPCGLandscapeCacheEntry* CacheEntry = nullptr;
+	TSharedPtr<const FPCGLandscapeCacheEntry> CacheEntry = nullptr;
 	CacheMapKey ComponentKey(LandscapeGuid, ComponentCoordinate, HintActor);
 
 	{
@@ -953,7 +948,7 @@ const FPCGLandscapeCacheEntry* UPCGLandscapeCache::GetCacheEntryInternal(AActor*
 			bLoggedNoCacheError = true;
 		}
 #endif
-		if (FPCGLandscapeCacheEntry** FoundEntry = CachedData.Find(ComponentKey))
+		if (const TSharedPtr<FPCGLandscapeCacheEntry>* FoundEntry = CachedData.Find(ComponentKey))
 		{
 			CacheEntry = *FoundEntry;
 		}
@@ -1000,7 +995,7 @@ void UPCGLandscapeCache::SampleMetadataOnPoint(ALandscapeProxy* Landscape, const
 	const FVector LocalPoint = LandscapeTransform.InverseTransformPosition(InTransform.GetLocation());
 	const FIntPoint ComponentMapKey(FMath::FloorToInt(LocalPoint.X / LandscapeInfo->ComponentSizeQuads), FMath::FloorToInt(LocalPoint.Y / LandscapeInfo->ComponentSizeQuads));
 
-	const FPCGLandscapeCacheEntry* CacheEntry = GetCacheEntry(LandscapeInfo, ComponentMapKey, Landscape);
+	TSharedPtr<const FPCGLandscapeCacheEntry> CacheEntry = GetCacheEntry(LandscapeInfo, ComponentMapKey, Landscape);
 
 	if (!CacheEntry || CacheEntry->LayerData.IsEmpty())
 	{
@@ -1092,12 +1087,10 @@ void UPCGLandscapeCache::OnLandscapeChanged(ALandscapeProxy* InLandscape, const
 		if (LandscapeComponent)
 		{
 			const CacheMapKey ComponentKey(InLandscape->GetOriginalLandscapeGuid(), PCGLandscapeCache::GetCoordinates(LandscapeComponent), InLandscape);
-			FPCGLandscapeCacheEntry* EntryToDelete = nullptr;
 
-			if (CachedData.RemoveAndCopyValue(ComponentKey, EntryToDelete))
+			if (CachedData.Remove(ComponentKey) > 0)
 			{
 				--CacheEntryCount;
-				delete EntryToDelete;
 			}
 		}
 	});
@@ -1174,16 +1167,16 @@ void UPCGLandscapeCache::RemoveComponentFromCache(const ALandscapeProxy* Landsca
 {
 	check(LandscapeProxy);
 
+	FWriteScopeLock Lock(CacheLock);
+
 	LandscapeProxy->ForEachComponent<ULandscapeComponent>(/*bIncludeFromChildActors=*/false, [this, LandscapeProxy](const ULandscapeComponent* LandscapeComponent)
 	{
 		if (LandscapeComponent)
 		{
 			const CacheMapKey ComponentKey(LandscapeProxy->GetOriginalLandscapeGuid(), PCGLandscapeCache::GetCoordinates(LandscapeComponent), LandscapeProxy);
-			if (FPCGLandscapeCacheEntry** FoundEntry = CachedData.Find(ComponentKey))
+			if (CachedData.Remove(ComponentKey) > 0)
 			{
-				CachedData.Remove(ComponentKey);
 				--CacheEntryCount;
-				delete* FoundEntry;
 			}
 		}
 	});
diff --git a/Engine/Plugins/PCG/Source/PCG/Private/Helpers/PCGBlueprintHelpers.cpp b/Engine/Plugins/PCG/Source/PCG/Private/Helpers/PCGBlueprintHelpers.cpp
--- a/Engine/Plugins/PCG/Source/PCG/Private/Helpers/PCGBlueprintHelpers.cpp
+++ b/Engine/Plugins/PCG/Source/PCG/Private/Helpers/PCGBlueprintHelpers.cpp
@@ -304,7 +304,7 @@ TArray<FPCGLandscapeLayerWeight> UPCGBlueprintHelpers::GetInterpolatedPCGLandsca
 		const FVector LocalPoint = Landscape->LandscapeActorToWorld().InverseTransformPosition(Location);
 		const FIntPoint ComponentMapKey(FMath::FloorToInt(LocalPoint.X / LandscapeInfo->ComponentSizeQuads), FMath::FloorToInt(LocalPoint.Y / LandscapeInfo->ComponentSizeQuads));
 
-		const FPCGLandscapeCacheEntry* CacheEntry = LandscapeCache->GetCacheEntry(LandscapeInfo, ComponentMapKey, Landscape);
+		TSharedPtr<const FPCGLandscapeCacheEntry> CacheEntry = LandscapeCache->GetCacheEntry(LandscapeInfo, ComponentMapKey, Landscape);
 
 		if (!CacheEntry)
 		{
@@ -317,6 +317,9 @@ TArray<FPCGLandscapeLayerWeight> UPCGBlueprintHelpers::GetInterpolatedPCGLandsca
 		TArray<FPCGLandscapeLayerWeight> Result;
 		CacheEntry->GetInterpolatedLayerWeights(ComponentLocalPoint, Result);
 
+		// Not needed anymore
+		CacheEntry.Reset();
+
 		Result.Sort([](const FPCGLandscapeLayerWeight& Lhs, const FPCGLandscapeLayerWeight& Rhs) {
 			return Lhs.Weight > Rhs.Weight;
 		});
diff --git a/Engine/Plugins/PCG/Source/PCG/Public/Grid/PCGLandscapeCache.h b/Engine/Plugins/PCG/Source/PCG/Public/Grid/PCGLandscapeCache.h
--- a/Engine/Plugins/PCG/Source/PCG/Public/Grid/PCGLandscapeCache.h
+++ b/Engine/Plugins/PCG/Source/PCG/Public/Grid/PCGLandscapeCache.h
@@ -87,7 +87,7 @@ private:
 	int32 GetMemorySize() const;
 
 #if WITH_EDITOR
-	static FPCGLandscapeCacheEntry* CreateCacheEntry(ULandscapeInfo* LandscapeInfo, ULandscapeComponent* InComponent);
+	static TSharedPtr<FPCGLandscapeCacheEntry> CreateCacheEntry(ULandscapeInfo* LandscapeInfo, ULandscapeComponent* InComponent);
 #endif
 
 	// Serialize called from the landscape cache
@@ -142,10 +142,10 @@ public:
 
 	void Tick(float DeltaSeconds);
 
-	const bool AreCacheEntriesReady(const ULandscapeInfo* InLandscapeInfo, const FIntPoint& InMinComponentKey, const FIntPoint& InMaxComponentKey);
+	bool AreCacheEntriesReady(const ULandscapeInfo* InLandscapeInfo, const FIntPoint& InMinComponentKey, const FIntPoint& InMaxComponentKey);
 
 	/** Gets landscape cache entry for specified component key */
-	const FPCGLandscapeCacheEntry* GetCacheEntry(const ULandscapeInfo* InLandscapeInfo, const FIntPoint& InComponentKey, ALandscapeProxy* InLandscapeProxy = nullptr, bool bInAllowLoadOrCreate = true);
+	TSharedPtr<const FPCGLandscapeCacheEntry> GetCacheEntry(const ULandscapeInfo* InLandscapeInfo, const FIntPoint& InComponentKey, ALandscapeProxy* InLandscapeProxy = nullptr, bool bInAllowLoadOrCreate = true);
 
 	/** Load/Create cache entries between Min/Max component key inclusively */
 	void PrimeCache(const ULandscapeInfo* InLandscapeInfo, const FIntPoint& InMinComponentKey, const FIntPoint& InMaxComponentKey);
@@ -180,11 +180,11 @@ private:
 	void RemoveComponentFromCache(const ALandscapeProxy* LandscapeProxy);
 
 	/** Gets (and creates if needed) the cache entry - available only in Editor */
-	const FPCGLandscapeCacheEntry* GetOrCreateCacheEntryInternal(const ULandscapeInfo* InLandscapeInfo, const FIntPoint& InComponentKey, bool bInAllowLoadOrCreate, bool& bOutValidEntry);
+	TSharedPtr<const FPCGLandscapeCacheEntry> GetOrCreateCacheEntryInternal(const ULandscapeInfo* InLandscapeInfo, const FIntPoint& InComponentKey, bool bInAllowLoadOrCreate, bool& bOutValidEntry);
 #endif
 	/** Gets landscape cache entry, works both in editor (but does not create) but works in game mode too. */
-	const FPCGLandscapeCacheEntry* GetCacheEntryInternal(AActor* InHintActor, const FGuid& InLandscapeGuid, const FIntPoint& InComponentKey, bool bInAllowLoad);
-	const FPCGLandscapeCacheEntry* GetCacheEntryInternal(const ULandscapeInfo* InLandscapeInfo, const FIntPoint& InComponentKey, ALandscapeProxy* InLandscapeProxy, bool bInAllowLoadOrCreate, bool& bOutValidEntry);
+	TSharedPtr<const FPCGLandscapeCacheEntry> GetCacheEntryInternal(AActor* InHintActor, const FGuid& InLandscapeGuid, const FIntPoint& InComponentKey, bool bInAllowLoad);
+	TSharedPtr<const FPCGLandscapeCacheEntry> GetCacheEntryInternal(const ULandscapeInfo* InLandscapeInfo, const FIntPoint& InComponentKey, ALandscapeProxy* InLandscapeProxy, bool bInAllowLoadOrCreate, bool& bOutValidEntry);
 
 	struct CacheMapKey
 	{
@@ -224,7 +224,7 @@ private:
 	void UpdateCacheWorldKeys();
 
 	// Mapping of landscape guid + coordinates to entries. This is manually serialized as needed (depends on the serialize options).
-	TMap<CacheMapKey, FPCGLandscapeCacheEntry*> CachedData;
+	TMap<CacheMapKey, TSharedPtr<FPCGLandscapeCacheEntry>> CachedData;
 
 	//TODO: separate by landscape
 	UPROPERTY(VisibleAnywhere, Category = "Landscape Cache")
