[PCG] First version of volume tool.

#rnx

diff --git a/Engine/Plugins/PCG/Source/PCG/Private/Data/PCGToolData.cpp b/Engine/Plugins/PCG/Source/PCG/Private/Data/PCGToolData.cpp
--- a/Engine/Plugins/PCG/Source/PCG/Private/Data/PCGToolData.cpp
+++ b/Engine/Plugins/PCG/Source/PCG/Private/Data/PCGToolData.cpp
@@ -7,14 +7,19 @@
 #include "Data/PCGPointArrayData.h"
 #include "Helpers/PCGDynamicTrackingHelpers.h"
 
-#include "Engine/Engine.h"
+#include "Components/BrushComponent.h"
 #include "Components/SplineComponent.h"
 #include "Components/BoxComponent.h"
+#include "Engine/Engine.h"
+#include "Engine/BrushBuilder.h"
+#include "GameFramework/Volume.h"
 #include "Logging/StructuredLog.h"
 #include "UObject/UObjectGlobals.h"
 #include "UObject/Package.h"
+
 #if WITH_EDITOR
 #include "InteractiveTool.h"
+#include "ActorFactories/ActorFactory.h"
 #endif
 
 #if WITH_EDITORONLY_DATA
@@ -166,7 +171,7 @@ USplineComponent* FPCGInteractiveToolWorkingData_Spline::GetSplineComponent() co
 	return SplineComponent.Get();
 }
 
-#if WITH_EDITORONLY_DATA
+#if WITH_EDITOR
 void FPCGInteractiveToolWorkingData_Spline::InitializeInternal(const FPCGInteractiveToolWorkingDataContext& Context)
 {	
 	SplineComponent = FindOrGenerateSplineComponent(Context);
@@ -229,29 +234,17 @@ USplineComponent* FPCGInteractiveToolWorkingData_Spline::FindMatchingSplineCompo
 	TArray<USplineComponent*> AllSplineComponents;
 	Context.OwningActor->GetComponents<USplineComponent>(AllSplineComponents, false);
 	
-	USplineComponent* MatchingSplineComponent = nullptr;
-	if (Context.DataInstanceIdentifier.IsNone())
-	{
-		if(AllSplineComponents.Num() > 0)
-		{			
-			MatchingSplineComponent = AllSplineComponents[0];
-		}
-	}
-	else
+	if (!Context.DataInstanceIdentifier.IsNone())
 	{
 		// When the DataInstanceIdentifier is set, we only care about the splines with the given tag
 		AllSplineComponents.RemoveAll([DataInstanceIdentifier = Context.DataInstanceIdentifier](USplineComponent* Candidate)
 		{
 			return Candidate->ComponentHasTag(DataInstanceIdentifier) == false;
 		});
-		
-		// If we find a matching spline component with the specific tag, use it
-		if (AllSplineComponents.Num() > 0)
-		{
-			MatchingSplineComponent = Cast<USplineComponent>(AllSplineComponents[0]);
-		}
 	}
 
+	USplineComponent* MatchingSplineComponent = AllSplineComponents.IsEmpty() ? nullptr : AllSplineComponents[0];
+
 	if(AllSplineComponents.Num() > 1 && MatchingSplineComponent != nullptr)
 	{
 		UE_LOGFMT(LogPCG, Warning, "More than one spline component found. Choosing the first available: {0}", MatchingSplineComponent->GetName());
@@ -343,13 +336,13 @@ USplineComponent* FPCGInteractiveToolWorkingData_SplineSurface::GenerateMatching
 	MatchingSplineComponent->ClearSplinePoints();
 	MatchingSplineComponent->bSplineHasBeenEdited = true;
 	MatchingSplineComponent->SetClosedLoop(true);
-		
+
 	MatchingSplineComponent->RegisterComponent();
 	Context.OwningActor->AddInstanceComponent(MatchingSplineComponent);
 
 	return MatchingSplineComponent;
 }
-#endif
+#endif // WITH_EDITOR
 
 void FPCGInteractiveToolWorkingData_Spline::InitializeRuntimeElementData(FPCGContext* Context) const
 {
@@ -380,6 +373,131 @@ bool FPCGInteractiveToolWorkingData_Spline::IsValid() const
 	return SplineComponent.IsValid();
 }
 
+#if WITH_EDITOR
+void FPCGInteractiveToolWorkingData_Volume::InitializeInternal(const FPCGInteractiveToolWorkingDataContext& Context)
+{
+	Volume = Cast<AVolume>(Context.OwningActor.Get());
+	BrushComponent = FindOrCreateBrushComponent(Context);
+}
+
+void FPCGInteractiveToolWorkingData_Volume::OnToolStart(const FPCGInteractiveToolWorkingDataContext& Context)
+{
+	FPCGInteractiveToolWorkingData::OnToolStart(Context);
+
+	if (BrushComponent.IsValid())
+	{
+		OnToolStartBrushComponent = Cast<UBrushComponent>(StaticDuplicateObject(BrushComponent.Get(), GetTransientPackage()));
+	}
+
+	if (Volume.IsValid())
+	{
+		OnToolStartVolumeTransform = Volume->GetActorTransform();
+
+		if (Volume->BrushBuilder)
+		{
+			OnToolStartBrushBuilder = DuplicateObject<UBrushBuilder>(Volume->BrushBuilder, GetTransientPackage());
+		}
+	}
+}
+
+void FPCGInteractiveToolWorkingData_Volume::OnToolCancel(const FPCGInteractiveToolWorkingDataContext& Context)
+{
+	if (BrushComponent.IsValid() && GeneratedResources.GeneratedComponents.Contains(BrushComponent) == false && OnToolStartBrushComponent)
+	{
+		Volume->SetActorTransform(OnToolStartVolumeTransform);
+
+		if (OnToolStartBrushBuilder)
+		{
+			// Recreate the brush from the original actor
+			UActorFactory::CreateBrushForVolumeActor(Volume.Get(), OnToolStartBrushBuilder);
+			OnToolStartBrushBuilder->MarkAsGarbage();
+			OnToolStartBrushBuilder = nullptr;
+		}
+
+		// Restore the brush component to its previous state
+		UEngine::CopyPropertiesForUnrelatedObjects(OnToolStartBrushComponent, BrushComponent.Get());
+		OnToolStartBrushComponent->MarkAsGarbage();
+		OnToolStartBrushComponent = nullptr;
+	}
+
+	FPCGInteractiveToolWorkingData::OnToolCancel(Context);
+}
+
+UBrushComponent* FPCGInteractiveToolWorkingData_Volume::FindOrCreateBrushComponent(const FPCGInteractiveToolWorkingDataContext& Context)
+{
+	UBrushComponent* MatchingComponent = FindMatchingBrushComponent(Context);
+
+	if (!MatchingComponent)
+	{
+		MatchingComponent = GenerateMatchingBrushComponent(Context);
+		GeneratedResources.GeneratedComponents.Add(MatchingComponent);
+	}
+
+	ensureMsgf(MatchingComponent != nullptr, TEXT("Matching Brush Component should always be found after initialization as it was either found or created."));
+	return MatchingComponent;
+}
+
+UBrushComponent* FPCGInteractiveToolWorkingData_Volume::FindMatchingBrushComponent(const FPCGInteractiveToolWorkingDataContext& Context)
+{
+	TArray<UBrushComponent*> AllBrushComponents;
+	if (Context.OwningActor.IsValid())
+	{
+		Context.OwningActor->GetComponents<UBrushComponent>(AllBrushComponents, false);
+	}
+
+	UBrushComponent* MatchingComponent = nullptr;
+
+	if(!Context.DataInstanceIdentifier.IsNone())
+	{
+		AllBrushComponents.RemoveAll([DataInstanceIdentifier = Context.DataInstanceIdentifier](UBrushComponent* Candidate)
+		{
+			return !Candidate || !Candidate->ComponentHasTag(DataInstanceIdentifier);
+		});
+	}
+
+	MatchingComponent = AllBrushComponents.IsEmpty() ? nullptr : AllBrushComponents[0];
+
+	if (AllBrushComponents.Num() > 1 && MatchingComponent != nullptr)
+	{
+		UE_LOGFMT(LogPCG, Warning, "More than one brush component found. Choosing the first available: {0}", MatchingComponent->GetName());
+	}
+
+	return MatchingComponent;
+}
+
+UBrushComponent* FPCGInteractiveToolWorkingData_Volume::GenerateMatchingBrushComponent(const FPCGInteractiveToolWorkingDataContext& Context)
+{
+	if (!Context.OwningActor.IsValid())
+	{
+		return nullptr;
+	}
+
+	UBrushComponent* MatchingBrushComponent = NewObject<UBrushComponent>(Context.OwningActor.Get());
+	// ABrush setup
+	MatchingBrushComponent->Mobility = EComponentMobility::Static;
+	MatchingBrushComponent->SetCanEverAffectNavigation(false);
+	// AVolume setup
+	MatchingBrushComponent->AlwaysLoadOnClient = true;
+	MatchingBrushComponent->AlwaysLoadOnServer = true;
+	// APCGVolume setup
+	MatchingBrushComponent->SetCollisionObjectType(ECC_WorldStatic);
+	MatchingBrushComponent->SetCollisionResponseToAllChannels(ECR_Ignore);
+	MatchingBrushComponent->SetGenerateOverlapEvents(false);
+	// Transient until applied
+	MatchingBrushComponent->SetFlags(RF_Transient);
+
+	if (!Context.DataInstanceIdentifier.IsNone())
+	{
+		MatchingBrushComponent->ComponentTags.Add(Context.DataInstanceIdentifier);
+	}
+
+	MatchingBrushComponent->RegisterComponent();
+	Context.OwningActor->AddInstanceComponent(MatchingBrushComponent);
+
+	return MatchingBrushComponent;
+}
+#endif
+
 FPCGGraphToolData::FPCGGraphToolData()
 	: InitialActorClassToSpawn(AActor::StaticClass())
 {
diff --git a/Engine/Plugins/PCG/Source/PCG/Public/Data/PCGToolData.h b/Engine/Plugins/PCG/Source/PCG/Public/Data/PCGToolData.h
--- a/Engine/Plugins/PCG/Source/PCG/Public/Data/PCGToolData.h
+++ b/Engine/Plugins/PCG/Source/PCG/Public/Data/PCGToolData.h
@@ -8,11 +8,14 @@
 #include "PCGToolData.generated.h"
 
 class AActor;
+class AVolume;
 class USplineComponent;
 class UPCGPointArrayData;
 class UPCGComponent;
+class UBrushComponent;
 class UInteractiveTool;
 class UInteractiveToolPropertySet;
+class UBrushBuilder;
 
 /** A transient struct that contains resources generated during tool usage for a specific tool data struct. Used to finalize or clean up when a tool is applied or cancelled. */
 USTRUCT()
@@ -209,6 +212,43 @@ private:
 #endif
 };
 
+USTRUCT(BlueprintType, DisplayName="Volume Tool Data")
+struct FPCGInteractiveToolWorkingData_Volume : public FPCGInteractiveToolWorkingData
+{
+	GENERATED_BODY()
+public:
+	UPROPERTY(VisibleAnywhere, BlueprintReadOnly, Category = "Working Data")
+	TWeakObjectPtr<AVolume> Volume = nullptr;
+
+	UPROPERTY(VisibleAnywhere, BlueprintReadOnly, Category = "Working Data")
+	TWeakObjectPtr<UBrushComponent> BrushComponent = nullptr;
+
+#if WITH_EDITOR
+public:
+	virtual void OnToolStart(const FPCGInteractiveToolWorkingDataContext& Context) override;
+	virtual void OnToolCancel(const FPCGInteractiveToolWorkingDataContext& Context) override;
+
+private:
+	virtual void InitializeInternal(const FPCGInteractiveToolWorkingDataContext& Context) override;
+
+	UBrushComponent* FindOrCreateBrushComponent(const FPCGInteractiveToolWorkingDataContext& Context);
+	virtual UBrushComponent* FindMatchingBrushComponent(const FPCGInteractiveToolWorkingDataContext& Context);
+	virtual UBrushComponent* GenerateMatchingBrushComponent(const FPCGInteractiveToolWorkingDataContext& Context);
+#endif
+
+protected:
+#if WITH_EDITORONLY_DATA
+	UPROPERTY(Transient)
+	TObjectPtr<UBrushComponent> OnToolStartBrushComponent = nullptr;
+
+	UPROPERTY(Transient)
+	TObjectPtr<UBrushBuilder> OnToolStartBrushBuilder = nullptr;
+
+	UPROPERTY(Transient)
+	FTransform OnToolStartVolumeTransform;
+#endif
+};
+
 /** The tool data stored on the graph itself. Informs the UI. */
 USTRUCT(BlueprintType)
 struct FPCGGraphToolData
diff --git a/Engine/Plugins/PCG/Source/PCGEditor/Private/EditorMode/PCGEdMode.cpp b/Engine/Plugins/PCG/Source/PCGEditor/Private/EditorMode/PCGEdMode.cpp
--- a/Engine/Plugins/PCG/Source/PCGEditor/Private/EditorMode/PCGEdMode.cpp
+++ b/Engine/Plugins/PCG/Source/PCGEditor/Private/EditorMode/PCGEdMode.cpp
@@ -5,7 +5,7 @@
 #include "EditorMode/PCGEdModeCommands.h"
 #include "EditorMode/PCGEdModeToolkit.h"
 #include "EditorMode/Tools/Line/PCGDrawSplineTool.h"
-#include "EditorMode/Tools/Volume/PCGCreateVolumeTool.h"
+#include "EditorMode/Tools/Volume/PCGVolumeTool.h"
 
 #include "ContextObjectStore.h"
 #include "EngineAnalytics.h"
@@ -197,7 +197,7 @@ void UPCGEditorMode::Enter()
 
 	for(TSharedPtr<FUICommandInfo> ToolCommandInfo : FPCGEditorModeToolCommands::Get().GetVolumeContextToolCommands())
 	{
-		UPCGCreateVolumeToolBuilder* VolumeToolBuilder = NewObject<UPCGCreateVolumeToolBuilder>();
+		UPCGVolumeToolBuilder* VolumeToolBuilder = NewObject<UPCGVolumeToolBuilder>();
 		VolumeToolBuilder->SetToolGraph(CastChecked<UPCGGraph>(FSoftObjectPath(ToolCommandInfo->GetCommandName().ToString()).TryLoad()));
 		RegisterTool(ToolCommandInfo, ToolCommandInfo->GetCommandName().ToString(), VolumeToolBuilder);
 	}
@@ -215,9 +215,8 @@ void UPCGEditorMode::Enter()
 	UPCGPaintToolBuilder* PaintToolBuilder = NewObject<UPCGPaintToolBuilder>();
 	RegisterTool(ToolManagerCommands.EnablePaintTool, TEXT("Paint"), PaintToolBuilder);
 
-	// @TODO_PCG Reenable Volume Tool when it makes sense
-	// UPCGCreateVolumeToolBuilder* CreateVolumeToolBuilder = NewObject<UPCGCreateVolumeToolBuilder>();
-	// RegisterTool(ToolManagerCommands.EnableCreateVolumeTool, TEXT("CreateVolume"), CreateVolumeToolBuilder);
+	UPCGVolumeToolBuilder* VolumeToolBuilder = NewObject<UPCGVolumeToolBuilder>();
+	RegisterTool(ToolManagerCommands.EnableVolumeTool, TEXT("Volume"), VolumeToolBuilder);
 
 	ToolManager->SelectActiveToolType(EToolSide::Left, "DrawSpline");
 #endif
diff --git a/Engine/Plugins/PCG/Source/PCGEditor/Private/EditorMode/PCGEdModeCommands.cpp b/Engine/Plugins/PCG/Source/PCGEditor/Private/EditorMode/PCGEdModeCommands.cpp
--- a/Engine/Plugins/PCG/Source/PCGEditor/Private/EditorMode/PCGEdModeCommands.cpp
+++ b/Engine/Plugins/PCG/Source/PCGEditor/Private/EditorMode/PCGEdModeCommands.cpp
@@ -9,7 +9,7 @@
 #include "EditorMode/Tools/Helpers/PCGEdModeEditorUtilities.h"
 #include "EditorMode/Tools/Line/PCGDrawSplineTool.h"
 #include "EditorMode/Tools/Paint/PCGPaintTool.h"
-#include "EditorMode/Tools/Volume/PCGCreateVolumeTool.h"
+#include "EditorMode/Tools/Volume/PCGVolumeTool.h"
 #include "Logging/StructuredLog.h"
 #include "Misc/OutputDeviceNull.h"
 
@@ -87,19 +87,19 @@ void FPCGEditorModeToolCommands::RegisterCommands()
 	PCG_TOOL_COMMAND_INFO(EnableDrawSplineTool, "DrawSpline", "Draw Spline", "Use the 'Draw Spline' tool to build a spline for PCG.", "PCGEditorMode.Tools.DrawSpline");
 	PCG_TOOL_COMMAND_INFO(EnableDrawSurfaceTool, "DrawSurface", "Draw Surface", "Use the 'Draw Surface' tool to build a spline for PCG.", "PCGEditorMode.Tools.DrawSurface");
 	PCG_TOOL_COMMAND_INFO(EnablePaintTool, "Paint", "Paint", "Use the 'Paint' tool to apply PCG points with a brush.", "PCGEditorMode.Context.Paint");
-	PCG_TOOL_COMMAND_INFO(EnableCreateVolumeTool, "CreateVolume", "Create Volume", "Create a PCG volume in the scene.", "PCGEditorMode.Tools.CreateVolume");
+	PCG_TOOL_COMMAND_INFO(EnableVolumeTool, "Volume", "Volume", "Create or updates a volume in the scene.", "PCGEditorMode.Tools.Volume");
 }
 
 void FPCGEditorModeToolCommands::RegisterAllToolCommands()
 {
 	FPCGDrawSplineCommands::Register();
 	FPCGPaintCommands::Register();
-	FPCGCreateVolumeCommands::Register();
+	FPCGVolumeCommands::Register();
 }
 
 void FPCGEditorModeToolCommands::UnregisterAllToolCommands()
 {
-	FPCGCreateVolumeCommands::Unregister();
+	FPCGVolumeCommands::Unregister();
 	FPCGPaintCommands::Unregister();
 	FPCGDrawSplineCommands::Unregister();
 }
@@ -132,7 +132,7 @@ TArray<TSharedPtr<FUICommandInfo>> FPCGEditorModeToolCommands::GetVolumeContextT
 #if USE_PCG_MODE_GRAPHS_AS_BUTTONS
 	return PCG::EditorMode::Utility::GetOrCreateUICommandsWithToolTag(AsShared(), "VolumeTool");
 #else
-	return {EnableCreateVolumeTool};
+	return {EnableVolumeTool};
 #endif
 }
 
@@ -156,7 +156,7 @@ ToolCDOs.Add(GetMutableDefault<ToolClassName>());\
 PCG_DEFINE_TOOL_COMMANDS(FPCGDrawSplineCommands, "DrawSpline", "Draw Spline", UPCGDrawSplineToolBase);
 PCG_DEFINE_TOOL_COMMANDS(FPCGDrawSurfaceCommands, "DrawSurface", "Draw Surface", UPCGDrawSplineToolBase);
 PCG_DEFINE_TOOL_COMMANDS(FPCGPaintCommands, "Paint", "Paint", UPCGPaintTool);
-PCG_DEFINE_TOOL_COMMANDS(FPCGCreateVolumeCommands, "CreateVolume", "Create Volume", UPCGCreateVolumeTool);
+PCG_DEFINE_TOOL_COMMANDS(FPCGVolumeCommands, "Volume", "Volume", UPCGVolumeTool);
 
 #undef PCG_DEFINE_TOOL_COMMANDS
 
diff --git a/Engine/Plugins/PCG/Source/PCGEditor/Private/EditorMode/PCGEdModeStyle.cpp b/Engine/Plugins/PCG/Source/PCGEditor/Private/EditorMode/PCGEdModeStyle.cpp
--- a/Engine/Plugins/PCG/Source/PCGEditor/Private/EditorMode/PCGEdModeStyle.cpp
+++ b/Engine/Plugins/PCG/Source/PCGEditor/Private/EditorMode/PCGEdModeStyle.cpp
@@ -73,7 +73,7 @@ FPCGEditorModeStyle::FPCGEditorModeStyle()
 
 		Set("PCGEditorMode.Tools.Paint", new IMAGE_BRUSH_SVG("Icons/EdMode/Paint", Icon20x20));
 
-		Set("PCGEditorMode.Tools.CreateVolume", new IMAGE_BRUSH_SVG("Icons/EdMode/Volume", Icon20x20));
+		Set("PCGEditorMode.Tools.Volume", new IMAGE_BRUSH_SVG("Icons/EdMode/Volume", Icon20x20));
 	}
 
 	// Backdrops and backgrounds
diff --git a/Engine/Plugins/PCG/Source/PCGEditor/Private/EditorMode/Tools/Volume/PCGCreateVolumeTool.cpp b/Engine/Plugins/PCG/Source/PCGEditor/Private/EditorMode/Tools/Volume/PCGCreateVolumeTool.cpp
deleted file mode 100644
diff --git a/Engine/Plugins/PCG/Source/PCGEditor/Private/EditorMode/Tools/Volume/PCGVolumeTool.cpp b/Engine/Plugins/PCG/Source/PCGEditor/Private/EditorMode/Tools/Volume/PCGVolumeTool.cpp
new file mode 100644
--- /dev/null
+++ b/Engine/Plugins/PCG/Source/PCGEditor/Private/EditorMode/Tools/Volume/PCGVolumeTool.cpp
@@ -0,0 +1,312 @@
+// Copyright Epic Games, Inc. All Rights Reserved.
+
+#include "EditorMode/Tools/Volume/PCGVolumeTool.h"
+
+#include "EditorMode/PCGEdModeStyle.h"
+
+#include "ActorFactories/ActorFactory.h"
+#include "BaseBehaviors/SingleClickOrDragBehavior.h"
+#include "BaseBehaviors/MouseHoverBehavior.h"
+#include "Builders/CubeBuilder.h"
+#include "GameFramework/Volume.h"
+#include "InteractiveToolManager.h"
+
+#include UE_INLINE_GENERATED_CPP_BY_NAME(PCGVolumeTool)
+
+#define LOCTEXT_NAMESPACE "UPCGVolumeTool"
+
+namespace PCGVolumeTool::Constants
+{
+	const FText ToolDisplayName = LOCTEXT("ToolDisplayName", "Volume");
+	const FName ToolTag = "VolumeTool";
+}
+
+UPCGInteractiveToolSettings_Volume::UPCGInteractiveToolSettings_Volume()
+{
+	TInstancedStruct<FPCGRaycastFilterRule_Landscape> LandscapeRule = TInstancedStruct<FPCGRaycastFilterRule_Landscape>::Make();
+	TInstancedStruct<FPCGRaycastFilterRule_Meshes> MeshesRule = TInstancedStruct<FPCGRaycastFilterRule_Meshes>::Make();
+	TInstancedStruct<FPCGRaycastFilterRule_IgnorePCGGeneratedComponents> IgnorePCGGeneratedComponentsRule = TInstancedStruct<FPCGRaycastFilterRule_IgnorePCGGeneratedComponents>::Make();
+	RaycastRuleCollection.RaycastRules.Add(LandscapeRule);
+	RaycastRuleCollection.RaycastRules.Add(MeshesRule);
+	RaycastRuleCollection.RaycastRules.Add(IgnorePCGGeneratedComponentsRule);
+}
+
+FName UPCGInteractiveToolSettings_Volume::GetToolTag() const
+{
+	return PCGVolumeTool::Constants::ToolTag;
+}
+
+const UScriptStruct* UPCGInteractiveToolSettings_Volume::GetWorkingDataType() const
+{
+	return FPCGInteractiveToolWorkingData_Volume::StaticStruct();
+}
+
+AVolume* UPCGInteractiveToolSettings_Volume::GetWorkingVolume(FName DataInstanceIdentifier) const
+{
+	if (const FPCGInteractiveToolWorkingData_Volume* ToolData = GetTypedWorkingData<FPCGInteractiveToolWorkingData_Volume>(DataInstanceIdentifier))
+	{
+		return ToolData->Volume.Get();
+	}
+
+	return nullptr;
+}
+
+void UPCGVolumeTool::Setup()
+{
+	Super::Setup();
+
+	DefaultHeight = 100; // @todo_pcg: get from existing volume if any
+
+	ClickOrDragBehavior = NewObject<USingleClickOrDragInputBehavior>();
+	ClickOrDragBehavior->Initialize(this, this);
+	AddInputBehavior(ClickOrDragBehavior);
+
+	HoverBehavior = NewObject<UMouseHoverBehavior>();
+	HoverBehavior->Initialize(this);
+	AddInputBehavior(HoverBehavior);
+	
+	AddToolPropertySource(ToolSettings);
+
+	CachedRayParams =
+	{
+		.World = GetWorld(),
+		.FilterRuleCollection = &ToolSettings->RaycastRuleCollection
+	};
+}
+
+void UPCGVolumeTool::Shutdown(EToolShutdownType ShutdownType)
+{
+	UE::PCG::EditorMode::Tool::Shutdown(this, ShutdownType);
+
+	Super::Shutdown(ShutdownType);
+}
+
+void UPCGVolumeTool::OnTick(const float DeltaTime)
+{
+	// if we are currently dragging, we should update the temporary contents with a fixed Z size
+	// or update w/ live Z coordinates
+	if (bIsDraggingBase || bIsDraggingHeight)
+	{
+		FBox Box = GetCurrentBox();
+
+		if (LastBox.IsValid &&
+			(Box.GetCenter() - LastBox.GetCenter()).SquaredLength() < UE_KINDA_SMALL_NUMBER &&
+			(Box.GetExtent() - LastBox.GetExtent()).SquaredLength() < UE_KINDA_SMALL_NUMBER)
+		{
+			// Nothing to do
+		}
+		else
+		{
+			LastBox = Box;
+			UpdateVolumeFromBox();
+		}
+	}
+
+	Super::OnTick(DeltaTime);
+}
+
+void UPCGVolumeTool::UpdateVolumeFromBox()
+{
+	if (AVolume* WorkingVolume = ToolSettings ? ToolSettings->GetWorkingVolume(ToolSettings->DataInstance) : nullptr)
+	{
+		const FVector BoxCenter = LastBox.GetCenter();
+		const FVector Extent = 2.0 * LastBox.GetExtent();
+
+		// Move actor to bounds center
+		WorkingVolume->SetActorTransform(FTransform(BoxCenter));
+
+		// setup cube builder
+		UCubeBuilder* Builder = NewObject<UCubeBuilder>();
+		Builder->X = Extent.X;
+		Builder->Y = Extent.Y;
+		Builder->Z = Extent.Z;
+		UActorFactory::CreateBrushForVolumeActor(WorkingVolume, Builder);
+	}
+}
+
+FBox UPCGVolumeTool::GetCurrentBox() const
+{
+	FVector BoxCenter;
+	FVector HalfExtents;
+
+	if (bIsDraggingBase)
+	{
+		BoxCenter.X = 0.5 * (FirstPoint.X + InteractiveSecondPoint.X);
+		BoxCenter.Y = 0.5 * (FirstPoint.Y + InteractiveSecondPoint.Y);
+
+		HalfExtents.X = 0.5 * FMath::Abs(FirstPoint.X - InteractiveSecondPoint.X);
+		HalfExtents.Y = 0.5 * FMath::Abs(FirstPoint.Y - InteractiveSecondPoint.Y);
+	}
+	else
+	{
+		BoxCenter.X = 0.5 * (FirstPoint.X + SecondPoint.X);
+		BoxCenter.Y = 0.5 * (FirstPoint.Y + SecondPoint.Y);
+
+		HalfExtents.X = 0.5 * FMath::Abs(FirstPoint.X - SecondPoint.X);
+		HalfExtents.Y = 0.5 * FMath::Abs(FirstPoint.Y - SecondPoint.Y);
+	}
+
+	// Bias needed so the volume ends up covering the touching face.
+	const double Bias = 1.0;
+
+	if (bIsDraggingHeight)
+	{
+		const double MaxRayLength = 1000000.0;
+		FVector PointOnRay, PointOnVertical;
+		FMath::SegmentDistToSegment(LastRay.Origin, LastRay.Origin + LastRay.Direction * MaxRayLength, SecondPoint - 0.5 * MaxRayLength * FVector::ZAxisVector, SecondPoint + 0.5 * MaxRayLength * FVector::ZAxisVector, PointOnRay, PointOnVertical);
+
+		double MinZ, MaxZ;
+		if (PointOnVertical.Z < FirstPoint.Z)
+		{
+			MinZ = PointOnVertical.Z;
+			MaxZ = FirstPoint.Z + Bias;
+		}
+		else
+		{
+			MinZ = FirstPoint.Z - Bias;
+			MaxZ = PointOnVertical.Z;
+		}
+
+		BoxCenter.Z = 0.5 * (MaxZ + MinZ);
+		HalfExtents.Z = 0.5 * (MaxZ - MinZ);
+	}
+	else
+	{
+		BoxCenter.Z = FirstPoint.Z - Bias;
+		HalfExtents.Z = DefaultHeight;
+	}
+
+	return FBox(BoxCenter - HalfExtents, BoxCenter + HalfExtents);
+}
+
+bool UPCGVolumeTool::CanAccept() const
+{
+	AVolume* WorkingVolume = ToolSettings ? ToolSettings->GetWorkingVolume(ToolSettings->DataInstance) : nullptr;
+	// @todo_pcg also make sure we're in a correct state (e.g. not being interacted with at this time).
+	return WorkingVolume != nullptr;
+}
+
+FInputRayHit UPCGVolumeTool::IsHitByClick(const FInputDeviceRay& ClickPosition)
+{
+	// Always accept clicks; invalid clicks when in the base-dragging phase will be discarded by the CanBeginClickDragSequence call.
+	FInputRayHit DummyHit;
+	DummyHit.bHit = true;
+	return DummyHit;
+}
+
+void UPCGVolumeTool::OnClicked(const FInputDeviceRay& ClickPosition)
+{
+	UPCGInteractiveToolSettings_Volume* Settings = CastChecked<UPCGInteractiveToolSettings_Volume>(ToolSettings);
+
+	bIsDraggingBase = false;
+
+	if (bBaseSelectionDone)
+	{
+		LastBox = GetCurrentBox();
+		bBaseSelectionDone = false;
+		bIsDraggingHeight = false;
+		
+		UpdateVolumeFromBox();
+	}
+
+	// @todo_pcg: Different chords to translate, rotate, scale at any step of the process?
+}
+
+FInputRayHit UPCGVolumeTool::CanBeginClickDragSequence(const FInputDeviceRay& DragPosition)
+{
+	// Dragging the base requires a ray to hit the physical scene.
+	using namespace UE::PCG::EditorMode::Scene;
+	if (TOptional<FViewHitResult> Result = ViewportRay(DragPosition, CachedRayParams))
+	{
+		return Result->ToInputRayHit();
+	}
+
+	return {};
+}
+
+void UPCGVolumeTool::OnClickPress(const FInputDeviceRay& ClickPosition)
+{
+	// Restart xy selection
+	bBaseSelectionDone = false;
+	bIsDraggingBase = false;
+
+	using namespace UE::PCG::EditorMode::Scene;
+	if (TOptional<FViewHitResult> Result = ViewportRay(ClickPosition, CachedRayParams))
+	{
+		FirstPoint = Result->ImpactPosition;
+		InteractiveSecondPoint = FirstPoint;
+		bIsDraggingBase = true;
+	}
+}
+
+void UPCGVolumeTool::OnClickDrag(const FInputDeviceRay& DragPosition)
+{
+	// Update interactive second point only.
+	InteractiveSecondPoint = GetBaseDragPosition(DragPosition);
+}
+
+void UPCGVolumeTool::OnClickRelease(const FInputDeviceRay& ReleasePosition)
+{
+	OnClickDrag(ReleasePosition);
+	SecondPoint = GetBaseDragPosition(ReleasePosition);
+	// Hover isn't updated during drag, so we need to update it here, otherwise there will be a visible jerk
+	LastRay = ReleasePosition.WorldRay;
+
+	bIsDraggingBase = false;
+	bBaseSelectionDone = true;
+	bIsDraggingHeight = true;
+	OnTerminateDragSequence();
+}
+
+FVector UPCGVolumeTool::GetBaseDragPosition(const FInputDeviceRay& DragPosition) const
+{
+	// Converts world ray into a projection on the first point plane, with normal z-up.
+	FPlane BasePlane(FirstPoint, FVector::ZAxisVector);
+	return FMath::RayPlaneIntersection(DragPosition.WorldRay.Origin, DragPosition.WorldRay.Direction, BasePlane);
+}
+
+FInputRayHit UPCGVolumeTool::BeginHoverSequenceHitTest(const FInputDeviceRay& PressPos)
+{
+	// This hit will not be consumed, but only used to control when we are in the hover state.
+	FInputRayHit Hit;
+
+	// Hover sequence can only happen when we're done with the base selection
+	if (bBaseSelectionDone)
+	{
+		Hit.bHit = true;
+	}
+	
+	return Hit;
+}
+
+bool UPCGVolumeTool::OnUpdateHover(const FInputDeviceRay& DevicePos)
+{
+	LastRay = DevicePos.WorldRay;
+	return true;
+}
+
+bool UPCGVolumeToolBuilder::CanBuildTool(const FToolBuilderState& SceneState) const
+{
+	return true;
+}
+
+UInteractiveTool* UPCGVolumeToolBuilder::BuildTool(const FToolBuilderState& SceneState) const
+{
+	UPCGVolumeTool* NewTool = NewObject<UPCGVolumeTool>(SceneState.ToolManager);
+
+	NewTool->SetToolInfo(
+	{
+		.ToolDisplayName = PCGVolumeTool::Constants::ToolDisplayName,
+		.ToolDisplayMessage = LOCTEXT("VolumeToolDisplayMessage", "Creates or changes a volume."),
+		.ToolIcon = FPCGEditorModeStyle::Get().GetBrush("PCGEditorMode.Tools.Volume")
+	});
+	
+	if(UE::PCG::EditorMode::Tool::BuildTool(NewTool) == false)
+	{
+		return nullptr;
+	}	
+
+	return NewTool;
+}
+
+#undef LOCTEXT_NAMESPACE
diff --git a/Engine/Plugins/PCG/Source/PCGEditor/Public/EditorMode/PCGEdModeCommands.h b/Engine/Plugins/PCG/Source/PCGEditor/Public/EditorMode/PCGEdModeCommands.h
--- a/Engine/Plugins/PCG/Source/PCGEditor/Public/EditorMode/PCGEdModeCommands.h
+++ b/Engine/Plugins/PCG/Source/PCGEditor/Public/EditorMode/PCGEdModeCommands.h
@@ -69,7 +69,7 @@ public:
 	TArray<TSharedPtr<FUICommandInfo>> GetPaintContextToolCommands();
 
 	/*** PCG EDITOR MODE TOOLS - Volume Context **/
-	TSharedPtr<FUICommandInfo> EnableCreateVolumeTool;
+	TSharedPtr<FUICommandInfo> EnableVolumeTool;
 	TArray<TSharedPtr<FUICommandInfo>> GetVolumeContextToolCommands();
 };
 
@@ -91,6 +91,6 @@ virtual void GetToolDefaultObjectList(TArray<UInteractiveTool*>& ToolCDOs) overr
 DECLARE_TOOL_COMMANDS(FPCGDrawSplineCommands);
 DECLARE_TOOL_COMMANDS(FPCGDrawSurfaceCommands);
 DECLARE_TOOL_COMMANDS(FPCGPaintCommands);
-DECLARE_TOOL_COMMANDS(FPCGCreateVolumeCommands);
+DECLARE_TOOL_COMMANDS(FPCGVolumeCommands);
 
 #undef DECLARE_TOOL_COMMANDS
diff --git a/Engine/Plugins/PCG/Source/PCGEditor/Public/EditorMode/Tools/Volume/PCGCreateVolumeTool.h b/Engine/Plugins/PCG/Source/PCGEditor/Public/EditorMode/Tools/Volume/PCGCreateVolumeTool.h
deleted file mode 100644
diff --git a/Engine/Plugins/PCG/Source/PCGEditor/Public/EditorMode/Tools/Volume/PCGVolumeTool.h b/Engine/Plugins/PCG/Source/PCGEditor/Public/EditorMode/Tools/Volume/PCGVolumeTool.h
new file mode 100644
--- /dev/null
+++ b/Engine/Plugins/PCG/Source/PCGEditor/Public/EditorMode/Tools/Volume/PCGVolumeTool.h
@@ -0,0 +1,113 @@
+// Copyright Epic Games, Inc. All Rights Reserved.
+
+#pragma once
+
+#include "EditorMode/Tools/PCGInteractiveToolBuilder.h"
+#include "EditorMode/Tools/PCGInteractiveToolSettings.h"
+#include "EditorMode/Tools/Helpers/PCGEdModeSceneQueryHelpers.h"
+#include "BaseBehaviors/BehaviorTargetInterfaces.h"
+
+#include "PCGVolumeTool.generated.h"
+
+class USingleClickOrDragInputBehavior;
+class UMouseHoverBehavior;
+
+UCLASS()
+class PCGEDITOR_API UPCGInteractiveToolSettings_Volume : public UPCGInteractiveToolSettings
+{
+	GENERATED_BODY()
+
+public:
+	UPCGInteractiveToolSettings_Volume();
+
+	// ~Begin UPCGInteractiveToolSettings interface
+	virtual FName GetToolTag() const override;
+	virtual const UScriptStruct* GetWorkingDataType() const override;
+	// ~End UPCGInteractiveToolSettings interface
+
+	AVolume* GetWorkingVolume(FName DataInstanceIdentifier) const;
+	// @todo_pcg: Origin modes: Center, Corner
+	// @todo_pcg: Snap to grid
+
+	UPROPERTY(EditAnywhere, Category = Raycast)
+	FPCGRaycastFilterRuleCollection RaycastRuleCollection;
+};
+
+/** 
+ * Tool class to create 3d volumes. Current implementation only creates cubes.
+ *	Click + drag -> creates the base.
+ *	Second click -> creates the height.
+ */
+UCLASS(Transient, MinimalAPI)
+class UPCGVolumeTool : public UInteractiveTool, public IClickBehaviorTarget, public IClickDragBehaviorTarget, public IHoverBehaviorTarget
+{
+	GENERATED_BODY()
+
+public:
+	// ~Begin UInteractiveTool interface
+	virtual bool HasCancel() const override { return true; }
+	virtual bool HasAccept() const override { return true; }
+
+	virtual void Setup() override;
+	virtual void Shutdown(EToolShutdownType ShutdownType) override;
+	virtual void OnTick(float DeltaTime) override;
+	virtual bool CanAccept() const override;
+	// ~End UInteractiveTool interface
+
+	// ~Begin IClickBehaviorTarget interface
+	virtual FInputRayHit IsHitByClick(const FInputDeviceRay& ClickPosition) override;
+	virtual void OnClicked(const FInputDeviceRay& ClickPosition) override;
+	// ~End IClickBehaviorTarget interface
+
+	// ~Begin IClickDragBehaviorTarget interface
+	virtual FInputRayHit CanBeginClickDragSequence(const FInputDeviceRay& DragPosition) override;
+	virtual void OnClickPress(const FInputDeviceRay& ClickPosition) override;
+	virtual void OnClickDrag(const FInputDeviceRay& DragPosition) override;
+	virtual void OnClickRelease(const FInputDeviceRay& ReleasePosition) override;
+	virtual void OnTerminateDragSequence() override {}
+	// ~End IClickDragBehaviorTarget interface
+
+	// ~Begin IHoverBehaviorTarget interface
+	virtual FInputRayHit BeginHoverSequenceHitTest(const FInputDeviceRay& PressPos) override;
+	virtual void OnBeginHover(const FInputDeviceRay& DevicePos) override {}
+	virtual bool OnUpdateHover(const FInputDeviceRay& DevicePos) override;
+	virtual void OnEndHover() override {}
+	// ~End IHoverBehaviorTarget interface
+
+private:
+	FBox GetCurrentBox() const;
+	void UpdateVolumeFromBox();
+	FVector GetBaseDragPosition(const FInputDeviceRay& DragPosition) const;
+
+	UE::PCG::EditorMode::Scene::FViewRayParams CachedRayParams;
+
+	UPROPERTY(Instanced)
+	TObjectPtr<USingleClickOrDragInputBehavior> ClickOrDragBehavior = nullptr;
+
+	UPROPERTY(Instanced)
+	TObjectPtr<UMouseHoverBehavior> HoverBehavior = nullptr;
+
+	UPROPERTY(Transient)
+	TObjectPtr<UPCGInteractiveToolSettings_Volume> ToolSettings;
+
+	bool bIsDraggingBase = false;
+	bool bIsDraggingHeight = false;
+	bool bBaseSelectionDone = false;
+	FVector FirstPoint = FVector::ZeroVector;
+	FVector SecondPoint = FVector::ZeroVector;
+	FVector InteractiveSecondPoint = FVector::ZeroVector;
+	double DefaultHeight = 0;
+
+	FRay LastRay;
+	FBox LastBox = FBox(EForceInit::ForceInit);
+};
+
+UCLASS(Transient, MinimalAPI)
+class UPCGVolumeToolBuilder : public UPCGInteractiveToolBuilder
+{
+	GENERATED_BODY()
+
+public:
+	virtual bool CanBuildTool(const FToolBuilderState& SceneState) const override;
+	virtual UInteractiveTool* BuildTool(const FToolBuilderState& SceneState) const override;
+};
