Small PCG Editor Mode Fixups

- Made tool data editanywhere so users can delete it manually
- cleared tool notifications on tool start
- generated pcg component schedules cleanup and deletion now; this fixes generated mesh components not being deleted properly when hitting Cancel

diff --git a/Engine/Plugins/PCG/Source/PCG/Public/PCGComponent.h b/Engine/Plugins/PCG/Source/PCG/Public/PCGComponent.h
--- a/Engine/Plugins/PCG/Source/PCG/Public/PCGComponent.h
+++ b/Engine/Plugins/PCG/Source/PCG/Public/PCGComponent.h
@@ -365,7 +365,7 @@ public:
 	TObjectPtr<UPCGSchedulingPolicyBase> SchedulingPolicy;
 
 	/** This stores working data of a tool; runtime property as the PCG Component will rely on the tool's working data to generate properly. */
-	UPROPERTY(BlueprintReadOnly, VisibleAnywhere, Category = "Tool Data")
+	UPROPERTY(BlueprintReadOnly, EditAnywhere, Category = "Tool Data")
 	FPCGInteractiveToolWorkingDataContainer ToolDataContainer;
 	
 #if WITH_EDITORONLY_DATA
diff --git a/Engine/Plugins/PCG/Source/PCGEditor/Private/EditorMode/PCGEdModeToolkit.cpp b/Engine/Plugins/PCG/Source/PCGEditor/Private/EditorMode/PCGEdModeToolkit.cpp
--- a/Engine/Plugins/PCG/Source/PCGEditor/Private/EditorMode/PCGEdModeToolkit.cpp
+++ b/Engine/Plugins/PCG/Source/PCGEditor/Private/EditorMode/PCGEdModeToolkit.cpp
@@ -93,6 +93,9 @@ FText FPCGEditorModeToolkit::GetActiveToolMessage() const
 
 void FPCGEditorModeToolkit::OnToolStarted(UInteractiveToolManager* Manager, UInteractiveTool* Tool)
 {
+	ClearToolNotification();
+	ClearToolWarning();
+	
 	UpdateActiveTool();
 
 	if (IsHosted())
diff --git a/Engine/Plugins/PCG/Source/PCGEditor/Private/EditorMode/Tools/PCGInteractiveToolSettings.cpp b/Engine/Plugins/PCG/Source/PCGEditor/Private/EditorMode/Tools/PCGInteractiveToolSettings.cpp
--- a/Engine/Plugins/PCG/Source/PCGEditor/Private/EditorMode/Tools/PCGInteractiveToolSettings.cpp
+++ b/Engine/Plugins/PCG/Source/PCGEditor/Private/EditorMode/Tools/PCGInteractiveToolSettings.cpp
@@ -332,8 +332,16 @@ void UPCGInteractiveToolSettings::DeleteGeneratedResources() const
 	
 	if(GeneratedResources.GeneratedPCGComponent.IsValid())
 	{
-		GeneratedResources.GeneratedPCGComponent->Cleanup(true);
-		GeneratedResources.GeneratedPCGComponent->DestroyComponent();
+		// We schedule a cleanup and deletion
+		FPCGTaskId CleanupTaskId = GeneratedResources.GeneratedPCGComponent->GetSubsystem()->ScheduleCleanup(GeneratedResources.GeneratedPCGComponent.Get(), true, {});
+		GeneratedResources.GeneratedPCGComponent->GetSubsystem()->ScheduleGeneric([PCGComponent = GeneratedResources.GeneratedPCGComponent]() -> bool
+		{
+			PCGComponent->DestroyComponent();
+			return true;
+		},
+		GeneratedResources.GeneratedPCGComponent.Get(),
+		{CleanupTaskId});
+		
 		GeneratedResources.GeneratedPCGComponent.Reset();
 	}
 
