[PCG] Fix Pcg Grid node are not disabled when hierachical genration is disabled.

#rnx

diff --git a/Engine/Plugins/PCG/Source/PCGEditor/Private/Nodes/PCGEditorGraphNodeBase.cpp b/Engine/Plugins/PCG/Source/PCGEditor/Private/Nodes/PCGEditorGraphNodeBase.cpp
--- a/Engine/Plugins/PCG/Source/PCGEditor/Private/Nodes/PCGEditorGraphNodeBase.cpp
+++ b/Engine/Plugins/PCG/Source/PCGEditor/Private/Nodes/PCGEditorGraphNodeBase.cpp
@@ -697,19 +697,16 @@ EPCGChangeType UPCGEditorGraphNodeBase::UpdateErrorsAndWarnings()
 EPCGChangeType UPCGEditorGraphNodeBase::UpdateStructuralVisualization(const IPCGGraphExecutionSource* InSourceBeingDebugged, const FPCGStack* InStackBeingInspected, bool bNewlyPlaced)
 {
 	const UPCGGraph* Graph = PCGNode ? PCGNode->GetGraph() : nullptr;
-	if (!Graph || !InSourceBeingDebugged)
+	if (!Graph)
 	{
 		return EPCGChangeType::None;
 	}
 
-	const UPCGComponent* ComponentBeingDebugged = Cast<const UPCGComponent>(InSourceBeingDebugged);
-
-	const FPCGGraphExecutionInspection& Inspection = InSourceBeingDebugged->GetExecutionState().GetInspection();
-	const bool bInspecting = InStackBeingInspected && !InStackBeingInspected->GetStackFrames().IsEmpty();
+	const bool bInspecting = InSourceBeingDebugged && InStackBeingInspected && !InStackBeingInspected->GetStackFrames().IsEmpty();
 
 	EPCGChangeType ChangeType = EPCGChangeType::None;
 
-	const uint64 NewInactiveMask = bInspecting ? Inspection.GetNodeInactivePinMask(PCGNode, *InStackBeingInspected) : 0;
+	const uint64 NewInactiveMask = bInspecting ? InSourceBeingDebugged->GetExecutionState().GetInspection().GetNodeInactivePinMask(PCGNode, *InStackBeingInspected) : 0;
 	if (NewInactiveMask != InactiveOutputPinMask)
 	{
 		InactiveOutputPinMask = NewInactiveMask;
@@ -720,6 +717,7 @@ EPCGChangeType UPCGEditorGraphNodeBase::UpdateStructuralVisualization(const IPCG
 	const UPCGGraph* TopGraph = bInspecting ? InStackBeingInspected->GetRootGraph() : Graph;
 	const bool HiGenEnabled = TopGraph && TopGraph->IsHierarchicalGenerationEnabled();
 
+	const UPCGComponent* ComponentBeingDebugged = Cast<const UPCGComponent>(InSourceBeingDebugged);
 	// Set the inspected grid size - this is used for grid size visualization.
 	uint32 InspectingGridSize = PCGHiGenGrid::UninitializedGridSize();
 	EPCGHiGenGrid InspectingGrid = EPCGHiGenGrid::Uninitialized;
