[PCG] Enabled pcg editor mode toolkit by default.

#rnx

diff --git a/Engine/Plugins/PCG/Source/PCGEditor/Public/EditorMode/PCGEdModeSettings.h b/Engine/Plugins/PCG/Source/PCGEditor/Public/EditorMode/PCGEdModeSettings.h
--- a/Engine/Plugins/PCG/Source/PCGEditor/Public/EditorMode/PCGEdModeSettings.h
+++ b/Engine/Plugins/PCG/Source/PCGEditor/Public/EditorMode/PCGEdModeSettings.h
@@ -32,7 +32,6 @@ class UPCGEditorModeSettings : public UDeveloperSettings
 public:
 	UPCGEditorModeSettings(const FObjectInitializer& ObjectInitializer)
 		: UDeveloperSettings(ObjectInitializer)
-		, bEnableEditorMode(false) // @todo_pcg: Added to the constructor to avoid false as default. Remove when permanently enabling toolkit.
 	{}
 
 	virtual FName GetContainerName() const override { return FName("Editor"); }
