[PCG] Finalize compute kernel API

#rnx
#FailedValidations Preflight

diff --git a/Engine/Plugins/PCG/Source/PCG/Private/Compute/BuiltInKernels/PCGCountUniqueAttributeValuesKernel.cpp b/Engine/Plugins/PCG/Source/PCG/Private/Compute/BuiltInKernels/PCGCountUniqueAttributeValuesKernel.cpp
--- a/Engine/Plugins/PCG/Source/PCG/Private/Compute/BuiltInKernels/PCGCountUniqueAttributeValuesKernel.cpp
+++ b/Engine/Plugins/PCG/Source/PCG/Private/Compute/BuiltInKernels/PCGCountUniqueAttributeValuesKernel.cpp
@@ -98,7 +98,7 @@ TSharedPtr<const FPCGDataCollectionDesc> UPCGCountUniqueAttributeValuesKernel::C
 		return nullptr;
 	}
 
-	const FPCGKernelPin InputKernelPin(KernelIndex, PCGPinConstants::DefaultInputLabel, /*bIsInput=*/true);
+	const FPCGKernelPin InputKernelPin(GetKernelIndex(), PCGPinConstants::DefaultInputLabel, /*bIsInput=*/true);
 	const TSharedPtr<const FPCGDataCollectionDesc> InputDesc = InBinding->ComputeKernelPinDataDesc(InputKernelPin);
 
 	if (!ensure(InputDesc))
diff --git a/Engine/Plugins/PCG/Source/PCG/Private/Compute/Elements/PCGCullPointsOutsideActorBoundsKernel.cpp b/Engine/Plugins/PCG/Source/PCG/Private/Compute/Elements/PCGCullPointsOutsideActorBoundsKernel.cpp
--- a/Engine/Plugins/PCG/Source/PCG/Private/Compute/Elements/PCGCullPointsOutsideActorBoundsKernel.cpp
+++ b/Engine/Plugins/PCG/Source/PCG/Private/Compute/Elements/PCGCullPointsOutsideActorBoundsKernel.cpp
@@ -24,7 +24,7 @@ TSharedPtr<const FPCGDataCollectionDesc> UPCGCullPointsOutsideActorBoundsKernel:
 		return nullptr;
 	}
 
-	const FPCGKernelPin SourceKernelPin(KernelIndex, PCGPinConstants::DefaultInputLabel, /*bIsInput=*/true);
+	const FPCGKernelPin SourceKernelPin(GetKernelIndex(), PCGPinConstants::DefaultInputLabel, /*bIsInput=*/true);
 	return InBinding->ComputeKernelPinDataDesc(SourceKernelPin);
 }
 
diff --git a/Engine/Plugins/PCG/Source/PCG/Private/Compute/Elements/PCGCustomHLSLKernel.cpp b/Engine/Plugins/PCG/Source/PCG/Private/Compute/Elements/PCGCustomHLSLKernel.cpp
--- a/Engine/Plugins/PCG/Source/PCG/Private/Compute/Elements/PCGCustomHLSLKernel.cpp
+++ b/Engine/Plugins/PCG/Source/PCG/Private/Compute/Elements/PCGCustomHLSLKernel.cpp
@@ -236,7 +236,7 @@ TSharedPtr<const FPCGDataCollectionDesc> UPCGCustomHLSLKernel::ComputeOutputBind
 	{
 		if (const FPCGPinProperties* FirstInputPinProps = GetFirstInputPin())
 		{
-			const FPCGKernelPin FirstKernelPin(KernelIndex, FirstInputPinProps->Label, /*bIsInput=*/true);
+			const FPCGKernelPin FirstKernelPin(GetKernelIndex(), FirstInputPinProps->Label, /*bIsInput=*/true);
 			const TSharedPtr<const FPCGDataCollectionDesc> InputDataDesc = InBinding->ComputeKernelPinDataDesc(FirstKernelPin);
 
 			if (ensure(InputDataDesc))
@@ -319,7 +319,7 @@ TSharedPtr<const FPCGDataCollectionDesc> UPCGCustomHLSLKernel::ComputeOutputBind
 		// Collect descriptions of input data items that have string key attributes.
 		for(const FPCGPinProperties& PinProps : CustomHLSLSettings->InputPins)
 		{
-			const FPCGKernelPin InputKernelPin(KernelIndex, PinProps.Label, /*bIsInput=*/true);
+			const FPCGKernelPin InputKernelPin(GetKernelIndex(), PinProps.Label, /*bIsInput=*/true);
 			const TSharedPtr<const FPCGDataCollectionDesc> InputPinDesc = InBinding->ComputeKernelPinDataDesc(InputKernelPin);
 
 			if (!ensure(InputPinDesc))
@@ -1198,31 +1198,31 @@ bool UPCGCustomHLSLKernel::PerformStaticValidation()
 	if (CustomHLSLSettings->OutputPins.IsEmpty())
 	{
 #if PCG_KERNEL_LOGGING_ENABLED
-		StaticLogEntries.Emplace(LOCTEXT("NoOutputs", "Custom HLSL nodes must have at least one output."), EPCGKernelLogVerbosity::Error);
+		AddStaticLogEntry(LOCTEXT("NoOutputs", "Custom HLSL nodes must have at least one output."), EPCGKernelLogVerbosity::Error);
 #endif
 		return false;
 	}
 
-	auto CheckPinLabel = [Settings=CustomHLSLSettings, &StaticLogEntries=StaticLogEntries](FName PinLabel)
+	auto CheckPinLabel = [this, Settings=CustomHLSLSettings](FName PinLabel)
 	{
 		if (PinLabel == NAME_None)
 		{
 #if PCG_KERNEL_LOGGING_ENABLED
-			StaticLogEntries.Emplace(LOCTEXT("InvalidPinLabelNone", "Pin label 'None' is not a valid pin label."), EPCGKernelLogVerbosity::Error);
+			AddStaticLogEntry(LOCTEXT("InvalidPinLabelNone", "Pin label 'None' is not a valid pin label."), EPCGKernelLogVerbosity::Error);
 #endif
 			return false;
 		}
 
 		bool bFoundPinLabel = false;
 
-		auto IsAlreadyFound = [Settings, &StaticLogEntries, PinLabel, &bFoundPinLabel](const FPCGPinProperties PinProps)
+		auto IsAlreadyFound = [this, Settings, PinLabel, &bFoundPinLabel](const FPCGPinProperties PinProps)
 		{
 			if (PinProps.Label == PinLabel)
 			{
 				if (bFoundPinLabel)
 				{
 #if PCG_KERNEL_LOGGING_ENABLED
-					StaticLogEntries.Emplace(FText::Format(LOCTEXT("DuplicatedPinLabels", "Duplicate pin label '{0}', all labels must be unique."), FText::FromName(PinLabel)), EPCGKernelLogVerbosity::Error);
+					AddStaticLogEntry(FText::Format(LOCTEXT("DuplicatedPinLabels", "Duplicate pin label '{0}', all labels must be unique."), FText::FromName(PinLabel)), EPCGKernelLogVerbosity::Error);
 #endif
 					return true;
 				}
@@ -1266,7 +1266,7 @@ bool UPCGCustomHLSLKernel::PerformStaticValidation()
 			if (Properties.AllowedTypes != EPCGDataType::Point)
 			{
 #if PCG_KERNEL_LOGGING_ENABLED
-				StaticLogEntries.Emplace(FText::Format(
+				AddStaticLogEntry(FText::Format(
 					LOCTEXT("InvalidNonPointPrimaryInput", "'Point Processor' nodes require primary input pin to be of type 'Point', but found '{0}'."),
 					FText::FromString(PCGCustomHLSLKernel::GetDataTypeString(Properties.AllowedTypes))),
 					EPCGKernelLogVerbosity::Error);
@@ -1279,7 +1279,7 @@ bool UPCGCustomHLSLKernel::PerformStaticValidation()
 			if (!(Properties.AllowedTypes & EPCGDataType::BaseTexture))
 			{
 #if PCG_KERNEL_LOGGING_ENABLED
-				StaticLogEntries.Emplace(FText::Format(
+				AddStaticLogEntry(FText::Format(
 					LOCTEXT("InvalidNonTexturePrimaryInput", "'Texture Processor' nodes require primary input pin to be of type 'Base Texture', but found '{0}'."),
 					FText::FromString(PCGCustomHLSLKernel::GetDataTypeString(Properties.AllowedTypes))),
 					EPCGKernelLogVerbosity::Error);
@@ -1292,7 +1292,7 @@ bool UPCGCustomHLSLKernel::PerformStaticValidation()
 			if (!(Properties.AllowedTypes & EPCGDataType::Param))
 			{
 #if PCG_KERNEL_LOGGING_ENABLED
-				StaticLogEntries.Emplace(FText::Format(
+				AddStaticLogEntry(FText::Format(
 					LOCTEXT("InvalidNonAttributeSetPrimaryInput", "'Attribute Set Processor' nodes require primary input pin to be of type 'Attribute Set', but found '{0}'."),
 					FText::FromString(PCGCustomHLSLKernel::GetDataTypeString(Properties.AllowedTypes))),
 					EPCGKernelLogVerbosity::Error);
@@ -1304,7 +1304,7 @@ bool UPCGCustomHLSLKernel::PerformStaticValidation()
 		if (!PCGComputeHelpers::IsTypeAllowedAsInput(Properties.AllowedTypes))
 		{
 #if PCG_KERNEL_LOGGING_ENABLED
-			StaticLogEntries.Emplace(FText::Format(
+			AddStaticLogEntry(FText::Format(
 				LOCTEXT("InvalidInputType", "Unsupported input type '{0}', found on pin '{1}'."),
 				FText::FromString(PCGCustomHLSLKernel::GetDataTypeString(Properties.AllowedTypes)),
 				FText::FromName(Properties.Label)),
@@ -1332,7 +1332,7 @@ bool UPCGCustomHLSLKernel::PerformStaticValidation()
 			if (Properties.AllowedTypes != EPCGDataType::Point)
 			{
 #if PCG_KERNEL_LOGGING_ENABLED
-				StaticLogEntries.Emplace(FText::Format(
+				AddStaticLogEntry(FText::Format(
 					LOCTEXT("InvalidNonPointPrimaryOutput", "'Point Processor' and 'Point Generator' nodes require primary output pin to be of type 'Point', but found '{0}'."),
 					FText::FromString(PCGCustomHLSLKernel::GetDataTypeString(Properties.AllowedTypes))),
 					EPCGKernelLogVerbosity::Error);
@@ -1345,7 +1345,7 @@ bool UPCGCustomHLSLKernel::PerformStaticValidation()
 			if (!(Properties.AllowedTypes & EPCGDataType::BaseTexture))
 			{
 #if PCG_KERNEL_LOGGING_ENABLED
-				StaticLogEntries.Emplace(FText::Format(
+				AddStaticLogEntry(FText::Format(
 					LOCTEXT("InvalidNonTexturePrimaryOutput", "'Texture Processor' and 'Texture Generator' nodes require primary output pin to be of type 'Base Texture', but found '{0}'."),
 					FText::FromString(PCGCustomHLSLKernel::GetDataTypeString(Properties.AllowedTypes))),
 					EPCGKernelLogVerbosity::Error);
@@ -1358,7 +1358,7 @@ bool UPCGCustomHLSLKernel::PerformStaticValidation()
 			if (!(Properties.AllowedTypes & EPCGDataType::Param))
 			{
 #if PCG_KERNEL_LOGGING_ENABLED
-				StaticLogEntries.Emplace(FText::Format(
+				AddStaticLogEntry(FText::Format(
 					LOCTEXT("InvalidNonAttributeSetPrimaryOutput", "'Attribute Set Processor' nodes require primary output pin to be of type 'Attribute Set', but found '{0}'."),
 					FText::FromString(PCGCustomHLSLKernel::GetDataTypeString(Properties.AllowedTypes))),
 					EPCGKernelLogVerbosity::Error);
@@ -1370,7 +1370,7 @@ bool UPCGCustomHLSLKernel::PerformStaticValidation()
 		if (!PCGComputeHelpers::IsTypeAllowedAsOutput(Properties.AllowedTypes))
 		{
 #if PCG_KERNEL_LOGGING_ENABLED
-			StaticLogEntries.Emplace(FText::Format(
+			AddStaticLogEntry(FText::Format(
 				LOCTEXT("InvalidOutputType", "Unsupported output type '{0}', found on pin '{1}'."),
 				FText::FromString(PCGCustomHLSLKernel::GetDataTypeString(Properties.AllowedTypes)),
 				FText::FromName(Properties.Label)),
@@ -1388,7 +1388,7 @@ bool UPCGCustomHLSLKernel::PerformStaticValidation()
 				if (Props.PinsToInititalizeFrom.IsEmpty())
 				{
 #if PCG_KERNEL_LOGGING_ENABLED
-					StaticLogEntries.Emplace(FText::Format(
+					AddStaticLogEntry(FText::Format(
 						LOCTEXT("InitFromEmptyPins", "Output pin '{0}' tried to initialize from input pins, but no pins were specified."),
 						FText::FromName(Properties.Label)),
 						EPCGKernelLogVerbosity::Error);
@@ -1408,7 +1408,7 @@ bool UPCGCustomHLSLKernel::PerformStaticValidation()
 						if (!PCGComputeHelpers::IsTypeAllowedAsOutput(InitPinProps->AllowedTypes))
 						{
 #if PCG_KERNEL_LOGGING_ENABLED
-							StaticLogEntries.Emplace(FText::Format(
+							AddStaticLogEntry(FText::Format(
 								LOCTEXT("InitFromInvalidPinType", "Output pin '{0}' tried to initialize from input pin '{1}', but pin '{1}' has an invalid type."),
 								FText::FromName(Properties.Label),
 								FText::FromName(InitPinName)),
@@ -1420,7 +1420,7 @@ bool UPCGCustomHLSLKernel::PerformStaticValidation()
 					else
 					{
 #if PCG_KERNEL_LOGGING_ENABLED
-						StaticLogEntries.Emplace(FText::Format(
+						AddStaticLogEntry(FText::Format(
 							LOCTEXT("InitFromNonExistentPin", "Output pin '{0}' tried to initialize from non-existent input pin '{1}'."),
 							FText::FromName(Properties.Label),
 							FText::FromName(InitPinName)),
@@ -1440,7 +1440,7 @@ bool UPCGCustomHLSLKernel::PerformStaticValidation()
 				if (Props.DataCount < 1)
 				{
 #if PCG_KERNEL_LOGGING_ENABLED
-					StaticLogEntries.Emplace(FText::Format(
+					AddStaticLogEntry(FText::Format(
 						LOCTEXT("InvalidDataCount", "Invalid fixed data count {0} on output pin '{1}'. Must be greater than 0."),
 						FText::AsNumber(Props.DataCount),
 						FText::FromName(Properties.Label)),
@@ -1457,7 +1457,7 @@ bool UPCGCustomHLSLKernel::PerformStaticValidation()
 				if (Props.ElementCount < 1)
 				{
 #if PCG_KERNEL_LOGGING_ENABLED
-					StaticLogEntries.Emplace(FText::Format(
+					AddStaticLogEntry(FText::Format(
 						LOCTEXT("InvalidElementCount", "Invalid fixed num elements {0} on output pin '{1}'. Must be greater than 0."),
 						FText::AsNumber(Props.ElementCount),
 						FText::FromName(Properties.Label)),
@@ -1469,7 +1469,7 @@ bool UPCGCustomHLSLKernel::PerformStaticValidation()
 				if (Props.NumElements2D.GetMin() < 1)
 				{
 #if PCG_KERNEL_LOGGING_ENABLED
-					StaticLogEntries.Emplace(FText::Format(
+					AddStaticLogEntry(FText::Format(
 						LOCTEXT("InvalidElementCount2D", "Invalid fixed num elements ({0}, {1}) on output pin '{2}'. Must be greater than 0."),
 						FText::AsNumber(Props.NumElements2D.X),
 						FText::AsNumber(Props.NumElements2D.Y),
@@ -1483,7 +1483,7 @@ bool UPCGCustomHLSLKernel::PerformStaticValidation()
 			if (Props.ElementCountMultiplier < 1)
 			{
 #if PCG_KERNEL_LOGGING_ENABLED
-				StaticLogEntries.Emplace(FText::Format(
+				AddStaticLogEntry(FText::Format(
 					LOCTEXT("InvalidElementCountMultiplier", "Invalid element count multiplier {0} on output pin '{1}'. Must be greater than 0."),
 					FText::AsNumber(Props.ElementCountMultiplier),
 					FText::FromName(Properties.Label)),
@@ -1501,7 +1501,7 @@ bool UPCGCustomHLSLKernel::PerformStaticValidation()
 		if (CustomHLSLSettings->ThreadCountInputPinLabels.IsEmpty())
 		{
 #if PCG_KERNEL_LOGGING_ENABLED
-			StaticLogEntries.Emplace(LOCTEXT("MissingThreadCountPins", "Dispatch thread count is based on input pins but no labels have been set in Input Pins array."), EPCGKernelLogVerbosity::Error);
+			AddStaticLogEntry(LOCTEXT("MissingThreadCountPins", "Dispatch thread count is based on input pins but no labels have been set in Input Pins array."), EPCGKernelLogVerbosity::Error);
 #endif
 			return false;
 		}
@@ -1511,7 +1511,7 @@ bool UPCGCustomHLSLKernel::PerformStaticValidation()
 			if (!CustomHLSLSettings->InputPins.FindByPredicate([Label](const FPCGPinProperties& InProps) { return InProps.Label == Label; }))
 			{
 #if PCG_KERNEL_LOGGING_ENABLED
-				StaticLogEntries.Emplace(FText::Format(LOCTEXT("MissingThreadCountPin", "Invalid pin specified in Input Pins array: '{0}'."), FText::FromName(Label)), EPCGKernelLogVerbosity::Error);
+				AddStaticLogEntry(FText::Format(LOCTEXT("MissingThreadCountPin", "Invalid pin specified in Input Pins array: '{0}'."), FText::FromName(Label)), EPCGKernelLogVerbosity::Error);
 #endif
 				return false;
 			}
@@ -1523,7 +1523,7 @@ bool UPCGCustomHLSLKernel::PerformStaticValidation()
 		if (CustomHLSLSettings->ThreadCountMultiplier < 1)
 		{
 #if PCG_KERNEL_LOGGING_ENABLED
-			StaticLogEntries.Emplace(FText::Format(LOCTEXT("InvalidThreadCountMultiplier", "Thread Count Multiplier has invalid value ({0}). Must be greater than 0."), CustomHLSLSettings->ThreadCountMultiplier), EPCGKernelLogVerbosity::Error);
+			AddStaticLogEntry(FText::Format(LOCTEXT("InvalidThreadCountMultiplier", "Thread Count Multiplier has invalid value ({0}). Must be greater than 0."), CustomHLSLSettings->ThreadCountMultiplier), EPCGKernelLogVerbosity::Error);
 #endif
 			return false;
 		}
@@ -1537,7 +1537,7 @@ bool UPCGCustomHLSLKernel::PerformStaticValidation()
 			check(AttributeTypeEnum);
 
 #if PCG_KERNEL_LOGGING_ENABLED
-			StaticLogEntries.Emplace(FText::Format(
+			AddStaticLogEntry(FText::Format(
 				LOCTEXT("InvalidAttributeTypeNone", "Attribute '{0}' has invalid GPU attribute type '{1}', check the 'Attributes to Create' array on your pins."),
 				FText::FromName(AttributeKey.Identifier.Name),
 				FText::FromString(AttributeTypeEnum->GetNameStringByValue(static_cast<int64>(AttributeKey.Type)))),
@@ -1547,7 +1547,7 @@ bool UPCGCustomHLSLKernel::PerformStaticValidation()
 		}
 	}
 
-	if (!ValidateShaderSource(StaticLogEntries))
+	if (!ValidateShaderSource())
 	{
 		return false;
 	}
@@ -1555,11 +1555,11 @@ bool UPCGCustomHLSLKernel::PerformStaticValidation()
 	return true;
 }
 
-bool UPCGCustomHLSLKernel::ValidateShaderSource(TArray<FPCGKernelLogEntry>& OutLogEntries) const
+bool UPCGCustomHLSLKernel::ValidateShaderSource()
 {
 	const UPCGCustomHLSLSettings* CustomHLSLSettings = CastChecked<UPCGCustomHLSLSettings>(GetSettings());
 
-	if (!CustomHLSLSettings->bMuteUnwrittenPinDataErrors && !AreAllOutputPinsWritten(OutLogEntries))
+	if (!CustomHLSLSettings->bMuteUnwrittenPinDataErrors && !AreAllOutputPinsWritten())
 	{
 		return false;
 	}
@@ -1683,7 +1683,7 @@ FString UPCGCustomHLSLKernel::ProcessAdditionalShaderSources(FPCGGPUCompilationC
 	return OutShaderSource;
 }
 
-bool UPCGCustomHLSLKernel::AreAllOutputPinsWritten(TArray<FPCGKernelLogEntry>& OutLogEntries) const
+bool UPCGCustomHLSLKernel::AreAllOutputPinsWritten()
 {
 	TRACE_CPUPROFILER_EVENT_SCOPE(UPCGCustomHLSLKernel::AreAllOutputPinsWritten);
 
@@ -1720,7 +1720,7 @@ bool UPCGCustomHLSLKernel::AreAllOutputPinsWritten(TArray<FPCGKernelLogEntry>& O
 		if (!bInitializedByAnySource)
 		{
 #if PCG_KERNEL_LOGGING_ENABLED
-			OutLogEntries.Emplace(FText::Format(
+			AddStaticLogEntry(FText::Format(
 				LOCTEXT("PinMayNotBeWritten", "Data on pin '{0}' may be uninitialized. Add code to write to this data, or mute this error in the node settings."),
 				{ FText::FromString(PinStr) }),
 				EPCGKernelLogVerbosity::Error);
diff --git a/Engine/Plugins/PCG/Source/PCG/Private/Compute/Elements/PCGCustomHLSLKernel.h b/Engine/Plugins/PCG/Source/PCG/Private/Compute/Elements/PCGCustomHLSLKernel.h
--- a/Engine/Plugins/PCG/Source/PCG/Private/Compute/Elements/PCGCustomHLSLKernel.h
+++ b/Engine/Plugins/PCG/Source/PCG/Private/Compute/Elements/PCGCustomHLSLKernel.h
@@ -140,8 +140,8 @@ protected:
 	void ParseShaderSource();
 	void CreateParsedSources();
 	void CollectDataLabels(const FPCGCustomHLSLParsedSource& InParsedSource);
-	bool ValidateShaderSource(TArray<FPCGKernelLogEntry>& OutLogEntries) const;
-	bool AreAllOutputPinsWritten(TArray<FPCGKernelLogEntry>& OutLogEntries) const;
+	bool ValidateShaderSource();
+	bool AreAllOutputPinsWritten();
 
 	/** Apply any code gen/transformations/expansions to the shader source. */
 	virtual FString ProcessShaderSource(FPCGGPUCompilationContext& InOutContext, const FPCGCustomHLSLParsedSource& InParsedSource) const;
diff --git a/Engine/Plugins/PCG/Source/PCG/Private/Compute/Elements/PCGNormalToDensityKernel.cpp b/Engine/Plugins/PCG/Source/PCG/Private/Compute/Elements/PCGNormalToDensityKernel.cpp
--- a/Engine/Plugins/PCG/Source/PCG/Private/Compute/Elements/PCGNormalToDensityKernel.cpp
+++ b/Engine/Plugins/PCG/Source/PCG/Private/Compute/Elements/PCGNormalToDensityKernel.cpp
@@ -23,7 +23,7 @@ TSharedPtr<const FPCGDataCollectionDesc> UPCGNormalToDensityKernel::ComputeOutpu
 		return nullptr;
 	}
 
-	const FPCGKernelPin SourceKernelPin(KernelIndex, PCGPinConstants::DefaultInputLabel, /*bIsInput=*/true);
+	const FPCGKernelPin SourceKernelPin(GetKernelIndex(), PCGPinConstants::DefaultInputLabel, /*bIsInput=*/true);
 	const TSharedPtr<const FPCGDataCollectionDesc> InDataDesc = InBinding->ComputeKernelPinDataDesc(SourceKernelPin);
 
 	if (!ensure(InDataDesc))
diff --git a/Engine/Plugins/PCG/Source/PCG/Private/Compute/Elements/PCGTransformPointsKernel.cpp b/Engine/Plugins/PCG/Source/PCG/Private/Compute/Elements/PCGTransformPointsKernel.cpp
--- a/Engine/Plugins/PCG/Source/PCG/Private/Compute/Elements/PCGTransformPointsKernel.cpp
+++ b/Engine/Plugins/PCG/Source/PCG/Private/Compute/Elements/PCGTransformPointsKernel.cpp
@@ -24,7 +24,7 @@ TSharedPtr<const FPCGDataCollectionDesc> UPCGTransformPointsKernel::ComputeOutpu
 		return nullptr;
 	}
 
-	const FPCGKernelPin SourceKernelPin(KernelIndex, PCGPinConstants::DefaultInputLabel, /*bIsInput=*/true);
+	const FPCGKernelPin SourceKernelPin(GetKernelIndex(), PCGPinConstants::DefaultInputLabel, /*bIsInput=*/true);
 	const TSharedPtr<const FPCGDataCollectionDesc> InDataDesc = InBinding->ComputeKernelPinDataDesc(SourceKernelPin);
 
 	if (!ensure(InDataDesc))
diff --git a/Engine/Plugins/PCG/Source/PCG/Private/Elements/Metadata/PCGMetadataPartitionKernel.cpp b/Engine/Plugins/PCG/Source/PCG/Private/Elements/Metadata/PCGMetadataPartitionKernel.cpp
--- a/Engine/Plugins/PCG/Source/PCG/Private/Elements/Metadata/PCGMetadataPartitionKernel.cpp
+++ b/Engine/Plugins/PCG/Source/PCG/Private/Elements/Metadata/PCGMetadataPartitionKernel.cpp
@@ -93,7 +93,7 @@ TSharedPtr<const FPCGDataCollectionDesc> UPCGMetadataPartitionKernel::ComputeOut
 		return nullptr;
 	}
 
-	const FPCGKernelPin InputKernelPin(KernelIndex, PCGPinConstants::DefaultInputLabel, /*bIsInput=*/true);
+	const FPCGKernelPin InputKernelPin(GetKernelIndex(), PCGPinConstants::DefaultInputLabel, /*bIsInput=*/true);
 	const TSharedPtr<const FPCGDataCollectionDesc> InputPinDesc = InBinding->ComputeKernelPinDataDesc(InputKernelPin);
 
 	if (!ensure(InputPinDesc))
@@ -258,7 +258,7 @@ bool UPCGMetadataPartitionKernel::PerformStaticValidation()
 	if (MPSettings->PartitionAttributeSelectors.Num() != 1)
 	{
 #if PCG_KERNEL_LOGGING_ENABLED
-		StaticLogEntries.Emplace(LOCTEXT("MustProvideOnePartitionAttribute", "GPU implementation currently only supports a single partition attribute."), EPCGKernelLogVerbosity::Error);
+		AddStaticLogEntry(LOCTEXT("MustProvideOnePartitionAttribute", "GPU implementation currently only supports a single partition attribute."), EPCGKernelLogVerbosity::Error);
 #endif
 		return false;
 	}
@@ -266,7 +266,7 @@ bool UPCGMetadataPartitionKernel::PerformStaticValidation()
 	if (!MPSettings->PartitionAttributeSelectors[0].IsBasicAttribute())
 	{
 #if PCG_KERNEL_LOGGING_ENABLED
-		StaticLogEntries.Emplace(LOCTEXT("OnlyBasicAttributesSupported", "GPU implementation currently only supports basic attributes."), EPCGKernelLogVerbosity::Error);
+		AddStaticLogEntry(LOCTEXT("OnlyBasicAttributesSupported", "GPU implementation currently only supports basic attributes."), EPCGKernelLogVerbosity::Error);
 #endif
 		return false;
 	}
diff --git a/Engine/Plugins/PCG/Source/PCG/Private/Elements/PCGCopyPointsAnalysisKernel.cpp b/Engine/Plugins/PCG/Source/PCG/Private/Elements/PCGCopyPointsAnalysisKernel.cpp
--- a/Engine/Plugins/PCG/Source/PCG/Private/Elements/PCGCopyPointsAnalysisKernel.cpp
+++ b/Engine/Plugins/PCG/Source/PCG/Private/Elements/PCGCopyPointsAnalysisKernel.cpp
@@ -55,8 +55,8 @@ TSharedPtr<const FPCGDataCollectionDesc> UPCGCopyPointsAnalysisKernel::ComputeOu
 		return nullptr;
 	}
 
-	const FPCGKernelPin SourceKernelPin(KernelIndex, PCGCopyPointsConstants::SourcePointsLabel, /*bIsInput=*/true);
-	const FPCGKernelPin TargetKernelPin(KernelIndex, PCGCopyPointsConstants::TargetPointsLabel, /*bIsInput=*/true);
+	const FPCGKernelPin SourceKernelPin(GetKernelIndex(), PCGCopyPointsConstants::SourcePointsLabel, /*bIsInput=*/true);
+	const FPCGKernelPin TargetKernelPin(GetKernelIndex(), PCGCopyPointsConstants::TargetPointsLabel, /*bIsInput=*/true);
 
 	const TSharedPtr<const FPCGDataCollectionDesc> SourcesDesc = InBinding->ComputeKernelPinDataDesc(SourceKernelPin);
 	const TSharedPtr<const FPCGDataCollectionDesc> TargetsDesc = InBinding->ComputeKernelPinDataDesc(TargetKernelPin);
diff --git a/Engine/Plugins/PCG/Source/PCG/Private/Elements/PCGCopyPointsKernel.cpp b/Engine/Plugins/PCG/Source/PCG/Private/Elements/PCGCopyPointsKernel.cpp
--- a/Engine/Plugins/PCG/Source/PCG/Private/Elements/PCGCopyPointsKernel.cpp
+++ b/Engine/Plugins/PCG/Source/PCG/Private/Elements/PCGCopyPointsKernel.cpp
@@ -99,8 +99,8 @@ TSharedPtr<const FPCGDataCollectionDesc> UPCGCopyPointsKernel::ComputeOutputBind
 		}
 	}
 
-	const FPCGKernelPin SourceKernelPin(KernelIndex, PCGCopyPointsConstants::SourcePointsLabel, /*bIsInput=*/true);
-	const FPCGKernelPin TargetKernelPin(KernelIndex, PCGCopyPointsConstants::TargetPointsLabel, /*bIsInput=*/true);
+	const FPCGKernelPin SourceKernelPin(GetKernelIndex(), PCGCopyPointsConstants::SourcePointsLabel, /*bIsInput=*/true);
+	const FPCGKernelPin TargetKernelPin(GetKernelIndex(), PCGCopyPointsConstants::TargetPointsLabel, /*bIsInput=*/true);
 
 	const TSharedPtr<const FPCGDataCollectionDesc> SourcePinDesc = InBinding->ComputeKernelPinDataDesc(SourceKernelPin);
 	const TSharedPtr<const FPCGDataCollectionDesc> TargetPinDesc = InBinding->ComputeKernelPinDataDesc(TargetKernelPin);
diff --git a/Engine/Plugins/PCG/Source/PCG/Private/Elements/PCGSkinnedMeshSpawnerKernel.cpp b/Engine/Plugins/PCG/Source/PCG/Private/Elements/PCGSkinnedMeshSpawnerKernel.cpp
--- a/Engine/Plugins/PCG/Source/PCG/Private/Elements/PCGSkinnedMeshSpawnerKernel.cpp
+++ b/Engine/Plugins/PCG/Source/PCG/Private/Elements/PCGSkinnedMeshSpawnerKernel.cpp
@@ -42,7 +42,7 @@ TSharedPtr<const FPCGDataCollectionDesc> UPCGSkinnedMeshSpawnerKernel::ComputeOu
 	}
 
 	// Forward data from In to Out.
-	const FPCGKernelPin InputKernelPin(KernelIndex, PCGPinConstants::DefaultInputLabel, /*bIsInput=*/true);
+	const FPCGKernelPin InputKernelPin(GetKernelIndex(), PCGPinConstants::DefaultInputLabel, /*bIsInput=*/true);
 	const TSharedPtr<const FPCGDataCollectionDesc> InputDesc = InBinding->ComputeKernelPinDataDesc(InputKernelPin);
 
 	TSharedPtr<FPCGDataCollectionDesc> OutDataDesc = FPCGDataCollectionDesc::MakeSharedFrom(InputDesc);
@@ -111,7 +111,7 @@ bool UPCGSkinnedMeshSpawnerKernel::PerformStaticValidation()
 	if (SMSettings->InstanceDataPackerParameters && !SMSettings->InstanceDataPackerParameters->GetAttributeNames(/*OutNames=*/nullptr))
 	{
 #if PCG_KERNEL_LOGGING_ENABLED
-		StaticLogEntries.Emplace(LOCTEXT("InvalidInstancePacker", "Selected instance packer does not support GPU execution."), EPCGKernelLogVerbosity::Error);
+		AddStaticLogEntry(LOCTEXT("InvalidInstancePacker", "Selected instance packer does not support GPU execution."), EPCGKernelLogVerbosity::Error);
 #endif
 		return false;
 	}
diff --git a/Engine/Plugins/PCG/Source/PCG/Private/Elements/PCGStaticMeshSpawnerKernel.cpp b/Engine/Plugins/PCG/Source/PCG/Private/Elements/PCGStaticMeshSpawnerKernel.cpp
--- a/Engine/Plugins/PCG/Source/PCG/Private/Elements/PCGStaticMeshSpawnerKernel.cpp
+++ b/Engine/Plugins/PCG/Source/PCG/Private/Elements/PCGStaticMeshSpawnerKernel.cpp
@@ -44,7 +44,7 @@ TSharedPtr<const FPCGDataCollectionDesc> UPCGStaticMeshSpawnerKernel::ComputeOut
 	}
 
 	// Forward data from In to Out.
-	const FPCGKernelPin InputKernelPin(KernelIndex, PCGPinConstants::DefaultInputLabel, /*bIsInput=*/true);
+	const FPCGKernelPin InputKernelPin(GetKernelIndex(), PCGPinConstants::DefaultInputLabel, /*bIsInput=*/true);
 	const TSharedPtr<const FPCGDataCollectionDesc> InputDesc = InBinding->ComputeKernelPinDataDesc(InputKernelPin);
 
 	TSharedPtr<FPCGDataCollectionDesc> OutDataDesc = FPCGDataCollectionDesc::MakeSharedFrom(InputDesc);
@@ -189,7 +189,7 @@ bool UPCGStaticMeshSpawnerKernel::PerformStaticValidation()
 		if (SelectorWeighted->MeshEntries.IsEmpty())
 		{
 #if PCG_KERNEL_LOGGING_ENABLED
-			StaticLogEntries.Emplace(PCGStaticMeshSpawnerKernel::NoMeshEntriesFormat, EPCGKernelLogVerbosity::Error);
+			AddStaticLogEntry(PCGStaticMeshSpawnerKernel::NoMeshEntriesFormat, EPCGKernelLogVerbosity::Error);
 #endif
 			return false;
 		}
@@ -199,7 +199,7 @@ bool UPCGStaticMeshSpawnerKernel::PerformStaticValidation()
 			if (Entry.Descriptor.StaticMesh.IsNull())
 			{
 #if PCG_KERNEL_LOGGING_ENABLED
-				StaticLogEntries.Emplace(LOCTEXT("UnassignedMesh", "Unassigned mesh."), EPCGKernelLogVerbosity::Error);
+				AddStaticLogEntry(LOCTEXT("UnassignedMesh", "Unassigned mesh."), EPCGKernelLogVerbosity::Error);
 #endif
 				return false;
 			}
@@ -208,7 +208,7 @@ bool UPCGStaticMeshSpawnerKernel::PerformStaticValidation()
 	else if (!SMSettings->MeshSelectorParameters || !SMSettings->MeshSelectorParameters->IsA<UPCGMeshSelectorByAttribute>())
 	{
 #if PCG_KERNEL_LOGGING_ENABLED
-		StaticLogEntries.Emplace(LOCTEXT("InvalidMeshSelector", "Currently GPU Static Mesh Spawner nodes must use PCGMeshSelectorWeighted or UPCGMeshSelectorByAttribute as the mesh selector type."), EPCGKernelLogVerbosity::Error);
+		AddStaticLogEntry(LOCTEXT("InvalidMeshSelector", "Currently GPU Static Mesh Spawner nodes must use PCGMeshSelectorWeighted or UPCGMeshSelectorByAttribute as the mesh selector type."), EPCGKernelLogVerbosity::Error);
 #endif
 		return false;
 	}
@@ -218,7 +218,7 @@ bool UPCGStaticMeshSpawnerKernel::PerformStaticValidation()
 	if (SMSettings->InstanceDataPackerParameters && !SMSettings->InstanceDataPackerParameters->GetAttributeNames(/*OutNames=*/nullptr))
 	{
 #if PCG_KERNEL_LOGGING_ENABLED
-		StaticLogEntries.Emplace(LOCTEXT("InvalidInstancePacker", "Selected instance packer does not support GPU execution."), EPCGKernelLogVerbosity::Error);
+		AddStaticLogEntry(LOCTEXT("InvalidInstancePacker", "Selected instance packer does not support GPU execution."), EPCGKernelLogVerbosity::Error);
 #endif
 		return false;
 	}
diff --git a/Engine/Plugins/PCG/Source/PCG/Public/Compute/PCGComputeKernel.h b/Engine/Plugins/PCG/Source/PCG/Public/Compute/PCGComputeKernel.h
--- a/Engine/Plugins/PCG/Source/PCG/Public/Compute/PCGComputeKernel.h
+++ b/Engine/Plugins/PCG/Source/PCG/Public/Compute/PCGComputeKernel.h
@@ -119,10 +119,10 @@ public:
 	PCG_API virtual FString GetCookedSource(FPCGGPUCompilationContext& InOutContext) const;
 
 	/** Get the virtual file path of the shader source file for this kernel. */
-	PCG_API virtual const TCHAR* GetSourceFilePath() const { return nullptr; }
+	virtual const TCHAR* GetSourceFilePath() const { return nullptr; }
 
 	/** Get the name of the main kernel function in the source. This name is also displayed in GPU profile scopes. */
-	PCG_API virtual FString GetEntryPoint() const PURE_VIRTUAL(UPCGComputeKernel::GetEntryPoint, return TEXT("Main"););
+	virtual FString GetEntryPoint() const PURE_VIRTUAL(UPCGComputeKernel::GetEntryPoint, return TEXT("Main"););
 
 	/**
 	 * Gathers additional compute sources referenced by this kernel.
@@ -132,16 +132,16 @@ public:
 	PCG_API virtual void GatherAdditionalSources(TArray<TObjectPtr<class UComputeSource>>& OutAdditionalSources) const;
 
 	/** Create additional input data interfaces to marshal any required input data. */
-	PCG_API virtual void CreateAdditionalInputDataInterfaces(FPCGGPUCompilationContext& InOutContext, UObject* InObjectOuter, TArray<TObjectPtr<UComputeDataInterface>>& OutDataInterfaces) const {}
+	virtual void CreateAdditionalInputDataInterfaces(FPCGGPUCompilationContext& InOutContext, UObject* InObjectOuter, TArray<TObjectPtr<UComputeDataInterface>>& OutDataInterfaces) const {}
 
 	/** Create additional output data interfaces to marshal any required output data. */
-	PCG_API virtual void CreateAdditionalOutputDataInterfaces(FPCGGPUCompilationContext& InOutContext, UObject* InObjectOuter, TArray<TObjectPtr<UComputeDataInterface>>& OutDataInterfaces) const {}
+	virtual void CreateAdditionalOutputDataInterfaces(FPCGGPUCompilationContext& InOutContext, UObject* InObjectOuter, TArray<TObjectPtr<UComputeDataInterface>>& OutDataInterfaces) const {}
 
 	/** Whether compute graphs should be split at this input pin. Useful for overridable parameters. */
 	PCG_API virtual bool SplitGraphAtInput(FName InPinLabel) const;
 
 	/** Whether compute graphs should be split at the output of this node. Useful for reading back runtime statistics or diagnostic info before continuing. */
-	PCG_API virtual bool SplitGraphAtOutput() const { return false; }
+	virtual bool SplitGraphAtOutput() const { return false; }
 
 	/** Mark a pin as being internal to the kernel graph. This means it is not on the CPU/GPU boundary. */
 	void AddInternalPin(FName InPinLabel) { InternalPinLabels.Add(InPinLabel); }
@@ -154,37 +154,37 @@ public:
 #endif
 
 	/** Compute a description of data that will be output from pin InOutputPinLabel. */
-	PCG_API virtual TSharedPtr<const FPCGDataCollectionDesc> ComputeOutputBindingDataDesc(FName InOutputPinLabel, UPCGDataBinding* InBinding) const PURE_VIRTUAL(UPCGComputeKernel::ComputeOutputBindingDataDesc, return nullptr;);
+	virtual TSharedPtr<const FPCGDataCollectionDesc> ComputeOutputBindingDataDesc(FName InOutputPinLabel, UPCGDataBinding* InBinding) const PURE_VIRTUAL(UPCGComputeKernel::ComputeOutputBindingDataDesc, return nullptr;);
 
 	/** Compute how many threads should be dispatched to execute this node on the GPU. */
-	PCG_API virtual int ComputeThreadCount(const UPCGDataBinding* Binding) const PURE_VIRTUAL(UPCGComputeKernel::ComputeThreadCount, return 0;);
+	virtual int ComputeThreadCount(const UPCGDataBinding* Binding) const PURE_VIRTUAL(UPCGComputeKernel::ComputeThreadCount, return 0;);
 
 	/** Get any data labels that are used statically by this node. */
-	PCG_API virtual void GetDataLabels(FName InPinLabel, TArray<FString>& OutDataLabels) const {}
+	virtual void GetDataLabels(FName InPinLabel, TArray<FString>& OutDataLabels) const {}
 
 	/** Initialize all buffer data to 0. This is not free so should only be used when kernel execution requires 0-initialized data, such as for counters. */
-	PCG_API virtual bool DoesOutputPinRequireZeroInitialization(FName InOutputPinLabel) const { return false; }
+	virtual bool DoesOutputPinRequireZeroInitialization(FName InOutputPinLabel) const { return false; }
 
 	/** Add any strings emitted by this node that are known statically at compile time. */
-	PCG_API virtual void AddStaticCreatedStrings(TArray<FString>& InOutStringTable) const {}
+	virtual void AddStaticCreatedStrings(TArray<FString>& InOutStringTable) const {}
 
 	/** Get all the attributes read or written by this node for which we know the name and type statically. In cases that only the
 	* attribute name used by the node is known (example: user specifies SM Spawner instance data attributes by name only, not type), the attribute
 	* is omitted from this list and must be resolved at execution time in the data provider.
 	*/
-	PCG_API virtual void GetKernelAttributeKeys(TArray<FPCGKernelAttributeKey>& OutKeys) const {}
+	virtual void GetKernelAttributeKeys(TArray<FPCGKernelAttributeKey>& OutKeys) const {}
 
 	/** Get multiplier to be applied to the dispatched thread count. */
-	PCG_API virtual uint32 GetThreadCountMultiplier() const { return 1u; }
+	virtual uint32 GetThreadCountMultiplier() const { return 1u; }
 
 	/** Multiplier to be applied to the element count of each output data. */
-	PCG_API virtual uint32 GetElementCountMultiplier(FName InOutputPinLabel) const { return 1u; }
+	virtual uint32 GetElementCountMultiplier(FName InOutputPinLabel) const { return 1u; }
 
 	/** Helper to compute data descriptions when data is entirely determined by pin properties. */
 	void ComputeDataDescFromPinProperties(const FPCGPinPropertiesGPU& OutputPinProps, const TArrayView<const FPCGPinProperties>& InInputPinProps, UPCGDataBinding* InBinding, TSharedPtr<FPCGDataCollectionDesc> OutPinDesc) const;
 
-	PCG_API virtual void GetInputPins(TArray<FPCGPinProperties>& OutPins) const {}
-	PCG_API virtual void GetOutputPins(TArray<FPCGPinPropertiesGPU>& OutPins) const {}
+	virtual void GetInputPins(TArray<FPCGPinProperties>& OutPins) const {}
+	virtual void GetOutputPins(TArray<FPCGPinPropertiesGPU>& OutPins) const {}
 
 	/** Return the concatenation of GetInputPins() and GetOverridableParamPins(). */
 	void GetInputPinsAndOverridablePins(TArray<FPCGPinProperties>& OutPins) const;
@@ -207,11 +207,17 @@ public:
 protected:
 #if WITH_EDITOR
 	/** Implement to do any additional initialization specific to your kernel. */
-	PCG_API virtual void InitializeInternal() {}
+	virtual void InitializeInternal() {}
 
 	/** Performs validation on compile time information (e.g. Node, Settings, etc.). Caches validation errors/warnings. */
 	PCG_API virtual bool PerformStaticValidation();
 
+	template <typename... ArgsType>
+	void AddStaticLogEntry(ArgsType&&... Args)
+	{
+		StaticLogEntries.Emplace(Forward<ArgsType>(Args)...);
+	}
+
 	/** Do any validation of the PCG node. Returns true if node is valid and GPU execution can proceed. */
 	bool ValidatePCGNode(TArray<FPCGKernelLogEntry>& OutLogEntries) const;
 
@@ -222,7 +228,7 @@ protected:
 	/** Collects all param pins which are overridable on the GPU. */
 	void GetOverridableParamPins(TArray<FPCGPinProperties>& OutPins) const;
 
-protected:
+private:
 	UPROPERTY()
 	EPCGComputeKernelFlags PCGKernelFlags = EPCGComputeKernelFlags::None;
 
