[PCG] Finalize data interface API

#rnx
#FailedValidations Preflight

diff --git a/Engine/Plugins/PCG/Source/PCG/Private/Compute/BuiltInKernels/PCGCountUniqueAttributeValuesKernel.cpp b/Engine/Plugins/PCG/Source/PCG/Private/Compute/BuiltInKernels/PCGCountUniqueAttributeValuesKernel.cpp
--- a/Engine/Plugins/PCG/Source/PCG/Private/Compute/BuiltInKernels/PCGCountUniqueAttributeValuesKernel.cpp
+++ b/Engine/Plugins/PCG/Source/PCG/Private/Compute/BuiltInKernels/PCGCountUniqueAttributeValuesKernel.cpp
@@ -225,7 +225,7 @@ void UPCGCountUniqueAttributeValuesKernel::CreateAdditionalInputDataInterfaces(F
 	Super::CreateAdditionalInputDataInterfaces(InOutContext, InObjectOuter, OutDataInterfaces);
 
 	TObjectPtr<UPCGCountUniqueAttributeValuesDataInterface> KernelDI = InOutContext.NewObject_AnyThread<UPCGCountUniqueAttributeValuesDataInterface>(InObjectOuter);
-	KernelDI->ProducerKernel = this;
+	KernelDI->SetProducerKernel(this);
 	KernelDI->SetAttributeToCountName(AttributeName);
 	KernelDI->SetEmitPerDataCounts(bEmitPerDataCounts);
 	KernelDI->SetOutputRawBuffer(bOutputRawBuffer);
diff --git a/Engine/Plugins/PCG/Source/PCG/Private/Compute/BuiltInKernels/PCGDataNumKernel.cpp b/Engine/Plugins/PCG/Source/PCG/Private/Compute/BuiltInKernels/PCGDataNumKernel.cpp
--- a/Engine/Plugins/PCG/Source/PCG/Private/Compute/BuiltInKernels/PCGDataNumKernel.cpp
+++ b/Engine/Plugins/PCG/Source/PCG/Private/Compute/BuiltInKernels/PCGDataNumKernel.cpp
@@ -44,7 +44,7 @@ void UPCGDataNumKernel::CreateAdditionalInputDataInterfaces(FPCGGPUCompilationCo
 	Super::CreateAdditionalInputDataInterfaces(InOutContext, InObjectOuter, OutDataInterfaces);
 
 	TObjectPtr<UPCGDataNumDataInterface> NodeDI = InOutContext.NewObject_AnyThread<UPCGDataNumDataInterface>(InObjectOuter);
-	NodeDI->ProducerKernel = this;
+	NodeDI->SetProducerKernel(this);
 
 	OutDataInterfaces.Add(NodeDI);
 }
diff --git a/Engine/Plugins/PCG/Source/PCG/Private/Compute/DataInterfaces/BuiltInKernels/PCGCountUniqueAttributeValuesDataInterface.cpp b/Engine/Plugins/PCG/Source/PCG/Private/Compute/DataInterfaces/BuiltInKernels/PCGCountUniqueAttributeValuesDataInterface.cpp
--- a/Engine/Plugins/PCG/Source/PCG/Private/Compute/DataInterfaces/BuiltInKernels/PCGCountUniqueAttributeValuesDataInterface.cpp
+++ b/Engine/Plugins/PCG/Source/PCG/Private/Compute/DataInterfaces/BuiltInKernels/PCGCountUniqueAttributeValuesDataInterface.cpp
@@ -97,7 +97,6 @@ void UPCGCountUniqueAttributeValuesDataProvider::Initialize(const UComputeDataIn
 	Super::Initialize(InDataInterface, InBinding, InInputMask, InOutputMask);
 
 	const UPCGCountUniqueAttributeValuesDataInterface* DataInterface = CastChecked<UPCGCountUniqueAttributeValuesDataInterface>(InDataInterface);
-	Kernel = DataInterface->ProducerKernel;
 
 	AttributeToCountName = DataInterface->GetAttributeToCountName();
 	bEmitPerDataCounts = DataInterface->GetEmitPerDataCounts();
@@ -122,7 +121,7 @@ bool UPCGCountUniqueAttributeValuesDataProvider::PrepareForExecute_GameThread(UP
 		check(OutputCountAttributeId != -1);
 	}
 
-	const TSharedPtr<const FPCGDataCollectionDesc> InputDataDesc = InBinding->GetCachedKernelPinDataDesc(Kernel, PCGPinConstants::DefaultInputLabel, /*bIsInput*/true);
+	const TSharedPtr<const FPCGDataCollectionDesc> InputDataDesc = InBinding->GetCachedKernelPinDataDesc(GetProducerKernel(), PCGPinConstants::DefaultInputLabel, /*bIsInput*/true);
 
 	if (!ensure(InputDataDesc))
 	{
@@ -179,7 +178,6 @@ FComputeDataProviderRenderProxy* UPCGCountUniqueAttributeValuesDataProvider::Get
 
 void UPCGCountUniqueAttributeValuesDataProvider::Reset()
 {
-	Kernel = nullptr;
 	AttributeToCountName = NAME_None;
 	AttributeToCountId = INDEX_NONE;
 	OutputValueAttributeId = INDEX_NONE;
diff --git a/Engine/Plugins/PCG/Source/PCG/Private/Compute/DataInterfaces/BuiltInKernels/PCGCountUniqueAttributeValuesDataInterface.h b/Engine/Plugins/PCG/Source/PCG/Private/Compute/DataInterfaces/BuiltInKernels/PCGCountUniqueAttributeValuesDataInterface.h
--- a/Engine/Plugins/PCG/Source/PCG/Private/Compute/DataInterfaces/BuiltInKernels/PCGCountUniqueAttributeValuesDataInterface.h
+++ b/Engine/Plugins/PCG/Source/PCG/Private/Compute/DataInterfaces/BuiltInKernels/PCGCountUniqueAttributeValuesDataInterface.h
@@ -61,9 +61,6 @@ public:
 	//~ End UComputeDataProvider Interface
 
 public:
-	UPROPERTY()
-	TObjectPtr<const UPCGComputeKernel> Kernel;
-
 	FName AttributeToCountName;
 
 	int32 AttributeToCountId = INDEX_NONE;
diff --git a/Engine/Plugins/PCG/Source/PCG/Private/Compute/DataInterfaces/BuiltInKernels/PCGDataNumDataInterface.cpp b/Engine/Plugins/PCG/Source/PCG/Private/Compute/DataInterfaces/BuiltInKernels/PCGDataNumDataInterface.cpp
--- a/Engine/Plugins/PCG/Source/PCG/Private/Compute/DataInterfaces/BuiltInKernels/PCGDataNumDataInterface.cpp
+++ b/Engine/Plugins/PCG/Source/PCG/Private/Compute/DataInterfaces/BuiltInKernels/PCGDataNumDataInterface.cpp
@@ -56,7 +56,7 @@ void UPCGDataNumDataProvider::Initialize(const UComputeDataInterface* InDataInte
 
 	Super::Initialize(InDataInterface, InBinding, InInputMask, InOutputMask);
 
-	const UPCGDataNumSettings* Settings = CastChecked<UPCGDataNumSettings>(ProducerKernel->GetSettings());
+	const UPCGDataNumSettings* Settings = CastChecked<UPCGDataNumSettings>(GetProducerKernel()->GetSettings());
 	OutputAttributeName = Settings->OutputAttributeName;
 }
 
diff --git a/Engine/Plugins/PCG/Source/PCG/Private/Compute/DataInterfaces/Elements/PCGCopyPointsAnalysisDataInterface.cpp b/Engine/Plugins/PCG/Source/PCG/Private/Compute/DataInterfaces/Elements/PCGCopyPointsAnalysisDataInterface.cpp
--- a/Engine/Plugins/PCG/Source/PCG/Private/Compute/DataInterfaces/Elements/PCGCopyPointsAnalysisDataInterface.cpp
+++ b/Engine/Plugins/PCG/Source/PCG/Private/Compute/DataInterfaces/Elements/PCGCopyPointsAnalysisDataInterface.cpp
@@ -72,7 +72,7 @@ bool UPCGCopyPointsAnalysisDataProvider::PrepareForExecute_GameThread(UPCGDataBi
 		return false;
 	}
 
-	const FPCGKernelParams* KernelParams = InBinding->GetCachedKernelParams(ProducerKernel);
+	const FPCGKernelParams* KernelParams = InBinding->GetCachedKernelParams(GetProducerKernel());
 
 	if (!ensure(KernelParams))
 	{
diff --git a/Engine/Plugins/PCG/Source/PCG/Private/Compute/DataInterfaces/Elements/PCGCopyPointsDataInterface.cpp b/Engine/Plugins/PCG/Source/PCG/Private/Compute/DataInterfaces/Elements/PCGCopyPointsDataInterface.cpp
--- a/Engine/Plugins/PCG/Source/PCG/Private/Compute/DataInterfaces/Elements/PCGCopyPointsDataInterface.cpp
+++ b/Engine/Plugins/PCG/Source/PCG/Private/Compute/DataInterfaces/Elements/PCGCopyPointsDataInterface.cpp
@@ -123,7 +123,7 @@ bool UPCGCopyPointsDataProvider::PerformPreExecuteReadbacks_GameThread(UPCGDataB
 		return false;
 	}
 
-	KernelParams = InBinding->GetCachedKernelParams(ProducerKernel);
+	KernelParams = InBinding->GetCachedKernelParams(GetProducerKernel());
 
 	if (!ensure(KernelParams))
 	{
@@ -138,7 +138,7 @@ bool UPCGCopyPointsDataProvider::PerformPreExecuteReadbacks_GameThread(UPCGDataB
 		return true;
 	}
 
-	const int32 AnalysisDataIndex = ProducerKernel ? InBinding->GetFirstInputDataIndex(ProducerKernel, PCGCopyPointsConstants::SelectedFlagsPinLabel) : INDEX_NONE;
+	const int32 AnalysisDataIndex = GetProducerKernel() ? InBinding->GetFirstInputDataIndex(GetProducerKernel(), PCGCopyPointsConstants::SelectedFlagsPinLabel) : INDEX_NONE;
 	if (AnalysisDataIndex == INDEX_NONE)
 	{
 		// Analysis data was not produced, nothing to read back.
@@ -168,14 +168,14 @@ bool UPCGCopyPointsDataProvider::PrepareForExecute_GameThread(UPCGDataBinding* I
 
 	if (bMatchBasedOnAttribute)
 	{
-		const int32 AnalysisDataIndex = InBinding->GetFirstInputDataIndex(ProducerKernel, PCGCopyPointsConstants::SelectedFlagsPinLabel);
+		const int32 AnalysisDataIndex = InBinding->GetFirstInputDataIndex(GetProducerKernel(), PCGCopyPointsConstants::SelectedFlagsPinLabel);
 		if (AnalysisDataIndex == INDEX_NONE)
 		{
 			return true;
 		}
 
-		const TSharedPtr<const FPCGDataCollectionDesc> SourcePinDesc = InBinding->GetCachedKernelPinDataDesc(ProducerKernel, PCGCopyPointsConstants::SourcePointsLabel, /*bIsInput=*/true);
-		const TSharedPtr<const FPCGDataCollectionDesc> TargetPinDesc = InBinding->GetCachedKernelPinDataDesc(ProducerKernel, PCGCopyPointsConstants::TargetPointsLabel, /*bIsInput=*/true);
+		const TSharedPtr<const FPCGDataCollectionDesc> SourcePinDesc = InBinding->GetCachedKernelPinDataDesc(GetProducerKernel(), PCGCopyPointsConstants::SourcePointsLabel, /*bIsInput=*/true);
+		const TSharedPtr<const FPCGDataCollectionDesc> TargetPinDesc = InBinding->GetCachedKernelPinDataDesc(GetProducerKernel(), PCGCopyPointsConstants::TargetPointsLabel, /*bIsInput=*/true);
 
 		if (!ensure(SourcePinDesc) || !ensure(TargetPinDesc))
 		{
diff --git a/Engine/Plugins/PCG/Source/PCG/Private/Compute/DataInterfaces/Elements/PCGCullPointsOutsideActorBoundsDataInterface.cpp b/Engine/Plugins/PCG/Source/PCG/Private/Compute/DataInterfaces/Elements/PCGCullPointsOutsideActorBoundsDataInterface.cpp
--- a/Engine/Plugins/PCG/Source/PCG/Private/Compute/DataInterfaces/Elements/PCGCullPointsOutsideActorBoundsDataInterface.cpp
+++ b/Engine/Plugins/PCG/Source/PCG/Private/Compute/DataInterfaces/Elements/PCGCullPointsOutsideActorBoundsDataInterface.cpp
@@ -56,7 +56,7 @@ UComputeDataProvider* UPCGCullPointsOutsideActorBoundsDataInterface::CreateDataP
 
 FComputeDataProviderRenderProxy* UPCGCullPointsOutsideActorBoundsDataProvider::GetRenderProxy()
 {
-	const UPCGCullPointsOutsideActorBoundsSettings* Settings = CastChecked<UPCGCullPointsOutsideActorBoundsSettings>(ProducerKernel->GetSettings());
+	const UPCGCullPointsOutsideActorBoundsSettings* Settings = CastChecked<UPCGCullPointsOutsideActorBoundsSettings>(GetProducerKernel()->GetSettings());
 
 	FPCGCullPointsOutsideActorBoundsDataProviderProxy::FCullPointsOutsideActorBoundsData_RenderThread ProxyData;
 	ProxyData.BoundsExpansion = Settings->BoundsExpansion;
diff --git a/Engine/Plugins/PCG/Source/PCG/Private/Compute/DataInterfaces/Elements/PCGCustomHLSLDataInterface.cpp b/Engine/Plugins/PCG/Source/PCG/Private/Compute/DataInterfaces/Elements/PCGCustomHLSLDataInterface.cpp
--- a/Engine/Plugins/PCG/Source/PCG/Private/Compute/DataInterfaces/Elements/PCGCustomHLSLDataInterface.cpp
+++ b/Engine/Plugins/PCG/Source/PCG/Private/Compute/DataInterfaces/Elements/PCGCustomHLSLDataInterface.cpp
@@ -64,7 +64,7 @@ bool UPCGCustomHLSLDataProvider::PrepareForExecute_GameThread(UPCGDataBinding* I
 		return false;
 	}
 
-	KernelParams = InBinding->GetCachedKernelParams(ProducerKernel);
+	KernelParams = InBinding->GetCachedKernelParams(GetProducerKernel());
 
 	return true;
 }
diff --git a/Engine/Plugins/PCG/Source/PCG/Private/Compute/DataInterfaces/Elements/PCGMetadataPartitionDataInterface.cpp b/Engine/Plugins/PCG/Source/PCG/Private/Compute/DataInterfaces/Elements/PCGMetadataPartitionDataInterface.cpp
--- a/Engine/Plugins/PCG/Source/PCG/Private/Compute/DataInterfaces/Elements/PCGMetadataPartitionDataInterface.cpp
+++ b/Engine/Plugins/PCG/Source/PCG/Private/Compute/DataInterfaces/Elements/PCGMetadataPartitionDataInterface.cpp
@@ -86,16 +86,6 @@ UComputeDataProvider* UPCGMetadataPartitionDataInterface::CreateDataProvider() c
 	return NewObject<UPCGMetaDataPartitionDataProvider>();
 }
 
-void UPCGMetaDataPartitionDataProvider::Initialize(const UComputeDataInterface* InDataInterface, UObject* InBinding, uint64 InInputMask, uint64 InOutputMask)
-{
-	TRACE_CPUPROFILER_EVENT_SCOPE(UPCGMetaDataPartitionDataProvider::Initialize);
-
-	Super::Initialize(InDataInterface, InBinding, InInputMask, InOutputMask);
-
-	const UPCGMetadataPartitionDataInterface* DataInterface = CastChecked<UPCGMetadataPartitionDataInterface>(InDataInterface);
-	ProducerKernel = DataInterface->ProducerKernel;
-}
-
 bool UPCGMetaDataPartitionDataProvider::PerformPreExecuteReadbacks_GameThread(UPCGDataBinding* InBinding)
 {
 	TRACE_CPUPROFILER_EVENT_SCOPE(UPCGMetaDataPartitionDataProvider::PerformPreExecuteReadbacks_GameThread);
@@ -108,7 +98,7 @@ bool UPCGMetaDataPartitionDataProvider::PerformPreExecuteReadbacks_GameThread(UP
 
 	if (AnalysisDataIndices.IsEmpty())
 	{
-		AnalysisDataIndices = InBinding->GetPinInputDataIndices(ProducerKernel, PCGMetadataPartitionConstants::ElementCountsPinLabel);
+		AnalysisDataIndices = InBinding->GetPinInputDataIndices(GetProducerKernel(), PCGMetadataPartitionConstants::ElementCountsPinLabel);
 
 		if (AnalysisDataIndices.IsEmpty())
 		{
@@ -136,7 +126,6 @@ bool UPCGMetaDataPartitionDataProvider::PrepareForExecute_GameThread(UPCGDataBin
 {
 	TRACE_CPUPROFILER_EVENT_SCOPE(UPCGMetadataPartitionDataInterface::PrepareForExecute_GameThread);
 	check(InBinding);
-	check(ProducerKernel);
 
 	if (!Super::PrepareForExecute_GameThread(InBinding))
 	{
@@ -149,7 +138,7 @@ bool UPCGMetaDataPartitionDataProvider::PrepareForExecute_GameThread(UPCGDataBin
 		return true;
 	}
 
-	const TSharedPtr<const FPCGDataCollectionDesc> InputDesc = InBinding->GetCachedKernelPinDataDesc(ProducerKernel, PCGPinConstants::DefaultInputLabel, /*bIsInput=*/true);
+	const TSharedPtr<const FPCGDataCollectionDesc> InputDesc = InBinding->GetCachedKernelPinDataDesc(GetProducerKernel(), PCGPinConstants::DefaultInputLabel, /*bIsInput=*/true);
 	if (!InputDesc)
 	{
 		UE_LOG(LogPCG, Error, TEXT("UPCGMetaDataPartitionDataProvider: Could not retrieve input data description."));
@@ -216,12 +205,12 @@ bool UPCGMetaDataPartitionDataProvider::PrepareForExecute_GameThread(UPCGDataBin
 		}
 	}
 
-	const UPCGMetadataPartitionSettings* Settings = Cast<UPCGMetadataPartitionSettings>(ProducerKernel->GetSettings());
+	const UPCGMetadataPartitionSettings* Settings = Cast<UPCGMetadataPartitionSettings>(GetProducerKernel()->GetSettings());
 	check(Settings);
 
 	if (ensure(Settings->PartitionAttributeSelectors.Num() == 1 && Settings->PartitionAttributeSelectors[0].IsBasicAttribute()))
 	{
-		const TSharedPtr<const FPCGDataCollectionDesc> InputDataDesc = InBinding->GetCachedKernelPinDataDesc(ProducerKernel, PCGPinConstants::DefaultInputLabel, /*bIsInput*/true);
+		const TSharedPtr<const FPCGDataCollectionDesc> InputDataDesc = InBinding->GetCachedKernelPinDataDesc(GetProducerKernel(), PCGPinConstants::DefaultInputLabel, /*bIsInput*/true);
 
 		if (!ensure(InputDataDesc))
 		{
diff --git a/Engine/Plugins/PCG/Source/PCG/Private/Compute/DataInterfaces/Elements/PCGMetadataPartitionDataInterface.h b/Engine/Plugins/PCG/Source/PCG/Private/Compute/DataInterfaces/Elements/PCGMetadataPartitionDataInterface.h
--- a/Engine/Plugins/PCG/Source/PCG/Private/Compute/DataInterfaces/Elements/PCGMetadataPartitionDataInterface.h
+++ b/Engine/Plugins/PCG/Source/PCG/Private/Compute/DataInterfaces/Elements/PCGMetadataPartitionDataInterface.h
@@ -37,7 +37,6 @@ public:
 	//~ End UPCGComputeDataProvider Interface
 
 	//~ Begin UComputeDataProvider Interface
-	virtual void Initialize(const UComputeDataInterface* InDataInterface, UObject* InBinding, uint64 InInputMask, uint64 InOutputMask) override;
 	FComputeDataProviderRenderProxy* GetRenderProxy() override;
 	virtual void Reset() override;
 	//~ End UComputeDataProvider Interface
diff --git a/Engine/Plugins/PCG/Source/PCG/Private/Compute/DataInterfaces/Elements/PCGNormalToDensityDataInterface.cpp b/Engine/Plugins/PCG/Source/PCG/Private/Compute/DataInterfaces/Elements/PCGNormalToDensityDataInterface.cpp
--- a/Engine/Plugins/PCG/Source/PCG/Private/Compute/DataInterfaces/Elements/PCGNormalToDensityDataInterface.cpp
+++ b/Engine/Plugins/PCG/Source/PCG/Private/Compute/DataInterfaces/Elements/PCGNormalToDensityDataInterface.cpp
@@ -74,7 +74,7 @@ UComputeDataProvider* UPCGNormalToDensityDataInterface::CreateDataProvider() con
 
 FComputeDataProviderRenderProxy* UPCGNormalToDensityProvider::GetRenderProxy()
 {
-	const UPCGNormalToDensitySettings* Settings = CastChecked<UPCGNormalToDensitySettings>(ProducerKernel->GetSettings());
+	const UPCGNormalToDensitySettings* Settings = CastChecked<UPCGNormalToDensitySettings>(GetProducerKernel()->GetSettings());
 
 	FPCGNormalToDensityDataProviderProxy::FNormalToDensityData_RenderThread ProxyData;
 	ProxyData.Normal = Settings->Normal;
diff --git a/Engine/Plugins/PCG/Source/PCG/Private/Compute/DataInterfaces/Elements/PCGSkinnedMeshSpawnerDataInterface.cpp b/Engine/Plugins/PCG/Source/PCG/Private/Compute/DataInterfaces/Elements/PCGSkinnedMeshSpawnerDataInterface.cpp
--- a/Engine/Plugins/PCG/Source/PCG/Private/Compute/DataInterfaces/Elements/PCGSkinnedMeshSpawnerDataInterface.cpp
+++ b/Engine/Plugins/PCG/Source/PCG/Private/Compute/DataInterfaces/Elements/PCGSkinnedMeshSpawnerDataInterface.cpp
@@ -211,7 +211,7 @@ bool UPCGSkinnedMeshSpawnerDataProvider::PerformPreExecuteReadbacks_GameThread(U
 		return false;
 	}
 
-	const UPCGSkinnedMeshSpawnerSettings* Settings = CastChecked<UPCGSkinnedMeshSpawnerSettings>(ProducerKernel->GetSettings());
+	const UPCGSkinnedMeshSpawnerSettings* Settings = CastChecked<UPCGSkinnedMeshSpawnerSettings>(GetProducerKernel()->GetSettings());
 
 	// Obtain index of analysis data that we need to read back.
 	if (AnalysisDataIndex == INDEX_NONE)
@@ -222,7 +222,7 @@ bool UPCGSkinnedMeshSpawnerDataProvider::PerformPreExecuteReadbacks_GameThread(U
 			return true;
 		}
 
-		AnalysisDataIndex = ProducerKernel ? InBinding->GetFirstInputDataIndex(ProducerKernel, PCGSkinnedMeshSpawnerConstants::InstanceCountsPinLabel) : INDEX_NONE;
+		AnalysisDataIndex = GetProducerKernel() ? InBinding->GetFirstInputDataIndex(GetProducerKernel(), PCGSkinnedMeshSpawnerConstants::InstanceCountsPinLabel) : INDEX_NONE;
 
 		if (AnalysisDataIndex == INDEX_NONE)
 		{
@@ -278,7 +278,7 @@ bool UPCGSkinnedMeshSpawnerDataProvider::PrepareForExecute_GameThread(UPCGDataBi
 {
 	TRACE_CPUPROFILER_EVENT_SCOPE(UPCGSkinnedMeshSpawnerDataProvider::PrepareForExecute_GameThread);
 	check(InBinding);
-	const UPCGSkinnedMeshSpawnerSettings* Settings = CastChecked<UPCGSkinnedMeshSpawnerSettings>(ProducerKernel->GetSettings());
+	const UPCGSkinnedMeshSpawnerSettings* Settings = CastChecked<UPCGSkinnedMeshSpawnerSettings>(GetProducerKernel()->GetSettings());
 
 	if (!Super::PrepareForExecute_GameThread(InBinding))
 	{
@@ -287,7 +287,7 @@ bool UPCGSkinnedMeshSpawnerDataProvider::PrepareForExecute_GameThread(UPCGDataBi
 
 	if (!bRegisteredPrimitives)
 	{
-		InBinding->MeshSpawnersToPrimitives.Add(ProducerKernel, {});
+		InBinding->MeshSpawnersToPrimitives.Add(GetProducerKernel(), {});
 		bRegisteredPrimitives = true;
 	}
 
@@ -295,25 +295,25 @@ bool UPCGSkinnedMeshSpawnerDataProvider::PrepareForExecute_GameThread(UPCGDataBi
 	FPCGContext* Context = ContextHandle ? ContextHandle->GetContext() : nullptr;
 	if (!ensure(Context))
 	{
-		InBinding->CompletedMeshSpawners.Add(ProducerKernel);
+		InBinding->CompletedMeshSpawners.Add(GetProducerKernel());
 		return true;
 	}
 
 	UPCGComponent* SourceComponent = Cast<UPCGComponent>(Context->ExecutionSource.Get());
 	if (!ensure(SourceComponent))
 	{
-		InBinding->CompletedMeshSpawners.Add(ProducerKernel);
+		InBinding->CompletedMeshSpawners.Add(GetProducerKernel());
 		return true;
 	}
 
 	if (SelectorAttributeId == INDEX_NONE && ensure(Settings->MeshSelectorParameters))
 	{
 		const FName SelectorName = Settings->MeshSelectorParameters->MeshAttribute.GetName();
-		const TSharedPtr<const FPCGDataCollectionDesc> InputDataDesc = InBinding->GetCachedKernelPinDataDesc(ProducerKernel, PCGPinConstants::DefaultInputLabel, /*bIsInput=*/true);
+		const TSharedPtr<const FPCGDataCollectionDesc> InputDataDesc = InBinding->GetCachedKernelPinDataDesc(GetProducerKernel(), PCGPinConstants::DefaultInputLabel, /*bIsInput=*/true);
 
 		if (!ensure(InputDataDesc))
 		{
-			InBinding->CompletedMeshSpawners.Add(ProducerKernel);
+			InBinding->CompletedMeshSpawners.Add(GetProducerKernel());
 			return true;
 		}
 
@@ -353,7 +353,7 @@ bool UPCGSkinnedMeshSpawnerDataProvider::PrepareForExecute_GameThread(UPCGDataBi
 					FText::FromName(SelectorName)));
 			}
 
-			InBinding->CompletedMeshSpawners.Add(ProducerKernel);
+			InBinding->CompletedMeshSpawners.Add(GetProducerKernel());
 			return true;
 		}
 
@@ -416,12 +416,12 @@ bool UPCGSkinnedMeshSpawnerDataProvider::PrepareForExecute_GameThread(UPCGDataBi
 		else
 		{
 			// No component set up means we have no more work to do.
-			InBinding->CompletedMeshSpawners.Add(ProducerKernel);
+			InBinding->CompletedMeshSpawners.Add(GetProducerKernel());
 			return true;
 		}
 	}
 
-	InBinding->CompletedMeshSpawners.Add(ProducerKernel);
+	InBinding->CompletedMeshSpawners.Add(GetProducerKernel());
 	return true;
 }
 
@@ -455,7 +455,7 @@ void UPCGSkinnedMeshSpawnerDataProvider::CreatePrimitiveDescriptors(FPCGContext*
 {
 	PrimitiveDescriptors.Empty();
 
-	const UPCGSkinnedMeshSpawnerSettings* Settings = CastChecked<UPCGSkinnedMeshSpawnerSettings>(ProducerKernel->GetSettings());
+	const UPCGSkinnedMeshSpawnerSettings* Settings = CastChecked<UPCGSkinnedMeshSpawnerSettings>(GetProducerKernel()->GetSettings());
 
 	AActor* TargetActor = Settings->TargetActor.Get() ? Settings->TargetActor.Get() : InContext->GetTargetActor(nullptr);
 	if (!ensure(Settings->MeshSelectorParameters) || !ensure(TargetActor))
@@ -463,7 +463,7 @@ void UPCGSkinnedMeshSpawnerDataProvider::CreatePrimitiveDescriptors(FPCGContext*
 		return;
 	}
 
-	const TSharedPtr<const FPCGDataCollectionDesc> InputDataDesc = InBinding->GetCachedKernelPinDataDesc(ProducerKernel, PCGPinConstants::DefaultInputLabel, /*bIsInput=*/true);
+	const TSharedPtr<const FPCGDataCollectionDesc> InputDataDesc = InBinding->GetCachedKernelPinDataDesc(GetProducerKernel(), PCGPinConstants::DefaultInputLabel, /*bIsInput=*/true);
 
 	if (!ensure(InputDataDesc))
 	{
@@ -618,7 +618,7 @@ void UPCGSkinnedMeshSpawnerDataProvider::CreatePrimitiveDescriptors(FPCGContext*
 		PrimitiveSelectionCDF.SetNumZeroed(PrimitiveDescriptors.Num());
 	}
 
-	FPCGSpawnerPrimitives& Primitives = InBinding->MeshSpawnersToPrimitives.FindOrAdd(ProducerKernel);
+	FPCGSpawnerPrimitives& Primitives = InBinding->MeshSpawnersToPrimitives.FindOrAdd(GetProducerKernel());
 	Primitives.NumCustomFloats = CustomFloatCount;
 	Primitives.AttributeIdOffsetStrides = AttributeIdOffsetStrides;
 	Primitives.SelectorAttributeId = SelectorAttributeId;
@@ -637,7 +637,7 @@ bool UPCGSkinnedMeshSpawnerDataProvider::SetupPrimitives(FPCGContext* InContext,
 		return true;
 	}
 
-	const UPCGSkinnedMeshSpawnerSettings* Settings = CastChecked<UPCGSkinnedMeshSpawnerSettings>(ProducerKernel->GetSettings());
+	const UPCGSkinnedMeshSpawnerSettings* Settings = CastChecked<UPCGSkinnedMeshSpawnerSettings>(GetProducerKernel()->GetSettings());
 
 	AActor* TargetActor = Settings->TargetActor.Get() ? Settings->TargetActor.Get() : InContext->GetTargetActor(nullptr);
 	if (!ensure(TargetActor))
@@ -645,7 +645,7 @@ bool UPCGSkinnedMeshSpawnerDataProvider::SetupPrimitives(FPCGContext* InContext,
 		return true;
 	}
 
-	FPCGSpawnerPrimitives* Primitives = InBinding->MeshSpawnersToPrimitives.Find(ProducerKernel);
+	FPCGSpawnerPrimitives* Primitives = InBinding->MeshSpawnersToPrimitives.Find(GetProducerKernel());
 	if (!Primitives)
 	{
 		return true;
diff --git a/Engine/Plugins/PCG/Source/PCG/Private/Compute/DataInterfaces/Elements/PCGStaticMeshSpawnerDataInterface.cpp b/Engine/Plugins/PCG/Source/PCG/Private/Compute/DataInterfaces/Elements/PCGStaticMeshSpawnerDataInterface.cpp
--- a/Engine/Plugins/PCG/Source/PCG/Private/Compute/DataInterfaces/Elements/PCGStaticMeshSpawnerDataInterface.cpp
+++ b/Engine/Plugins/PCG/Source/PCG/Private/Compute/DataInterfaces/Elements/PCGStaticMeshSpawnerDataInterface.cpp
@@ -221,7 +221,7 @@ bool UPCGStaticMeshSpawnerDataProvider::PerformPreExecuteReadbacks_GameThread(UP
 		return false;
 	}
 
-	const UPCGStaticMeshSpawnerSettings* Settings = CastChecked<UPCGStaticMeshSpawnerSettings>(ProducerKernel->GetSettings());
+	const UPCGStaticMeshSpawnerSettings* Settings = CastChecked<UPCGStaticMeshSpawnerSettings>(GetProducerKernel()->GetSettings());
 
 	// Obtain index of analysis data that we need to read back.
 	if (AnalysisDataIndex == INDEX_NONE)
@@ -232,7 +232,7 @@ bool UPCGStaticMeshSpawnerDataProvider::PerformPreExecuteReadbacks_GameThread(UP
 			return true;
 		}
 
-		AnalysisDataIndex = ProducerKernel ? InBinding->GetFirstInputDataIndex(ProducerKernel, PCGStaticMeshSpawnerConstants::InstanceCountsPinLabel) : INDEX_NONE;
+		AnalysisDataIndex = GetProducerKernel() ? InBinding->GetFirstInputDataIndex(GetProducerKernel(), PCGStaticMeshSpawnerConstants::InstanceCountsPinLabel) : INDEX_NONE;
 
 		if (AnalysisDataIndex == INDEX_NONE)
 		{
@@ -303,7 +303,7 @@ bool UPCGStaticMeshSpawnerDataProvider::PrepareForExecute_GameThread(UPCGDataBin
 {
 	TRACE_CPUPROFILER_EVENT_SCOPE(UPCGStaticMeshSpawnerDataProvider::PrepareForExecute_GameThread);
 	check(InBinding);
-	const UPCGStaticMeshSpawnerSettings* Settings = CastChecked<UPCGStaticMeshSpawnerSettings>(ProducerKernel->GetSettings());
+	const UPCGStaticMeshSpawnerSettings* Settings = CastChecked<UPCGStaticMeshSpawnerSettings>(GetProducerKernel()->GetSettings());
 
 	if (!Super::PrepareForExecute_GameThread(InBinding))
 	{
@@ -312,7 +312,7 @@ bool UPCGStaticMeshSpawnerDataProvider::PrepareForExecute_GameThread(UPCGDataBin
 
 	if (!bRegisteredPrimitives)
 	{
-		InBinding->MeshSpawnersToPrimitives.Add(ProducerKernel, {});
+		InBinding->MeshSpawnersToPrimitives.Add(GetProducerKernel(), {});
 		bRegisteredPrimitives = true;
 	}
 
@@ -320,13 +320,13 @@ bool UPCGStaticMeshSpawnerDataProvider::PrepareForExecute_GameThread(UPCGDataBin
 	FPCGContext* Context = ContextHandle ? ContextHandle->GetContext() : nullptr;
 	if (!ensure(Context))
 	{
-		InBinding->CompletedMeshSpawners.Add(ProducerKernel);
+		InBinding->CompletedMeshSpawners.Add(GetProducerKernel());
 		return true;
 	}
 
 	if (!ensure(Context->ExecutionSource.IsValid()))
 	{
-		InBinding->CompletedMeshSpawners.Add(ProducerKernel);
+		InBinding->CompletedMeshSpawners.Add(GetProducerKernel());
 		return true;
 	}
 
@@ -334,11 +334,11 @@ bool UPCGStaticMeshSpawnerDataProvider::PrepareForExecute_GameThread(UPCGDataBin
 	if (SelectorAttributeId == INDEX_NONE && SelectorByAttribute)
 	{
 		const FName SelectorName = SelectorByAttribute->AttributeName;
-		const TSharedPtr<const FPCGDataCollectionDesc> InputDataDesc = InBinding->GetCachedKernelPinDataDesc(ProducerKernel, PCGPinConstants::DefaultInputLabel, /*bIsInput=*/true);
+		const TSharedPtr<const FPCGDataCollectionDesc> InputDataDesc = InBinding->GetCachedKernelPinDataDesc(GetProducerKernel(), PCGPinConstants::DefaultInputLabel, /*bIsInput=*/true);
 
 		if (!ensure(InputDataDesc))
 		{
-			InBinding->CompletedMeshSpawners.Add(ProducerKernel);
+			InBinding->CompletedMeshSpawners.Add(GetProducerKernel());
 			return true;
 		}
 
@@ -378,7 +378,7 @@ bool UPCGStaticMeshSpawnerDataProvider::PrepareForExecute_GameThread(UPCGDataBin
 					FText::FromName(SelectorName)));
 			}
 
-			InBinding->CompletedMeshSpawners.Add(ProducerKernel);
+			InBinding->CompletedMeshSpawners.Add(GetProducerKernel());
 			return true;
 		}
 
@@ -452,7 +452,7 @@ bool UPCGStaticMeshSpawnerDataProvider::PrepareForExecute_GameThread(UPCGDataBin
 		else
 		{
 			// No component set up means we have no more work to do.
-			InBinding->CompletedMeshSpawners.Add(ProducerKernel);
+			InBinding->CompletedMeshSpawners.Add(GetProducerKernel());
 			return true;
 		}
 
@@ -467,7 +467,7 @@ bool UPCGStaticMeshSpawnerDataProvider::PrepareForExecute_GameThread(UPCGDataBin
 	SelectedMeshAttributeId = InBinding->GetAttributeId(Settings->OutAttributeName, EPCGKernelAttributeType::StringKey);
 	ensure(SelectedMeshAttributeId != INDEX_NONE);
 
-	InBinding->CompletedMeshSpawners.Add(ProducerKernel);
+	InBinding->CompletedMeshSpawners.Add(GetProducerKernel());
 	return true;
 }
 
@@ -503,14 +503,14 @@ void UPCGStaticMeshSpawnerDataProvider::CreatePrimitiveDescriptors(FPCGContext*
 {
 	PrimitiveDescriptors.Empty();
 
-	const UPCGStaticMeshSpawnerSettings* Settings = CastChecked<UPCGStaticMeshSpawnerSettings>(ProducerKernel->GetSettings());
+	const UPCGStaticMeshSpawnerSettings* Settings = CastChecked<UPCGStaticMeshSpawnerSettings>(GetProducerKernel()->GetSettings());
 
 	if (!ensure(Settings->MeshSelectorParameters))
 	{
 		return;
 	}
 
-	const TSharedPtr<const FPCGDataCollectionDesc> InputDataDesc = InBinding->GetCachedKernelPinDataDesc(ProducerKernel, PCGPinConstants::DefaultInputLabel, /*bIsInput=*/true);
+	const TSharedPtr<const FPCGDataCollectionDesc> InputDataDesc = InBinding->GetCachedKernelPinDataDesc(GetProducerKernel(), PCGPinConstants::DefaultInputLabel, /*bIsInput=*/true);
 
 	if (!ensure(InputDataDesc))
 	{
@@ -740,7 +740,7 @@ void UPCGStaticMeshSpawnerDataProvider::CreatePrimitiveDescriptors(FPCGContext*
 		}
 	}
 
-	FPCGSpawnerPrimitives& Primitives = InBinding->MeshSpawnersToPrimitives.FindOrAdd(ProducerKernel);
+	FPCGSpawnerPrimitives& Primitives = InBinding->MeshSpawnersToPrimitives.FindOrAdd(GetProducerKernel());
 	Primitives.NumCustomFloats = CustomFloatCount;
 	Primitives.AttributeIdOffsetStrides = AttributeIdOffsetStrides;
 	Primitives.SelectorAttributeId = SelectorAttributeId;
@@ -759,9 +759,9 @@ bool UPCGStaticMeshSpawnerDataProvider::SetupPrimitives(FPCGContext* InContext,
 		return true;
 	}
 
-	const UPCGStaticMeshSpawnerSettings* Settings = CastChecked<UPCGStaticMeshSpawnerSettings>(ProducerKernel->GetSettings());
+	const UPCGStaticMeshSpawnerSettings* Settings = CastChecked<UPCGStaticMeshSpawnerSettings>(GetProducerKernel()->GetSettings());
 
-	FPCGSpawnerPrimitives* Primitives = InBinding->MeshSpawnersToPrimitives.Find(ProducerKernel);
+	FPCGSpawnerPrimitives* Primitives = InBinding->MeshSpawnersToPrimitives.Find(GetProducerKernel());
 	if (!Primitives)
 	{
 		return true;
diff --git a/Engine/Plugins/PCG/Source/PCG/Private/Compute/DataInterfaces/Elements/PCGTransformPointsDataInterface.cpp b/Engine/Plugins/PCG/Source/PCG/Private/Compute/DataInterfaces/Elements/PCGTransformPointsDataInterface.cpp
--- a/Engine/Plugins/PCG/Source/PCG/Private/Compute/DataInterfaces/Elements/PCGTransformPointsDataInterface.cpp
+++ b/Engine/Plugins/PCG/Source/PCG/Private/Compute/DataInterfaces/Elements/PCGTransformPointsDataInterface.cpp
@@ -143,7 +143,7 @@ UComputeDataProvider* UPCGTransformPointsDataInterface::CreateDataProvider() con
 
 FComputeDataProviderRenderProxy* UPCGTransformPointsDataProvider::GetRenderProxy()
 {
-	const UPCGTransformPointsSettings* Settings = CastChecked<UPCGTransformPointsSettings>(ProducerKernel->GetSettings());
+	const UPCGTransformPointsSettings* Settings = CastChecked<UPCGTransformPointsSettings>(GetProducerKernel()->GetSettings());
 
 	FPCGTransformPointsDataProviderProxy::FData ProxyData;
 
@@ -185,11 +185,11 @@ bool UPCGTransformPointsDataProvider::PrepareForExecute_GameThread(UPCGDataBindi
 		return true;
 	}
 
-	const UPCGTransformPointsSettings* Settings = CastChecked<UPCGTransformPointsSettings>(ProducerKernel->GetSettings());
+	const UPCGTransformPointsSettings* Settings = CastChecked<UPCGTransformPointsSettings>(GetProducerKernel()->GetSettings());
 
 	if (Settings->bApplyToAttribute && AttributeId == INDEX_NONE)
 	{
-		const TSharedPtr<const FPCGDataCollectionDesc> InputDataDesc = InBinding->GetCachedKernelPinDataDesc(ProducerKernel, PCGPinConstants::DefaultInputLabel, /*bIsInput=*/true);
+		const TSharedPtr<const FPCGDataCollectionDesc> InputDataDesc = InBinding->GetCachedKernelPinDataDesc(GetProducerKernel(), PCGPinConstants::DefaultInputLabel, /*bIsInput=*/true);
 
 		if (!ensure(InputDataDesc))
 		{
diff --git a/Engine/Plugins/PCG/Source/PCG/Private/Compute/DataInterfaces/PCGComputeDataInterface.cpp b/Engine/Plugins/PCG/Source/PCG/Private/Compute/DataInterfaces/PCGComputeDataInterface.cpp
--- a/Engine/Plugins/PCG/Source/PCG/Private/Compute/DataInterfaces/PCGComputeDataInterface.cpp
+++ b/Engine/Plugins/PCG/Source/PCG/Private/Compute/DataInterfaces/PCGComputeDataInterface.cpp
@@ -70,19 +70,19 @@ void UPCGComputeDataProvider::Initialize(const UComputeDataInterface* InDataInte
 	const UPCGComputeDataInterface* DataInterface = CastChecked<UPCGComputeDataInterface>(InDataInterface);
 	ResolvedProducerSettings = DataInterface->GetProducerSettings();
 	ProducerSettings = ResolvedProducerSettings;
-	bProducedByCPU = DataInterface->bProducedByCPU;
-	GraphBindingIndex = DataInterface->GraphBindingIndex;
-	ProducerKernel = DataInterface->ProducerKernel;
-	DownstreamInputPinLabelAliases = DataInterface->DownstreamInputPinLabelAliases;
+	bProducedByCPU = DataInterface->IsProducedByCPU();
+	GraphBindingIndex = DataInterface->GetGraphBindingIndex();
+	ProducerKernel = DataInterface->GetProducerKernel();
+	DownstreamInputPinLabelAliases = DataInterface->GetDownstreamInputPinLabelAliases();
 
 	// Set data that is only relevant for data produced on GPU, mainly to help prevent misuse.
 	if (!bProducedByCPU)
 	{
 		// The original label is needed to store inspection data.
-		OutputPinLabel = DataInterface->OutputPinLabel;
+		OutputPinLabel = DataInterface->GetOutputPinLabel();
 
 		// Use the aliased label for CPU data output as this is the output from the compute graph.
-		OutputPinLabelAlias = DataInterface->OutputPinLabelAlias;
+		OutputPinLabelAlias = DataInterface->GetOutputPinLabelAlias();
 	}
 }
 
@@ -169,12 +169,12 @@ void UPCGExportableDataProvider::Initialize(const UComputeDataInterface* InDataI
 	ExportMode = DataInterface->GetRequiresExport() ? EPCGExportMode::ComputeGraphOutput : EPCGExportMode::NoExport;
 
 	// Debug/inspection functionality for GPU-produced data.
-	if (!bProducedByCPU)
+	if (!IsProducedByCPU())
 	{
 #if WITH_EDITOR
 		// Some exportable data providers don't support inspect/debug because their producer settings are not available (e.g. GridLinkage).
 		const UPCGSettings* LocalProducerSettings = GetProducerSettings();
-		const bool bCanInspectOrDebug = LocalProducerSettings && (!ProducerKernel || !ProducerKernel->IsPinInternal(OutputPinLabel) || PCGComputeDataInterfaceHelpers::CVarEnableInternalKernelPinInspection.GetValueOnAnyThread());
+		const bool bCanInspectOrDebug = LocalProducerSettings && (!GetProducerKernel() || !GetProducerKernel()->IsPinInternal(GetOutputPinLabel()) || PCGComputeDataInterfaceHelpers::CVarEnableInternalKernelPinInspection.GetValueOnAnyThread());
 
 		if (bCanInspectOrDebug)
 		{
@@ -201,7 +201,7 @@ void UPCGExportableDataProvider::Reset()
 	ExportMode = EPCGExportMode::NoExport;
 	OnDataExported = {};
 	Binding.Reset();
-	PinDesc = nullptr;
+	PinDataDescription = nullptr;
 }
 
 bool UPCGExportableDataProvider::PrepareForExecute_GameThread(UPCGDataBinding* InBinding)
@@ -213,9 +213,9 @@ bool UPCGExportableDataProvider::PrepareForExecute_GameThread(UPCGDataBinding* I
 		return false;
 	}
 
-	if (!PinDesc)
+	if (!PinDataDescription)
 	{
-		PinDesc = InBinding->GetCachedKernelPinDataDesc(GraphBindingIndex);
+		PinDataDescription = InBinding->GetCachedKernelPinDataDesc(GetGraphBindingIndex());
 	}
 
 	return true;
@@ -223,9 +223,9 @@ bool UPCGExportableDataProvider::PrepareForExecute_GameThread(UPCGDataBinding* I
 
 void UPCGKernelParamsDataInterface::GetSupportedInputs(TArray<FShaderFunctionDefinition>& OutFunctions) const
 {
-	if (ProducerKernel)
+	if (GetProducerKernel())
 	{
-		for (const FPCGKernelOverridableParam& OverridableParam : ProducerKernel->GetCachedOverridableParams())
+		for (const FPCGKernelOverridableParam& OverridableParam : GetProducerKernel()->GetCachedOverridableParams())
 		{
 			// @todo_pcg: Broaden type support
 			OutFunctions.AddDefaulted_GetRef()
@@ -237,9 +237,9 @@ void UPCGKernelParamsDataInterface::GetSupportedInputs(TArray<FShaderFunctionDef
 
 void UPCGKernelParamsDataInterface::GetHLSL(FString& OutHLSL, FString const& InDataInterfaceName) const
 {
-	if (ProducerKernel)
+	if (GetProducerKernel())
 	{
-		for (const FPCGKernelOverridableParam& OverridableParam : ProducerKernel->GetCachedOverridableParams())
+		for (const FPCGKernelOverridableParam& OverridableParam : GetProducerKernel()->GetCachedOverridableParams())
 		{
 			// @todo_pcg: Broaden type support.
 			OutHLSL += FString::Format(TEXT(
diff --git a/Engine/Plugins/PCG/Source/PCG/Private/Compute/DataInterfaces/PCGDataCollectionDataInterface.cpp b/Engine/Plugins/PCG/Source/PCG/Private/Compute/DataInterfaces/PCGDataCollectionDataInterface.cpp
--- a/Engine/Plugins/PCG/Source/PCG/Private/Compute/DataInterfaces/PCGDataCollectionDataInterface.cpp
+++ b/Engine/Plugins/PCG/Source/PCG/Private/Compute/DataInterfaces/PCGDataCollectionDataInterface.cpp
@@ -544,7 +544,7 @@ void UPCGDataCollectionDataProvider::Initialize(const UComputeDataInterface* InD
 	Super::Initialize(InDataInterface, InBinding, InInputMask, InOutputMask);
 
 	// Initialize settings relevant for data produced on GPU.
-	if (!bProducedByCPU)
+	if (!IsProducedByCPU())
 	{
 		const UPCGDataCollectionDataInterface* DataInterface = CastChecked<UPCGDataCollectionDataInterface>(InDataInterface);
 
@@ -564,18 +564,18 @@ bool UPCGDataCollectionDataProvider::PrepareForExecute_GameThread(UPCGDataBindin
 		return false;
 	}
 
-	if (!PinDesc)
+	if (!GetPinDescription())
 	{
 		UE_LOG(LogPCG, Warning, TEXT("UPCGDataCollectionDataProvider: missing kernel pin data description. Kernel='%s', output pin label='%s'"),
-			ProducerKernel ? *ProducerKernel->GetName() : TEXT("MISSING"),
-			*OutputPinLabel.ToString());
+			GetProducerKernel() ? *GetProducerKernel()->GetName() : TEXT("MISSING"),
+			*GetOutputPinLabel().ToString());
 
 		return true;
 	}
 
 	if (!bBufferSizeValidated)
 	{
-		SizeBytes = PinDesc->ComputePackedSizeBytes();
+		SizeBytes = GetPinDescription()->ComputePackedSizeBytes();
 		bIsBufferSizeValid = !PCGComputeHelpers::IsBufferSizeTooLarge(SizeBytes);
 
 		if (!bIsBufferSizeValid)
@@ -586,16 +586,16 @@ bool UPCGDataCollectionDataProvider::PrepareForExecute_GameThread(UPCGDataBindin
 		bBufferSizeValidated = true;
 	}
 
-	if (bProducedByCPU)
+	if (IsProducedByCPU())
 	{
 		if (!DataAdaptor)
 		{
 			// Requires virtual input pin labels to pick the data items from the compute graph element's input data collection.
-			check(!DownstreamInputPinLabelAliases.IsEmpty());
-			DataAdaptor = MakeShared<FPCGDataUploadAdaptor>(InBinding, PinDesc, DownstreamInputPinLabelAliases[0]);
+			check(!GetDownstreamInputPinLabelAliases().IsEmpty());
+			DataAdaptor = MakeShared<FPCGDataUploadAdaptor>(InBinding, GetPinDescription(), GetDownstreamInputPinLabelAliases()[0]);
 		}
 
-		const bool bPreparationDone = !DataAdaptor || DataAdaptor->PrepareData_GameThread(ProducerKernel, GetProducerSettings());
+		const bool bPreparationDone = !DataAdaptor || DataAdaptor->PrepareData_GameThread(GetProducerKernel(), GetProducerSettings());
 
 #if WITH_EDITOR
 		if (bPreparationDone && DataAdaptor && DataAdaptor->IsUploadingFromCPU())
@@ -631,15 +631,15 @@ FComputeDataProviderRenderProxy* UPCGDataCollectionDataProvider::GetRenderProxy(
 	FPCGDataCollectionDataProviderProxy::FSetupParams ProxyParams =
 	{
 		.DataProvider = this,
-		.PinDesc = PinDesc,
+		.PinDesc = GetPinDescription(),
 		.SizeBytes = SizeBytes,
 		.bIsBufferSizeValid = bIsBufferSizeValid,
-		.ExportMode = ExportMode,
+		.ExportMode = GetExportMode(),
 		.bZeroInitialize = bRequiresZeroInitialization,
 		.ElementCountMultiplier = ElementCountMultiplier,
-		.OutputPinLabel = OutputPinLabel,
-		.OutputPinLabelAlias = OutputPinLabelAlias,
-		.bProducedByCPU = bProducedByCPU,
+		.OutputPinLabel = GetOutputPinLabel(),
+		.OutputPinLabelAlias = GetOutputPinLabelAlias(),
+		.bProducedByCPU = IsProducedByCPU(),
 		.DataAdaptor = DataAdaptor,
 	};
 
@@ -670,7 +670,7 @@ FPCGDataCollectionDataProviderProxy::FPCGDataCollectionDataProviderProxy(const F
 	, bProducedByCPU(InParams.bProducedByCPU)
 	, DataAdaptor(InParams.DataAdaptor)
 {
-	OriginatingGenerationCount = DataProviderWeakPtr->GenerationCounter;
+	OriginatingGenerationCount = DataProviderWeakPtr->GetGenerationCounter();
 
 	if (InParams.PinDesc)
 	{
@@ -804,7 +804,7 @@ void FPCGDataCollectionDataProviderProxy::AllocateResources(FRDGBuilder& GraphBu
 					return;
 				}
 
-				if (DataProvider->GenerationCounter != GenerationCount)
+				if (DataProvider->GetGenerationCounter() != GenerationCount)
 				{
 					return;
 				}
diff --git a/Engine/Plugins/PCG/Source/PCG/Private/Compute/DataInterfaces/PCGDebugDataInterface.cpp b/Engine/Plugins/PCG/Source/PCG/Private/Compute/DataInterfaces/PCGDebugDataInterface.cpp
--- a/Engine/Plugins/PCG/Source/PCG/Private/Compute/DataInterfaces/PCGDebugDataInterface.cpp
+++ b/Engine/Plugins/PCG/Source/PCG/Private/Compute/DataInterfaces/PCGDebugDataInterface.cpp
@@ -83,7 +83,7 @@ FComputeDataProviderRenderProxy* UPCGDebugDataProvider::GetRenderProxy()
 {
 	const UPCGSettings* Settings = GetProducerSettings();
 	const FString SettingsName = Settings ? Settings->GetName() : TEXT("MISSINGSETTINGS");
-	const FString KernelName = ProducerKernel ? ProducerKernel->GetName() : TEXT("MISSINGKERNEL");
+	const FString KernelName = GetProducerKernel() ? GetProducerKernel()->GetName() : TEXT("MISSINGKERNEL");
 
 	auto ProcessReadbackData_RenderThread = [SettingsName=SettingsName, KernelName=KernelName](const void* InData, int InNumBytes)
 	{
diff --git a/Engine/Plugins/PCG/Source/PCG/Private/Compute/DataInterfaces/PCGInstanceDataInterface.cpp b/Engine/Plugins/PCG/Source/PCG/Private/Compute/DataInterfaces/PCGInstanceDataInterface.cpp
--- a/Engine/Plugins/PCG/Source/PCG/Private/Compute/DataInterfaces/PCGInstanceDataInterface.cpp
+++ b/Engine/Plugins/PCG/Source/PCG/Private/Compute/DataInterfaces/PCGInstanceDataInterface.cpp
@@ -148,13 +148,13 @@ bool UPCGInstanceDataProvider::PrepareForExecute_GameThread(UPCGDataBinding* InB
 	TRACE_CPUPROFILER_EVENT_SCOPE(UPCGInstanceDataProvider::PrepareForExecute_GameThread);
 	UPCGDataBinding* Binding = CastChecked<UPCGDataBinding>(InBinding);
 
-	if (!Binding->CompletedMeshSpawners.Contains(ProducerKernel))
+	if (!Binding->CompletedMeshSpawners.Contains(GetProducerKernel()))
 	{
 		// The static mesh data interface(s) set this up, so wait until it is ready.
 		return false;
 	}
 
-	FPCGSpawnerPrimitives* FoundPrimitives = Binding->MeshSpawnersToPrimitives.Find(ProducerKernel);
+	FPCGSpawnerPrimitives* FoundPrimitives = Binding->MeshSpawnersToPrimitives.Find(GetProducerKernel());
 	if (!FoundPrimitives)
 	{
 		return true;
@@ -243,7 +243,7 @@ FPCGInstanceDataProviderProxy::FPCGInstanceDataProviderProxy(const FParams& InPa
 
 	if (ensure(DataProvider.IsValid()))
 	{
-		DataProviderGeneration = DataProvider->GenerationCounter;
+		DataProviderGeneration = DataProvider->GetGenerationCounter();
 
 		if (NumInstancesAllPrimitives == 0)
 		{
@@ -347,7 +347,7 @@ void FPCGInstanceDataProviderProxy::PostSubmit(FComputeContext& Context) const
 	FTSTicker::GetCoreTicker().AddTicker(TEXT("ApplyPrimitiveSceneUpdates"), /*InDelay=*/0.0f, [DataProviderWeak=DataProvider, DataProviderGeneration=DataProviderGeneration, InstanceDataExported, InstanceCustomFloatDataExported, WriteCountersExported, ContextHandle=ContextHandle](float)
 	{
 		UPCGInstanceDataProvider* DataProviderInner = DataProviderWeak.Get();
-		if (!DataProviderInner || DataProviderInner->GenerationCounter != DataProviderGeneration)
+		if (!DataProviderInner || DataProviderInner->GetGenerationCounter() != DataProviderGeneration)
 		{
 			UE_LOG(LogPCG, Verbose, TEXT("Data provider object lost, GPU instancing will fail."));
 			return false;
diff --git a/Engine/Plugins/PCG/Source/PCG/Private/Compute/DataInterfaces/PCGLandscapeDataInterface.cpp b/Engine/Plugins/PCG/Source/PCG/Private/Compute/DataInterfaces/PCGLandscapeDataInterface.cpp
--- a/Engine/Plugins/PCG/Source/PCG/Private/Compute/DataInterfaces/PCGLandscapeDataInterface.cpp
+++ b/Engine/Plugins/PCG/Source/PCG/Private/Compute/DataInterfaces/PCGLandscapeDataInterface.cpp
@@ -318,8 +318,8 @@ void UPCGLandscapeDataProvider::Initialize(const UComputeDataInterface* InDataIn
 	};
 
 	// Take any input pin label alias to obtain the data from the input data collection.
-	check(!DataInterface->DownstreamInputPinLabelAliases.IsEmpty());
-	const TArray<FPCGTaggedData> TaggedDatas = Binding->InputDataCollection.GetInputsByPin(DataInterface->DownstreamInputPinLabelAliases[0]);
+	check(!GetDownstreamInputPinLabelAliases().IsEmpty());
+	const TArray<FPCGTaggedData> TaggedDatas = Binding->InputDataCollection.GetInputsByPin(GetDownstreamInputPinLabelAliases()[0]);
 
 	if (!TaggedDatas.IsEmpty())
 	{
diff --git a/Engine/Plugins/PCG/Source/PCG/Private/Compute/DataInterfaces/PCGRawBufferDataInterface.cpp b/Engine/Plugins/PCG/Source/PCG/Private/Compute/DataInterfaces/PCGRawBufferDataInterface.cpp
--- a/Engine/Plugins/PCG/Source/PCG/Private/Compute/DataInterfaces/PCGRawBufferDataInterface.cpp
+++ b/Engine/Plugins/PCG/Source/PCG/Private/Compute/DataInterfaces/PCGRawBufferDataInterface.cpp
@@ -117,7 +117,7 @@ bool UPCGRawBufferDataProvider::PrepareForExecute_GameThread(UPCGDataBinding* In
 		return false;
 	}
 
-	if (const TSharedPtr<const FPCGDataCollectionDesc> ProducerPinDesc = InBinding->GetCachedKernelPinDataDesc(GraphBindingIndex); ensure(ProducerPinDesc))
+	if (const TSharedPtr<const FPCGDataCollectionDesc> ProducerPinDesc = InBinding->GetCachedKernelPinDataDesc(GetGraphBindingIndex()); ensure(ProducerPinDesc))
 	{
 		const uint64 PackedSizeBytes = ProducerPinDesc->ComputePackedSizeBytes();
 
@@ -127,11 +127,11 @@ bool UPCGRawBufferDataProvider::PrepareForExecute_GameThread(UPCGDataBinding* In
 		}
 	}
 
-	if (bProducedByCPU)
+	if (IsProducedByCPU())
 	{
-		check(!DownstreamInputPinLabelAliases.IsEmpty());
+		check(!GetDownstreamInputPinLabelAliases().IsEmpty());
 
-		const TArray<FPCGTaggedData> TaggedDatas = InBinding->InputDataCollection.GetInputsByPin(DownstreamInputPinLabelAliases[0]);
+		const TArray<FPCGTaggedData> TaggedDatas = InBinding->InputDataCollection.GetInputsByPin(GetDownstreamInputPinLabelAliases()[0]);
 
 		if (!TaggedDatas.IsEmpty())
 		{
@@ -157,15 +157,15 @@ FComputeDataProviderRenderProxy* UPCGRawBufferDataProvider::GetRenderProxy()
 		.SizeBytes = SizeBytes,
 		.bZeroInitialize = bZeroInitialize,
 		.DataToUpload = DataToUpload,
-		.ExportMode = ExportMode,
-		.OutputPinLabel = OutputPinLabel,
-		.OutputPinLabelAlias = OutputPinLabelAlias,
+		.ExportMode = GetExportMode(),
+		.OutputPinLabel = GetOutputPinLabel(),
+		.OutputPinLabelAlias = GetOutputPinLabelAlias(),
 		.DataProviderWeakPtr = ThisWeak,
 	};
 
 	ensure(ReadbackState);
 
-	Params.AsyncReadbackCallback_RenderThread = [ReadbackState=ReadbackState, DataProviderWeakPtr=ThisWeak, GenerationCount=GenerationCounter.load(), ExportMode=ExportMode, OutputPinLabel=OutputPinLabel, OutputPinLabelAlias=OutputPinLabelAlias, BindingWeakPtr=Binding](const void* InData, int InNumBytes)
+	Params.AsyncReadbackCallback_RenderThread = [ReadbackState=ReadbackState, DataProviderWeakPtr=ThisWeak, GenerationCount=GetGenerationCounter(), ExportMode=GetExportMode(), OutputPinLabel = GetOutputPinLabel(), OutputPinLabelAlias = GetOutputPinLabelAlias(), BindingWeakPtr = TWeakObjectPtr<UPCGDataBinding>(GetDataBinding())](const void* InData, int InNumBytes)
 	{
 		if (!InData || !ReadbackState)
 		{
@@ -189,7 +189,7 @@ FComputeDataProviderRenderProxy* UPCGRawBufferDataProvider::GetRenderProxy()
 				return;
 			}
 
-			if (DataProvider->GenerationCounter != GenerationCount)
+			if (DataProvider->GetGenerationCounter() != GenerationCount)
 			{
 				// Data provider was returned to the pool, possibly due to cancellation. Nothing to do here.
 				return;
diff --git a/Engine/Plugins/PCG/Source/PCG/Private/Compute/DataInterfaces/PCGStaticMeshDataInterface.cpp b/Engine/Plugins/PCG/Source/PCG/Private/Compute/DataInterfaces/PCGStaticMeshDataInterface.cpp
--- a/Engine/Plugins/PCG/Source/PCG/Private/Compute/DataInterfaces/PCGStaticMeshDataInterface.cpp
+++ b/Engine/Plugins/PCG/Source/PCG/Private/Compute/DataInterfaces/PCGStaticMeshDataInterface.cpp
@@ -159,10 +159,10 @@ bool UPCGStaticMeshDataProvider::PrepareForExecute_GameThread(UPCGDataBinding* I
 	if (!LoadHandle.IsValid())
 	{
 		// Take any input pin label alias to obtain the data from the input data collection.
-		check(!DownstreamInputPinLabels.IsEmpty());
+		check(!GetDownstreamInputPinLabelAliases().IsEmpty());
 
 		const UPCGDataBinding* Binding = CastChecked<UPCGDataBinding>(InBinding);
-		const TArray<FPCGTaggedData> TaggedDatas = Binding->InputDataCollection.GetInputsByPin(DownstreamInputPinLabels[0]);
+		const TArray<FPCGTaggedData> TaggedDatas = Binding->InputDataCollection.GetInputsByPin(GetDownstreamInputPinLabelAliases()[0]);
 
 		if (!TaggedDatas.IsEmpty())
 		{
@@ -190,19 +190,6 @@ bool UPCGStaticMeshDataProvider::PrepareForExecute_GameThread(UPCGDataBinding* I
 	return false;
 }
 
-void UPCGStaticMeshDataProvider::Initialize(const UComputeDataInterface* InDataInterface, UObject* InBinding, uint64 InInputMask, uint64 InOutputMask)
-{
-	TRACE_CPUPROFILER_EVENT_SCOPE(UPCGStaticMeshDataProvider::Initialize);
-
-	Super::Initialize(InDataInterface, InBinding, InInputMask, InOutputMask);
-
-	const UPCGStaticMeshDataInterface* DataInterface = CastChecked<UPCGStaticMeshDataInterface>(InDataInterface);
-
-	// Take any input pin label alias to obtain the data from the input data collection.
-	check(!DataInterface->DownstreamInputPinLabelAliases.IsEmpty());
-	DownstreamInputPinLabels = DataInterface->DownstreamInputPinLabelAliases;
-}
-
 FComputeDataProviderRenderProxy* UPCGStaticMeshDataProvider::GetRenderProxy()
 {
 	return new FPCGStaticMeshDataProviderProxy(LoadedStaticMesh.Get());
@@ -211,7 +198,6 @@ FComputeDataProviderRenderProxy* UPCGStaticMeshDataProvider::GetRenderProxy()
 void UPCGStaticMeshDataProvider::Reset()
 {
 	LoadedStaticMesh = nullptr;
-	DownstreamInputPinLabels.Empty();
 	LoadHandle.Reset();
 
 	Super::Reset();
diff --git a/Engine/Plugins/PCG/Source/PCG/Private/Compute/DataInterfaces/PCGStaticMeshDataInterface.h b/Engine/Plugins/PCG/Source/PCG/Private/Compute/DataInterfaces/PCGStaticMeshDataInterface.h
--- a/Engine/Plugins/PCG/Source/PCG/Private/Compute/DataInterfaces/PCGStaticMeshDataInterface.h
+++ b/Engine/Plugins/PCG/Source/PCG/Private/Compute/DataInterfaces/PCGStaticMeshDataInterface.h
@@ -44,7 +44,6 @@ class UPCGStaticMeshDataProvider : public UPCGComputeDataProvider
 
 public:
 	//~ Begin UComputeDataProvider Interface
-	virtual void Initialize(const UComputeDataInterface* InDataInterface, UObject* InBinding, uint64 InInputMask, uint64 InOutputMask) override;
 	FComputeDataProviderRenderProxy* GetRenderProxy() override;
 	virtual void Reset() override;
 	//~ End UComputeDataProvider Interface
@@ -56,7 +55,6 @@ public:
 	UPROPERTY()
 	TObjectPtr<const UStaticMesh> LoadedStaticMesh = nullptr;
 
-	TArray<FName> DownstreamInputPinLabels;
 	TSharedPtr<FStreamableHandle> LoadHandle = nullptr;
 };
 
diff --git a/Engine/Plugins/PCG/Source/PCG/Private/Compute/DataInterfaces/PCGTextureDataInterface.cpp b/Engine/Plugins/PCG/Source/PCG/Private/Compute/DataInterfaces/PCGTextureDataInterface.cpp
--- a/Engine/Plugins/PCG/Source/PCG/Private/Compute/DataInterfaces/PCGTextureDataInterface.cpp
+++ b/Engine/Plugins/PCG/Source/PCG/Private/Compute/DataInterfaces/PCGTextureDataInterface.cpp
@@ -297,8 +297,8 @@ void UPCGTextureDataProvider::BuildInfosFromDataCollection(UPCGDataBinding* InBi
 	check(InBinding);
 
 	// Take any input pin label alias to obtain the data from the input data collection.
-	check(!DownstreamInputPinLabelAliases.IsEmpty());
-	const TArray<FPCGTaggedData> InputTaggedData = InBinding->InputDataCollection.GetInputsByPin(DownstreamInputPinLabelAliases[0]);
+	check(!GetDownstreamInputPinLabelAliases().IsEmpty());
+	const TArray<FPCGTaggedData> InputTaggedData = InBinding->InputDataCollection.GetInputsByPin(GetDownstreamInputPinLabelAliases()[0]);
 
 	for (int I = 0; I < InputTaggedData.Num(); ++I)
 	{
@@ -323,7 +323,7 @@ void UPCGTextureDataProvider::BuildInfosFromDataCollection(UPCGDataBinding* InBi
 			else
 			{
 				UE_LOG(LogPCG, Warning, TEXT("Texture data interface on pin '%s' received too many textures to bind. Only the first %d textures will be bound."),
-					*DownstreamInputPinLabelAliases[0].ToString(),
+					*GetDownstreamInputPinLabelAliases()[0].ToString(),
 					PCGTextureDataInterfaceConstants::MAX_NUM_SRV_BINDINGS);
 				continue;
 			}
@@ -338,7 +338,6 @@ void UPCGTextureDataProvider::BuildInfosFromDataCollection(UPCGDataBinding* InBi
 void UPCGTextureDataProvider::BuildInfosFromDataDescription(UPCGDataBinding* InBinding)
 {
 	check(InBinding);
-	check(PinDesc);
 
 	TSharedPtr<FPCGContextHandle> ContextHandle = InBinding->ContextHandle.Pin();
 	FPCGContext* Context = ContextHandle ? ContextHandle->GetContext() : nullptr;
@@ -348,6 +347,9 @@ void UPCGTextureDataProvider::BuildInfosFromDataDescription(UPCGDataBinding* InB
 		return;
 	}
 
+	TSharedPtr<const FPCGDataCollectionDesc> PinDesc = GetPinDescription();
+	check(PinDesc);
+
 	const FBox OriginalLocalBounds = Context->ExecutionSource->GetExecutionState().GetOriginalLocalSpaceBounds();
 	FTransform Transform = Context->ExecutionSource->GetExecutionState().GetOriginalTransform();
 	Transform.SetScale3D(Transform.GetScale3D() * 0.5 * (OriginalLocalBounds.Max - OriginalLocalBounds.Min));
@@ -368,7 +370,7 @@ void UPCGTextureDataProvider::BuildInfosFromDataDescription(UPCGDataBinding* InB
 			else
 			{
 				UE_LOG(LogPCG, Warning, TEXT("Texture data interface on pin '%s' received too many texture UAVs to bind. Only the first %d texture UAVs will be bound."),
-					*OutputPinLabel.ToString(),
+					*GetOutputPinLabel().ToString(),
 					PCGTextureDataInterfaceConstants::MAX_NUM_UAV_BINDINGS);
 				continue;
 			}
@@ -388,9 +390,9 @@ FPCGTextureDataProviderProxy::FPCGTextureDataProviderProxy(TWeakObjectPtr<UPCGTe
 	BindingInfos = InDataProvider->GetBindingInfos();
 	TextureInfos = InDataProvider->GetTextureInfos();
 	ExportMode = InDataProvider->GetExportMode();
-	OutputPinLabel = InDataProvider->OutputPinLabel;
-	OutputPinLabelAlias = InDataProvider->OutputPinLabelAlias;
-	OriginatingGenerationCount = InDataProvider->GenerationCounter;
+	OutputPinLabel = InDataProvider->GetOutputPinLabel();
+	OutputPinLabelAlias = InDataProvider->GetOutputPinLabelAlias();
+	OriginatingGenerationCount = InDataProvider->GetGenerationCounter();
 	DataProviderWeakPtr_GT = InDataProvider;
 	PinDesc = InDataProvider->GetPinDescription();
 }
@@ -604,7 +606,7 @@ void FPCGTextureDataProviderProxy::ExportTextureUAVs(const TArray<TRefCountPtr<I
 			return;
 		}
 
-		if (DataProvider->GenerationCounter != GenerationCount)
+		if (DataProvider->GetGenerationCounter() != GenerationCount)
 		{
 			return;
 		}
diff --git a/Engine/Plugins/PCG/Source/PCG/Private/Compute/DataInterfaces/PCGVirtualTextureDataInterface.cpp b/Engine/Plugins/PCG/Source/PCG/Private/Compute/DataInterfaces/PCGVirtualTextureDataInterface.cpp
--- a/Engine/Plugins/PCG/Source/PCG/Private/Compute/DataInterfaces/PCGVirtualTextureDataInterface.cpp
+++ b/Engine/Plugins/PCG/Source/PCG/Private/Compute/DataInterfaces/PCGVirtualTextureDataInterface.cpp
@@ -133,14 +133,14 @@ void UPCGVirtualTextureDataProvider::Initialize(const UComputeDataInterface* InD
 	const UPCGComponent* SourceComponent = Cast<UPCGComponent>(Binding->GetExecutionSource());
 
 	// Take any input pin label alias to obtain the data from the input data collection.
-	check(!DataInterface->DownstreamInputPinLabelAliases.IsEmpty());
-	const TArray<FPCGTaggedData> InputTaggedData = Binding->InputDataCollection.GetInputsByPin(DataInterface->DownstreamInputPinLabelAliases[0]);
+	check(!GetDownstreamInputPinLabelAliases().IsEmpty());
+	const TArray<FPCGTaggedData> InputTaggedData = Binding->InputDataCollection.GetInputsByPin(GetDownstreamInputPinLabelAliases()[0]);
 
 	// Sampling of virtual textures is not supported outside of runtime generation at this time, because the RuntimeGenScheduler is responsible for priming the virtual textures.
 	if (!SourceComponent || !SourceComponent->IsManagedByRuntimeGenSystem())
 	{
 		UE_LOG(LogPCG, Error, TEXT("Virtual texture data interface on pin '%s' failed to initialize. Virtual texture sampling is only supported for runtime generation."),
-			*DataInterface->DownstreamInputPinLabelAliases[0].ToString());
+			*GetDownstreamInputPinLabelAliases()[0].ToString());
 		return;
 	}
 
@@ -151,7 +151,7 @@ void UPCGVirtualTextureDataProvider::Initialize(const UComputeDataInterface* InD
 		if (!ensure(VirtualTextureData))
 		{
 			UE_LOG(LogPCG, Error, TEXT("Virtual texture data interface on pin '%s' receieved unsupported data type: '%s'"),
-				*DataInterface->DownstreamInputPinLabelAliases[0].ToString(),
+				*GetDownstreamInputPinLabelAliases()[0].ToString(),
 				TaggedData.Data ? *TaggedData.Data->GetName() : TEXT("NULL"));
 			continue;
 		}
@@ -163,7 +163,7 @@ void UPCGVirtualTextureDataProvider::Initialize(const UComputeDataInterface* InD
 		else
 		{
 			UE_LOG(LogPCG, Warning, TEXT("Texture data interface on pin '%s' received too many textures to bind. Only the first %d textures will be bound."),
-				*DataInterface->DownstreamInputPinLabelAliases[0].ToString(),
+				*GetDownstreamInputPinLabelAliases()[0].ToString(),
 				PCGVirtualTextureDIConstants::MAX_NUM_BINDINGS);
 			break;
 		}
diff --git a/Engine/Plugins/PCG/Source/PCG/Private/Compute/DataInterfaces/PCGVirtualTextureDataInterface.h b/Engine/Plugins/PCG/Source/PCG/Private/Compute/DataInterfaces/PCGVirtualTextureDataInterface.h
--- a/Engine/Plugins/PCG/Source/PCG/Private/Compute/DataInterfaces/PCGVirtualTextureDataInterface.h
+++ b/Engine/Plugins/PCG/Source/PCG/Private/Compute/DataInterfaces/PCGVirtualTextureDataInterface.h
@@ -51,7 +51,7 @@ private:
 
 /** Compute Framework Data Provider for reading a virtual texture. */
 UCLASS()
-class UPCGVirtualTextureDataProvider : public UComputeDataProvider
+class UPCGVirtualTextureDataProvider : public UPCGComputeDataProvider
 {
 	GENERATED_BODY()
 
diff --git a/Engine/Plugins/PCG/Source/PCG/Private/Compute/Elements/PCGComputeGraphElement.cpp b/Engine/Plugins/PCG/Source/PCG/Private/Compute/Elements/PCGComputeGraphElement.cpp
--- a/Engine/Plugins/PCG/Source/PCG/Private/Compute/Elements/PCGComputeGraphElement.cpp
+++ b/Engine/Plugins/PCG/Source/PCG/Private/Compute/Elements/PCGComputeGraphElement.cpp
@@ -267,7 +267,7 @@ bool FPCGComputeGraphElement::ExecuteInternal(FPCGContext* InContext) const
 
 				TWeakObjectPtr<UPCGExportableDataProvider> DataProviderWeak = DataProvider;
 
-				const uint64 OriginatingGenerationCount = DataProvider->GenerationCounter;
+				const uint64 OriginatingGenerationCount = DataProvider->GetGenerationCounter();
 
 				DataProvider->OnDataExported_GameThread().AddLambda([ContextHandle = Context->GetOrCreateHandle(), DataProviderWeak, GenerationCount = OriginatingGenerationCount]()
 				{
@@ -277,7 +277,7 @@ bool FPCGComputeGraphElement::ExecuteInternal(FPCGContext* InContext) const
 					FPCGComputeGraphContext* ContextPtr = static_cast<FPCGComputeGraphContext*>(SharedHandle ? SharedHandle->GetContext() : nullptr);
 					UPCGExportableDataProvider* DataProvider = DataProviderWeak.Get();
 
-					if (!DataProvider || !ContextPtr || GenerationCount != DataProvider->GenerationCounter)
+					if (!DataProvider || !ContextPtr || GenerationCount != DataProvider->GetGenerationCounter())
 					{
 						// Safe to just jump out. The GPU buffer is ref counted.
 						return;
diff --git a/Engine/Plugins/PCG/Source/PCG/Private/Compute/Elements/PCGCullPointsOutsideActorBoundsKernel.cpp b/Engine/Plugins/PCG/Source/PCG/Private/Compute/Elements/PCGCullPointsOutsideActorBoundsKernel.cpp
--- a/Engine/Plugins/PCG/Source/PCG/Private/Compute/Elements/PCGCullPointsOutsideActorBoundsKernel.cpp
+++ b/Engine/Plugins/PCG/Source/PCG/Private/Compute/Elements/PCGCullPointsOutsideActorBoundsKernel.cpp
@@ -40,7 +40,7 @@ void UPCGCullPointsOutsideActorBoundsKernel::CreateAdditionalInputDataInterfaces
 	Super::CreateAdditionalInputDataInterfaces(InOutContext, InObjectOuter, OutDataInterfaces);
 
 	TObjectPtr<UPCGCullPointsOutsideActorBoundsDataInterface> NodeDI = InOutContext.NewObject_AnyThread<UPCGCullPointsOutsideActorBoundsDataInterface>(InObjectOuter);
-	NodeDI->ProducerKernel = this;
+	NodeDI->SetProducerKernel(this);
 
 	OutDataInterfaces.Add(NodeDI);
 }
diff --git a/Engine/Plugins/PCG/Source/PCG/Private/Compute/Elements/PCGCustomHLSLKernel.cpp b/Engine/Plugins/PCG/Source/PCG/Private/Compute/Elements/PCGCustomHLSLKernel.cpp
--- a/Engine/Plugins/PCG/Source/PCG/Private/Compute/Elements/PCGCustomHLSLKernel.cpp
+++ b/Engine/Plugins/PCG/Source/PCG/Private/Compute/Elements/PCGCustomHLSLKernel.cpp
@@ -611,7 +611,7 @@ void UPCGCustomHLSLKernel::CreateAdditionalInputDataInterfaces(FPCGGPUCompilatio
 	Super::CreateAdditionalInputDataInterfaces(InOutContext, InObjectOuter, OutDataInterfaces);
 
 	TObjectPtr<UPCGCustomHLSLDataInterface> NodeDI = InOutContext.NewObject_AnyThread<UPCGCustomHLSLDataInterface>(InObjectOuter);
-	NodeDI->ProducerKernel = this;
+	NodeDI->SetProducerKernel(this);
 
 	OutDataInterfaces.Add(NodeDI);
 }
diff --git a/Engine/Plugins/PCG/Source/PCG/Private/Compute/Elements/PCGNormalToDensityKernel.cpp b/Engine/Plugins/PCG/Source/PCG/Private/Compute/Elements/PCGNormalToDensityKernel.cpp
--- a/Engine/Plugins/PCG/Source/PCG/Private/Compute/Elements/PCGNormalToDensityKernel.cpp
+++ b/Engine/Plugins/PCG/Source/PCG/Private/Compute/Elements/PCGNormalToDensityKernel.cpp
@@ -57,7 +57,7 @@ void UPCGNormalToDensityKernel::CreateAdditionalInputDataInterfaces(FPCGGPUCompi
 	Super::CreateAdditionalInputDataInterfaces(InOutContext, InObjectOuter, OutDataInterfaces);
 
 	TObjectPtr<UPCGNormalToDensityDataInterface> NodeDI = InOutContext.NewObject_AnyThread<UPCGNormalToDensityDataInterface>(InObjectOuter);
-	NodeDI->ProducerKernel = this;
+	NodeDI->SetProducerKernel(this);
 
 	OutDataInterfaces.Add(NodeDI);
 }
diff --git a/Engine/Plugins/PCG/Source/PCG/Private/Compute/Elements/PCGTransformPointsKernel.cpp b/Engine/Plugins/PCG/Source/PCG/Private/Compute/Elements/PCGTransformPointsKernel.cpp
--- a/Engine/Plugins/PCG/Source/PCG/Private/Compute/Elements/PCGTransformPointsKernel.cpp
+++ b/Engine/Plugins/PCG/Source/PCG/Private/Compute/Elements/PCGTransformPointsKernel.cpp
@@ -61,7 +61,7 @@ void UPCGTransformPointsKernel::CreateAdditionalInputDataInterfaces(FPCGGPUCompi
 	Super::CreateAdditionalInputDataInterfaces(InOutContext, InObjectOuter, OutDataInterfaces);
 
 	TObjectPtr<UPCGTransformPointsDataInterface> NodeDI = InOutContext.NewObject_AnyThread<UPCGTransformPointsDataInterface>(InObjectOuter);
-	NodeDI->ProducerKernel = this;
+	NodeDI->SetProducerKernel(this);
 
 	OutDataInterfaces.Add(NodeDI);
 }
diff --git a/Engine/Plugins/PCG/Source/PCG/Private/Elements/Metadata/PCGMetadataPartitionKernel.cpp b/Engine/Plugins/PCG/Source/PCG/Private/Elements/Metadata/PCGMetadataPartitionKernel.cpp
--- a/Engine/Plugins/PCG/Source/PCG/Private/Elements/Metadata/PCGMetadataPartitionKernel.cpp
+++ b/Engine/Plugins/PCG/Source/PCG/Private/Elements/Metadata/PCGMetadataPartitionKernel.cpp
@@ -207,7 +207,7 @@ void UPCGMetadataPartitionKernel::CreateAdditionalInputDataInterfaces(FPCGGPUCom
 	Super::CreateAdditionalInputDataInterfaces(InOutContext, InObjectOuter, OutDataInterfaces);
 
 	TObjectPtr<UPCGMetadataPartitionDataInterface> NodeDI = InOutContext.NewObject_AnyThread<UPCGMetadataPartitionDataInterface>(InObjectOuter);
-	NodeDI->ProducerKernel = this;
+	NodeDI->SetProducerKernel(this);
 	OutDataInterfaces.Add(NodeDI);
 }
 #endif
diff --git a/Engine/Plugins/PCG/Source/PCG/Private/Elements/PCGCopyPointsAnalysisKernel.cpp b/Engine/Plugins/PCG/Source/PCG/Private/Elements/PCGCopyPointsAnalysisKernel.cpp
--- a/Engine/Plugins/PCG/Source/PCG/Private/Elements/PCGCopyPointsAnalysisKernel.cpp
+++ b/Engine/Plugins/PCG/Source/PCG/Private/Elements/PCGCopyPointsAnalysisKernel.cpp
@@ -96,7 +96,7 @@ void UPCGCopyPointsAnalysisKernel::CreateAdditionalInputDataInterfaces(FPCGGPUCo
 	Super::CreateAdditionalInputDataInterfaces(InOutContext, InObjectOuter, OutDataInterfaces);
 
 	TObjectPtr<UPCGCopyPointsAnalysisDataInterface> NodeDI = InOutContext.NewObject_AnyThread<UPCGCopyPointsAnalysisDataInterface>(InObjectOuter);
-	NodeDI->ProducerKernel = this;
+	NodeDI->SetProducerKernel(this);
 
 	OutDataInterfaces.Add(NodeDI);
 }
diff --git a/Engine/Plugins/PCG/Source/PCG/Private/Elements/PCGCopyPointsKernel.cpp b/Engine/Plugins/PCG/Source/PCG/Private/Elements/PCGCopyPointsKernel.cpp
--- a/Engine/Plugins/PCG/Source/PCG/Private/Elements/PCGCopyPointsKernel.cpp
+++ b/Engine/Plugins/PCG/Source/PCG/Private/Elements/PCGCopyPointsKernel.cpp
@@ -229,7 +229,7 @@ void UPCGCopyPointsKernel::CreateAdditionalInputDataInterfaces(FPCGGPUCompilatio
 	Super::CreateAdditionalInputDataInterfaces(InOutContext, InObjectOuter, OutDataInterfaces);
 
 	TObjectPtr<UPCGCopyPointsDataInterface> NodeDI = InOutContext.NewObject_AnyThread<UPCGCopyPointsDataInterface>(InObjectOuter);
-	NodeDI->ProducerKernel = this;
+	NodeDI->SetProducerKernel(this);
 
 	OutDataInterfaces.Add(NodeDI);
 }
diff --git a/Engine/Plugins/PCG/Source/PCG/Private/Elements/PCGSkinnedMeshSpawnerKernel.cpp b/Engine/Plugins/PCG/Source/PCG/Private/Elements/PCGSkinnedMeshSpawnerKernel.cpp
--- a/Engine/Plugins/PCG/Source/PCG/Private/Elements/PCGSkinnedMeshSpawnerKernel.cpp
+++ b/Engine/Plugins/PCG/Source/PCG/Private/Elements/PCGSkinnedMeshSpawnerKernel.cpp
@@ -68,7 +68,7 @@ void UPCGSkinnedMeshSpawnerKernel::CreateAdditionalInputDataInterfaces(FPCGGPUCo
 	Super::CreateAdditionalInputDataInterfaces(InOutContext, InObjectOuter, OutDataInterfaces);
 
 	TObjectPtr<UPCGSkinnedMeshSpawnerDataInterface> NodeDI = InOutContext.NewObject_AnyThread<UPCGSkinnedMeshSpawnerDataInterface>(InObjectOuter);
-	NodeDI->ProducerKernel = this;
+	NodeDI->SetProducerKernel(this);
 	OutDataInterfaces.Add(NodeDI);
 }
 
@@ -77,7 +77,7 @@ void UPCGSkinnedMeshSpawnerKernel::CreateAdditionalOutputDataInterfaces(FPCGGPUC
 	Super::CreateAdditionalOutputDataInterfaces(InOutContext, InObjectOuter, OutDataInterfaces);
 
 	UPCGInstanceDataInterface* InstanceDI = InOutContext.NewObject_AnyThread<UPCGInstanceDataInterface>(InObjectOuter);
-	InstanceDI->ProducerKernel = this;
+	InstanceDI->SetProducerKernel(this);
 	InstanceDI->InputPinProvidingData = PCGPinConstants::DefaultInputLabel;
 	OutDataInterfaces.Add(InstanceDI);
 }
diff --git a/Engine/Plugins/PCG/Source/PCG/Private/Elements/PCGStaticMeshSpawnerKernel.cpp b/Engine/Plugins/PCG/Source/PCG/Private/Elements/PCGStaticMeshSpawnerKernel.cpp
--- a/Engine/Plugins/PCG/Source/PCG/Private/Elements/PCGStaticMeshSpawnerKernel.cpp
+++ b/Engine/Plugins/PCG/Source/PCG/Private/Elements/PCGStaticMeshSpawnerKernel.cpp
@@ -127,7 +127,7 @@ void UPCGStaticMeshSpawnerKernel::CreateAdditionalInputDataInterfaces(FPCGGPUCom
 	Super::CreateAdditionalInputDataInterfaces(InOutContext, InObjectOuter, OutDataInterfaces);
 
 	TObjectPtr<UPCGStaticMeshSpawnerDataInterface> NodeDI = InOutContext.NewObject_AnyThread<UPCGStaticMeshSpawnerDataInterface>(InObjectOuter);
-	NodeDI->ProducerKernel = this;
+	NodeDI->SetProducerKernel(this);
 	OutDataInterfaces.Add(NodeDI);
 }
 
@@ -136,7 +136,7 @@ void UPCGStaticMeshSpawnerKernel::CreateAdditionalOutputDataInterfaces(FPCGGPUCo
 	Super::CreateAdditionalOutputDataInterfaces(InOutContext, InObjectOuter, OutDataInterfaces);
 
 	UPCGInstanceDataInterface* InstanceDI = InOutContext.NewObject_AnyThread<UPCGInstanceDataInterface>(InObjectOuter);
-	InstanceDI->ProducerKernel = this;
+	InstanceDI->SetProducerKernel(this);
 	InstanceDI->InputPinProvidingData = PCGPinConstants::DefaultInputLabel;
 	OutDataInterfaces.Add(InstanceDI);
 }
diff --git a/Engine/Plugins/PCG/Source/PCG/Private/Graph/PCGGraphCompilerGPU.cpp b/Engine/Plugins/PCG/Source/PCG/Private/Graph/PCGGraphCompilerGPU.cpp
--- a/Engine/Plugins/PCG/Source/PCG/Private/Graph/PCGGraphCompilerGPU.cpp
+++ b/Engine/Plugins/PCG/Source/PCG/Private/Graph/PCGGraphCompilerGPU.cpp
@@ -990,7 +990,7 @@ void FPCGGraphCompilerGPU::CreateDataInterfaces(
 
 		if (DataInterface)
 		{
-			DataInterface->ProducerKernel = ProducerKernel;
+			DataInterface->SetProducerKernel(ProducerKernel);
 			DataInterface->SetProducerSettings(ProducerSettings);
 			DataInterface->SetOutputPin(PinProperties.Label);
 			DataInterface->SetProducedByCPU(Params.bProducedByCPU);
@@ -1429,7 +1429,7 @@ void FPCGGraphCompilerGPU::CompileComputeGraph(
 			UPCGDebugDataInterface* DebugDI = InOutContext.NewObject_AnyThread<UPCGDebugDataInterface>(InOutComputeGraph);
 			DebugDI->SetDebugBufferSize(Settings->DebugBufferSize);
 			DebugDI->SetProducerSettings(Settings);
-			DebugDI->ProducerKernel = KernelWithBindings.Kernel;
+			DebugDI->SetProducerKernel(KernelWithBindings.Kernel);
 
 			const int32 DebugDIIndex = InOutComputeGraph->DataInterfaces.Num();
 			InOutComputeGraph->DataInterfaces.Add(DebugDI);
@@ -1643,9 +1643,9 @@ void FPCGGraphCompilerGPU::CompileComputeGraph(
 						}
 
 						UPCGComputeDataInterface* PCGDataInterface = Cast<UPCGComputeDataInterface>(DataInterface);
-						if (PCGDataInterface && PCGDataInterface->GraphBindingIndex == INDEX_NONE)
+						if (PCGDataInterface && PCGDataInterface->GetGraphBindingIndex() == INDEX_NONE)
 						{
-							PCGDataInterface->GraphBindingIndex = GraphBindingIndex;
+							PCGDataInterface->SetGraphBindingIndex(GraphBindingIndex);
 						}
 					}
 				}
diff --git a/Engine/Plugins/PCG/Source/PCG/Public/Compute/DataInterfaces/PCGComputeDataInterface.h b/Engine/Plugins/PCG/Source/PCG/Public/Compute/DataInterfaces/PCGComputeDataInterface.h
--- a/Engine/Plugins/PCG/Source/PCG/Public/Compute/DataInterfaces/PCGComputeDataInterface.h
+++ b/Engine/Plugins/PCG/Source/PCG/Public/Compute/DataInterfaces/PCGComputeDataInterface.h
@@ -25,19 +25,27 @@ public:
 	PCG_API virtual void SetOutputPin(FName InOutputPinLabel, const FName* InOptionalOutputPinLabelAlias = nullptr);
 
 	void SetProducerSettings(const UPCGSettings* InProducerSettings);
-	const UPCGSettings* GetProducerSettings() const;
+	PCG_API const UPCGSettings* GetProducerSettings() const;
 
-	PCG_API static void AddReferencedObjects(UObject* InThis, FReferenceCollector& Collector);
+	void SetProducerKernel(TObjectPtr<const UPCGComputeKernel> InProducerKernel) { ProducerKernel = InProducerKernel; }
+	TObjectPtr<const UPCGComputeKernel> GetProducerKernel() const { return ProducerKernel; }
 
 	/** Set whether data originates from the CPU and should be uploaded to the GPU. */
 	void SetProducedByCPU(bool bInProducedByCPU) { bProducedByCPU = bInProducedByCPU; }
+	bool IsProducedByCPU() const { return bProducedByCPU; }
 
-public:
-	UPROPERTY()
-	int32 GraphBindingIndex = INDEX_NONE;
+	void SetGraphBindingIndex(int32 InGraphBindingIndex) { GraphBindingIndex = InGraphBindingIndex; }
+	int32 GetGraphBindingIndex() const { return GraphBindingIndex; }
 
+	FName GetOutputPinLabel() const { return OutputPinLabel; }
+	FName GetOutputPinLabelAlias() const { return OutputPinLabelAlias; }
+	const TArray<FName>& GetDownstreamInputPinLabelAliases() const { return DownstreamInputPinLabelAliases; }
+
+	PCG_API static void AddReferencedObjects(UObject* InThis, FReferenceCollector& Collector);
+
+private:
 	UPROPERTY()
-	TObjectPtr<const UPCGComputeKernel> ProducerKernel;
+	int32 GraphBindingIndex = INDEX_NONE;
 
 	/** Label of output pin that this data interface is associated with. */
 	UPROPERTY()
@@ -53,7 +61,9 @@ public:
 	UPROPERTY()
 	bool bProducedByCPU = false;
 
-protected:
+	UPROPERTY()
+	TObjectPtr<const UPCGComputeKernel> ProducerKernel;
+
 	UPROPERTY()
 	TSoftObjectPtr<const UPCGSettings> ProducerSettings;
 
@@ -70,47 +80,55 @@ public:
 	PCG_API virtual void Reset() override;
 
 	/** Called before PrepareForExecute to allow data providers to do any readbacks and finalize data descriptions. */
-	PCG_API virtual bool PerformPreExecuteReadbacks_GameThread(UPCGDataBinding* InBinding) { return true; }
+	virtual bool PerformPreExecuteReadbacks_GameThread(UPCGDataBinding* InBinding) { return true; }
 
 	/* Called prior to GetRenderProxy and execution. Any processing of data descriptions should be done here (after PerformPreExecuteReadbacks has been called). */
-	PCG_API virtual bool PrepareForExecute_GameThread(UPCGDataBinding* InBinding) { return true; }
+	virtual bool PrepareForExecute_GameThread(UPCGDataBinding* InBinding) { return true; }
 
 	/** Return true when all done. */
-	PCG_API virtual bool PostExecute(UPCGDataBinding* InBinding) { return true; }
+	virtual bool PostExecute(UPCGDataBinding* InBinding) { return true; }
 
 	/** Release any handles to transient resources like GPU buffers. */
 	virtual void ReleaseTransientResources(const TCHAR* InReason = nullptr) {}
 
-	const UPCGSettings* GetProducerSettings() const;
+	PCG_API const UPCGSettings* GetProducerSettings() const;
 	void SetProducerSettings(const UPCGSettings* InSettings);
 
+	TObjectPtr<const UPCGComputeKernel> GetProducerKernel() const { return ProducerKernel; }
+
+	/** Set whether data originates from the CPU and should be uploaded to the GPU. */
+	bool IsProducedByCPU() const { return bProducedByCPU; }
+	int32 GetGraphBindingIndex() const { return GraphBindingIndex; }
+	FName GetOutputPinLabel() const { return OutputPinLabel; }
+	FName GetOutputPinLabelAlias() const { return OutputPinLabelAlias; }
+	const TArray<FName>& GetDownstreamInputPinLabelAliases() const { return DownstreamInputPinLabelAliases; }
+	int64 GetGenerationCounter() const { return GenerationCounter; }
+
 	PCG_API static void AddReferencedObjects(UObject* InThis, FReferenceCollector& Collector);
 
 protected:
 #if WITH_EDITOR
-	void NotifyProducerUploadedData(UPCGDataBinding* InBinding);
+	PCG_API void NotifyProducerUploadedData(UPCGDataBinding* InBinding);
 #endif
 
-public:
+private:
 	int32 GraphBindingIndex = INDEX_NONE;
 	FName OutputPinLabel;
 	FName OutputPinLabelAlias;
 	TArray<FName> DownstreamInputPinLabelAliases;
 	bool bProducedByCPU = false;
 
-	/** Bumped each time the data provider is initialized or reset, so that async callbacks can detect if they originated from a previous usage of the data provider and no-op. */
-	std::atomic<uint64> GenerationCounter = 0;
-
-protected:
 	UPROPERTY()
 	TObjectPtr<const UPCGComputeKernel> ProducerKernel;
 
-private:
 	/** Settings of node that produces this data, normally the upstream node. */
 	UPROPERTY()
 	TSoftObjectPtr<const UPCGSettings> ProducerSettings;
 
 	mutable TObjectPtr<const UPCGSettings> ResolvedProducerSettings = nullptr;
+
+	/** Bumped each time the data provider is initialized or reset, so that async callbacks can detect if they originated from a previous usage of the data provider and no-op. */
+	std::atomic<uint64> GenerationCounter = 0;
 };
 
 /** Interface for data that can be exported/downloaded from the GPU to the CPU. */
@@ -123,7 +141,7 @@ public:
 	void SetRequiresExport(bool bInRequiresExport) { bRequiresExport = bInRequiresExport; }
 	bool GetRequiresExport() const { return bRequiresExport; }
 
-protected:
+private:
 	/** Whether this data is passed to downstream tasks outside of this compute graph. */
 	UPROPERTY()
 	bool bRequiresExport = false;
@@ -149,17 +167,17 @@ public:
 
 	EPCGExportMode GetExportMode() const { return ExportMode; }
 	UPCGDataBinding* GetDataBinding() const { return Binding.Get(); }
-	TSharedPtr<const FPCGDataCollectionDesc> GetPinDescription() const { return PinDesc; }
+	TSharedPtr<const FPCGDataCollectionDesc> GetPinDescription() const { return PinDataDescription; }
 
 	DECLARE_EVENT(UPCGExportableDataProvider, FOnDataExported);
 	FOnDataExported& OnDataExported_GameThread() { return OnDataExported; }
 
-protected:
+private:
 	EPCGExportMode ExportMode = EPCGExportMode::NoExport;
 	FOnDataExported OnDataExported;
 
 	TWeakObjectPtr<UPCGDataBinding> Binding;
-	TSharedPtr<const FPCGDataCollectionDesc> PinDesc = nullptr;
+	TSharedPtr<const FPCGDataCollectionDesc> PinDataDescription = nullptr;
 };
 
 /** Interface for marshalling settings/params to the GPU. */
