[PCG] Add "Don't show again" checkbox on the confirmation pop up for renaming graph parameters to set the default setting directly

diff --git a/Engine/Plugins/PCG/Source/PCGEditor/Private/Nodes/PCGEditorGraphNodeGetUserParameter.cpp b/Engine/Plugins/PCG/Source/PCGEditor/Private/Nodes/PCGEditorGraphNodeGetUserParameter.cpp
--- a/Engine/Plugins/PCG/Source/PCGEditor/Private/Nodes/PCGEditorGraphNodeGetUserParameter.cpp
+++ b/Engine/Plugins/PCG/Source/PCGEditor/Private/Nodes/PCGEditorGraphNodeGetUserParameter.cpp
@@ -7,12 +7,17 @@
 #include "PCGNode.h"
 #include "Elements/PCGReroute.h"
 #include "Elements/PCGUserParameterGet.h"
+#include "PCGEditorUtils.h"
 
 #include "Misc/MessageDialog.h"
+#include "Framework/Application/SlateApplication.h"
+#include "Framework/Docking/TabManager.h"
 #include "StructUtils/PropertyBag.h"
 
 #include "SPCGEditorGraphNodeCompact.h"
 
+#include "Widgets/Input/SCheckBox.h"
+
 #include UE_INLINE_GENERATED_CPP_BY_NAME(PCGEditorGraphNodeGetUserParameter)
 
 enum class EPropertyBagAlterationResult : uint8;
@@ -25,12 +30,48 @@ void UPCGEditorGraphGetUserParameter::OnRenameNode(const FString& NewName)
 	const UPCGUserParameterGetSettings* Settings = Cast<UPCGUserParameterGetSettings>(GetSettings());
 	if (ensure(Graph && Settings) && NewName != Settings->PropertyName.ToString())
 	{
-		const FText DialogMessage = FText::Format(LOCTEXT("RenameGraphParameterConfirmationMessage", "Rename from '{0}' to '{1}'?"), FText::FromName(Settings->PropertyName), FText::FromString(NewName));
-		const FText DialogTitle = LOCTEXT("RenameGraphParameterConfirmationTitle", "Rename Graph Parameter");
+		FText DialogMessage = FText::Format(LOCTEXT("RenameGraphParameterConfirmationMessage", "Confirm renaming node from '{0}' to '{1}'?"), FText::FromName(Settings->PropertyName), FText::FromString(NewName));
+		FText CheckBoxMessage = LOCTEXT("RenameGraphParameterConfirmationCheckBoxMessage", "Disable this notification. Re-enable via user preferences.");
 
 		if (GetDefault<UPCGEditorSettings>()->bConfirmLocalGraphParameterNameChanges)
 		{
-			if (EAppReturnType::Ok != FMessageDialog::Open(EAppMsgCategory::Info, EAppMsgType::OkCancel, std::move(DialogMessage), std::move(DialogTitle)))
+			TSharedPtr<SCheckBox> CheckBox = SNew(SCheckBox);
+			TSharedRef<SWidget> DisableNotificationWidget =
+				SNew(SHorizontalBox)
+				+ SHorizontalBox::Slot()
+				.HAlign(HAlign_Left)
+				.AutoWidth()
+				.Padding(4, 4, 0, 4)
+				[
+					CheckBox.ToSharedRef()
+				]
+				+ SHorizontalBox::Slot()
+				.HAlign(HAlign_Left)
+				.AutoWidth()
+				.Padding(4, 6, 0, 4)
+				[
+					SNew(STextBlock)
+					.Font(FAppStyle::Get().GetFontStyle("SmallFontBold"))
+					.Text(MoveTemp(CheckBoxMessage))
+				];
+
+			bool bConfirmClicked = false;
+			auto OnConfirm = [&bConfirmClicked, WeakCheckBox = CheckBox.ToWeakPtr()]
+			{
+				if (TSharedPtr<SCheckBox> CheckBoxPtr = WeakCheckBox.Pin())
+				{
+					if (CheckBoxPtr->IsChecked())
+					{
+						GetMutableDefault<UPCGEditorSettings>()->bConfirmLocalGraphParameterNameChanges = false;
+					}
+				}
+
+				bConfirmClicked = true;
+			};
+
+			PCGEditorUtils::OpenConfirmationDialog(DialogMessage, OnConfirm, /*OnCancel=*/[] {}, DisableNotificationWidget);
+
+			if (!bConfirmClicked)
 			{
 				return;
 			}
diff --git a/Engine/Plugins/PCG/Source/PCGEditor/Private/PCGEditorUtils.cpp b/Engine/Plugins/PCG/Source/PCGEditor/Private/PCGEditorUtils.cpp
--- a/Engine/Plugins/PCG/Source/PCGEditor/Private/PCGEditorUtils.cpp
+++ b/Engine/Plugins/PCG/Source/PCGEditor/Private/PCGEditorUtils.cpp
@@ -500,4 +500,166 @@ void PCGEditorUtils::OpenAssetOrMoveToActorOrComponent(const FSoftObjectPath& In
 	}
 }
 
+class SPCGSimpleConfirmationDialog : public SCompoundWidget
+{
+	SLATE_BEGIN_ARGS(SPCGSimpleConfirmationDialog) {}
+		SLATE_ARGUMENT(TWeakPtr<SWindow>, ParentWindow)
+		SLATE_ARGUMENT(TAttribute<FText>, Message)
+
+		SLATE_DEFAULT_SLOT(FArguments, AdditionalContent)
+
+		SLATE_EVENT(FSimpleDelegate, OnOkEvent)
+		SLATE_EVENT(FSimpleDelegate, OnCancelEvent)
+	SLATE_END_ARGS()
+
+	void Construct(const FArguments& InArgs)
+	{
+		ParentWindow = InArgs._ParentWindow;
+		OnConfirmEvent = InArgs._OnOkEvent;
+		OnCancelEvent = InArgs._OnCancelEvent;
+
+		ChildSlot
+		[
+			SNew(SVerticalBox)
+			+ SVerticalBox::Slot()
+			.VAlign(VAlign_Center)
+			.HAlign(HAlign_Center)
+			.AutoHeight()
+			.Padding(10, 16, 10, 0)
+			[
+				SNew(STextBlock)
+				.Text(InArgs._Message)
+				.AutoWrapText(true)
+			]
+			+ SVerticalBox::Slot()
+			.AutoHeight()
+			.Padding(6)
+			[
+				SNew(SHorizontalBox)
+				+ SHorizontalBox::Slot()
+				.Padding(20, 16, 12, 0)
+				.HAlign(HAlign_Fill)
+				.VAlign(VAlign_Center)
+				.FillWidth(1)
+				[
+					SNew(SBox)
+					.MinDesiredWidth(50)
+					.HAlign(HAlign_Center)
+					.VAlign(VAlign_Center)
+					[
+						SAssignNew(ConfirmButton, SButton)
+						.Text(LOCTEXT("SPCGSimpleConfirmationDialogConfirmButton", "Confirm"))
+						.OnClicked(this, &SPCGSimpleConfirmationDialog::OnConfirm)
+					]
+				]
+				+ SHorizontalBox::Slot()
+				.Padding(12, 16, 20, 0)
+				.HAlign(HAlign_Fill)
+				.VAlign(VAlign_Center)
+				.FillWidth(1)
+				[
+					SNew(SBox)
+					.MinDesiredWidth(50)
+					.HAlign(HAlign_Center)
+					.VAlign(VAlign_Center)
+					[
+						SAssignNew(CancelButton, SButton)
+						.Text(LOCTEXT("SPCGSimpleConfirmationDialogCancelButton", "Cancel"))
+						.ButtonColorAndOpacity(FSlateColor::UseSubduedForeground())
+						.OnClicked(this, &SPCGSimpleConfirmationDialog::OnCancel)
+					]
+				]
+			]
+			+ SVerticalBox::Slot()
+			.AutoHeight()
+			[
+				InArgs._AdditionalContent.Widget
+			]
+		];
+	}
+
+	bool ShouldProceed() const { return bShouldProceed; }
+
+	virtual bool SupportsKeyboardFocus() const override { return true; }
+
+	virtual FReply OnKeyDown(const FGeometry&, const FKeyEvent& InKeyEvent) override
+	{
+		const FKey Key = InKeyEvent.GetKey();
+
+		if (Key == EKeys::Escape)
+		{
+			return OnCancel();
+		}
+
+		return SCompoundWidget::OnKeyDown(FGeometry(), InKeyEvent);
+	}
+
+	TSharedPtr<SButton> GetConfirmButton() const { return ConfirmButton; }
+	TSharedPtr<SButton> GetCancelButton() const { return CancelButton; }
+
+private:
+	void RequestDestroyWindow() const
+	{
+		if (TSharedPtr<SWindow> Window = ParentWindow.Pin())
+		{
+			Window->RequestDestroyWindow();
+		}
+	}
+
+	FReply OnConfirm()
+	{
+		bShouldProceed = true;
+		OnConfirmEvent.ExecuteIfBound();
+		RequestDestroyWindow();
+		return FReply::Handled();
+	}
+
+	FReply OnCancel()
+	{
+		bShouldProceed = false;
+		OnCancelEvent.ExecuteIfBound();
+		RequestDestroyWindow();
+		return FReply::Handled();
+	}
+
+	bool bShouldProceed = false;
+	TWeakPtr<SWindow> ParentWindow = nullptr;
+	TSharedPtr<SButton> ConfirmButton = nullptr;
+	TSharedPtr<SButton> CancelButton = nullptr;
+
+	FSimpleDelegate OnConfirmEvent;
+	FSimpleDelegate OnCancelEvent;
+};
+
+void PCGEditorUtils::OpenConfirmationDialog(FText& InMessage, TFunctionRef<void()> OnOk, TFunctionRef<void()> OnCancel, TSharedPtr<SWidget> InAdditionalContent)
+{
+	TSharedRef<SWindow> Window = SNew(SWindow)
+		.CreateTitleBar(false)
+		.SizingRule(ESizingRule::Autosized)
+		.AutoCenter(EAutoCenter::PreferredWorkArea)
+		.SupportsMinimize(false)
+		.SupportsMaximize(false)
+		.ActivationPolicy(EWindowActivationPolicy::Always)
+		.FocusWhenFirstShown(true);
+
+	TSharedPtr<SPCGSimpleConfirmationDialog> Dialog;
+	Window->SetContent(
+		SAssignNew(Dialog, SPCGSimpleConfirmationDialog)
+		.ParentWindow(Window)
+		.Message(MoveTemp(InMessage))
+		.AdditionalContent()
+		[
+			InAdditionalContent.ToSharedRef()
+		]
+		.OnOkEvent(FSimpleDelegate::CreateLambda(OnOk))
+		.OnCancelEvent(FSimpleDelegate::CreateLambda(OnCancel)));
+
+	if (Dialog.IsValid() && Dialog->GetConfirmButton().IsValid())
+	{
+		Window->SetWidgetToFocusOnActivate(Dialog->GetConfirmButton());
+	}
+
+	FSlateApplication::Get().AddModalWindow(Window, FSlateApplication::Get().GetActiveTopLevelWindow());
+}
+
 #undef LOCTEXT_NAMESPACE
diff --git a/Engine/Plugins/PCG/Source/PCGEditor/Private/PCGEditorUtils.h b/Engine/Plugins/PCG/Source/PCGEditor/Private/PCGEditorUtils.h
--- a/Engine/Plugins/PCG/Source/PCGEditor/Private/PCGEditorUtils.h
+++ b/Engine/Plugins/PCG/Source/PCGEditor/Private/PCGEditorUtils.h
@@ -3,12 +3,14 @@
 #pragma once
 
 #include "Internationalization/Text.h"
+#include "Widgets/SNullWidget.h"
 
 template <typename FuncType> class TFunctionRef;
 struct FAssetData;
 struct FARFilter;
 class FString;
 class UObject;
+class SWidget;
 
 namespace PCGEditorUtils
 {
@@ -37,4 +39,6 @@ namespace PCGEditorUtils
 
 	/** Open an asset or jumps to an actor/component and select it. Will load it, and won't load a new level*/
 	void OpenAssetOrMoveToActorOrComponent(const FSoftObjectPath& InPath);
-}
+
+	void OpenConfirmationDialog(FText& InMessage, TFunctionRef<void()> OnOk, TFunctionRef<void()> OnCancel = [] {}, TSharedPtr<SWidget> InAdditionalContent = SNullWidget::NullWidget);
+} // namespace PCGEditorUtils
