[PCG] Update ray hit queries for spline tool

* Make sure it accept Landscape Proxies for it to work in WP
* Accept Static + Dynamic for collision channel for it to work on any meshes.

#rnx
#okforigithub public
#FailedValidations Preflight

diff --git a/Engine/Plugins/PCG/Source/PCGEditor/Private/EditorMode/Tools/Helpers/PCGEdModeSceneQueryHelpers.cpp b/Engine/Plugins/PCG/Source/PCGEditor/Private/EditorMode/Tools/Helpers/PCGEdModeSceneQueryHelpers.cpp
--- a/Engine/Plugins/PCG/Source/PCGEditor/Private/EditorMode/Tools/Helpers/PCGEdModeSceneQueryHelpers.cpp
+++ b/Engine/Plugins/PCG/Source/PCGEditor/Private/EditorMode/Tools/Helpers/PCGEdModeSceneQueryHelpers.cpp
@@ -43,7 +43,7 @@ namespace UE::PCG::EditorMode::Scene
 		}
 
 		// @todo_pcg: make this configurable
-		const FCollisionObjectQueryParams ObjectQueryParams(ECC_WorldStatic);
+		const FCollisionObjectQueryParams ObjectQueryParams(ECC_TO_BITFIELD(ECC_WorldStatic) | ECC_TO_BITFIELD(ECC_WorldDynamic));
 		FCollisionQueryParams QueryParams = FCollisionQueryParams::DefaultQueryParam;
 		QueryParams.bTraceComplex = true;
 
@@ -165,7 +165,7 @@ TOptional<bool> FPCGRaycastFilterRule_Landscape::IsRaycastHitValid(const FHitRes
 {
 	if (HitResult.bBlockingHit)
 	{
-		if (ALandscape* LandscapeActor = Cast<ALandscape>(HitResult.GetActor()))
+		if (ALandscapeProxy* LandscapeProxyActor = Cast<ALandscapeProxy>(HitResult.GetActor()))
 		{
 			return true;
 		}
