[PCG] Finalize data binding API

#rnx
#FailedValidations Preflight

diff --git a/Engine/Plugins/PCG/Source/PCG/Private/Compute/DataInterfaces/Elements/PCGCopyPointsDataInterface.cpp b/Engine/Plugins/PCG/Source/PCG/Private/Compute/DataInterfaces/Elements/PCGCopyPointsDataInterface.cpp
--- a/Engine/Plugins/PCG/Source/PCG/Private/Compute/DataInterfaces/Elements/PCGCopyPointsDataInterface.cpp
+++ b/Engine/Plugins/PCG/Source/PCG/Private/Compute/DataInterfaces/Elements/PCGCopyPointsDataInterface.cpp
@@ -187,7 +187,7 @@ bool UPCGCopyPointsDataProvider::PrepareForExecute_GameThread(UPCGDataBinding* I
 		const int32 NumOutputs = bCopyEachSourceOnEveryTarget ? SourcePinDesc->DataDescs.Num() * TargetPinDesc->DataDescs.Num() : FMath::Max(SourcePinDesc->DataDescs.Num(), TargetPinDesc->DataDescs.Num());
 		SourceAndTargetDataIndices.Reserve(NumOutputs);
 
-		const UPCGParamData* AnalysisResultsData = Cast<UPCGParamData>(InBinding->InputDataCollection.TaggedData[AnalysisDataIndex].Data);
+		const UPCGParamData* AnalysisResultsData = Cast<UPCGParamData>(InBinding->GetInputDataCollection().TaggedData[AnalysisDataIndex].Data);
 		const UPCGMetadata* AnalysisMetadata = AnalysisResultsData ? AnalysisResultsData->ConstMetadata() : nullptr;
 		const FPCGMetadataAttributeBase* AnalysisAttributeBase = AnalysisMetadata ? AnalysisMetadata->GetConstAttribute(PCGCopyPointsConstants::SelectedFlagAttributeName) : nullptr;
 		const int32 NumElements = AnalysisAttributeBase ? AnalysisMetadata->GetItemCountForChild() : INDEX_NONE;
diff --git a/Engine/Plugins/PCG/Source/PCG/Private/Compute/DataInterfaces/Elements/PCGMetadataPartitionDataInterface.cpp b/Engine/Plugins/PCG/Source/PCG/Private/Compute/DataInterfaces/Elements/PCGMetadataPartitionDataInterface.cpp
--- a/Engine/Plugins/PCG/Source/PCG/Private/Compute/DataInterfaces/Elements/PCGMetadataPartitionDataInterface.cpp
+++ b/Engine/Plugins/PCG/Source/PCG/Private/Compute/DataInterfaces/Elements/PCGMetadataPartitionDataInterface.cpp
@@ -159,12 +159,12 @@ bool UPCGMetaDataPartitionDataProvider::PrepareForExecute_GameThread(UPCGDataBin
 	{
 		const uint32 AnalysisDataIndex = AnalysisDataIndices[InputDataIndex];
 
-		if (!ensure(InBinding->InputDataCollection.TaggedData.IsValidIndex(AnalysisDataIndex)))
+		if (!ensure(InBinding->GetInputDataCollection().TaggedData.IsValidIndex(AnalysisDataIndex)))
 		{
 			continue;
 		}
 
-		const UPCGParamData* AnalysisResultsData = Cast<UPCGParamData>(InBinding->InputDataCollection.TaggedData[AnalysisDataIndex].Data);
+		const UPCGParamData* AnalysisResultsData = Cast<UPCGParamData>(InBinding->GetInputDataCollection().TaggedData[AnalysisDataIndex].Data);
 		const UPCGMetadata* AnalysisMetadata = AnalysisResultsData ? AnalysisResultsData->ConstMetadata() : nullptr;
 
 		const FPCGMetadataAttributeBase* AnalysisCountAttributeBase = AnalysisMetadata ? AnalysisMetadata->GetConstAttribute(PCGCountUniqueAttributeValuesConstants::ValueCountAttributeName) : nullptr;
@@ -192,7 +192,7 @@ bool UPCGMetaDataPartitionDataProvider::PrepareForExecute_GameThread(UPCGDataBin
 		}
 		else
 		{
-			if (InBinding->InputDataCollection.TaggedData[AnalysisDataIndex].Data->IsA<UPCGProxyForGPUData>())
+			if (InBinding->GetInputDataCollection().TaggedData[AnalysisDataIndex].Data->IsA<UPCGProxyForGPUData>())
 			{
 				UE_LOG(LogPCG, Error, TEXT("Data was not read back."));
 			}
diff --git a/Engine/Plugins/PCG/Source/PCG/Private/Compute/DataInterfaces/Elements/PCGSkinnedMeshSpawnerDataInterface.cpp b/Engine/Plugins/PCG/Source/PCG/Private/Compute/DataInterfaces/Elements/PCGSkinnedMeshSpawnerDataInterface.cpp
--- a/Engine/Plugins/PCG/Source/PCG/Private/Compute/DataInterfaces/Elements/PCGSkinnedMeshSpawnerDataInterface.cpp
+++ b/Engine/Plugins/PCG/Source/PCG/Private/Compute/DataInterfaces/Elements/PCGSkinnedMeshSpawnerDataInterface.cpp
@@ -237,7 +237,7 @@ bool UPCGSkinnedMeshSpawnerDataProvider::PerformPreExecuteReadbacks_GameThread(U
 		return false;
 	}
 
-	const UPCGParamData* AnalysisResultsData = Cast<UPCGParamData>(InBinding->InputDataCollection.TaggedData[AnalysisDataIndex].Data);
+	const UPCGParamData* AnalysisResultsData = Cast<UPCGParamData>(InBinding->GetInputDataCollection().TaggedData[AnalysisDataIndex].Data);
 	const UPCGMetadata* AnalysisMetadata = AnalysisResultsData ? AnalysisResultsData->ConstMetadata() : nullptr;
 
 	const FPCGMetadataAttributeBase* AnalysisValueAttributeBase = AnalysisMetadata ? AnalysisMetadata->GetConstAttribute(PCGCountUniqueAttributeValuesConstants::ValueAttributeName) : nullptr;
@@ -263,7 +263,7 @@ bool UPCGSkinnedMeshSpawnerDataProvider::PerformPreExecuteReadbacks_GameThread(U
 	{
 		UE_LOG(LogPCG, Warning, TEXT("No analysis data received by static mesh spawner kernel, worst case instance allocations will be made."));
 
-		if (InBinding->InputDataCollection.TaggedData[AnalysisDataIndex].Data->IsA<UPCGProxyForGPUData>())
+		if (InBinding->GetInputDataCollection().TaggedData[AnalysisDataIndex].Data->IsA<UPCGProxyForGPUData>())
 		{
 			UE_LOG(LogPCG, Error, TEXT("Data was not read back."));
 		}
@@ -287,22 +287,22 @@ bool UPCGSkinnedMeshSpawnerDataProvider::PrepareForExecute_GameThread(UPCGDataBi
 
 	if (!bRegisteredPrimitives)
 	{
-		InBinding->MeshSpawnersToPrimitives.Add(GetProducerKernel(), {});
+		InBinding->AddMeshSpawnerPrimitives(GetProducerKernel(), {});
 		bRegisteredPrimitives = true;
 	}
 
-	TSharedPtr<FPCGContextHandle> ContextHandle = InBinding->ContextHandle.Pin();
+	TSharedPtr<FPCGContextHandle> ContextHandle = InBinding->GetContextHandle().Pin();
 	FPCGContext* Context = ContextHandle ? ContextHandle->GetContext() : nullptr;
 	if (!ensure(Context))
 	{
-		InBinding->CompletedMeshSpawners.Add(GetProducerKernel());
+		InBinding->AddCompletedMeshSpawnerKernel(GetProducerKernel());
 		return true;
 	}
 
 	UPCGComponent* SourceComponent = Cast<UPCGComponent>(Context->ExecutionSource.Get());
 	if (!ensure(SourceComponent))
 	{
-		InBinding->CompletedMeshSpawners.Add(GetProducerKernel());
+		InBinding->AddCompletedMeshSpawnerKernel(GetProducerKernel());
 		return true;
 	}
 
@@ -313,7 +313,7 @@ bool UPCGSkinnedMeshSpawnerDataProvider::PrepareForExecute_GameThread(UPCGDataBi
 
 		if (!ensure(InputDataDesc))
 		{
-			InBinding->CompletedMeshSpawners.Add(GetProducerKernel());
+			InBinding->AddCompletedMeshSpawnerKernel(GetProducerKernel());
 			return true;
 		}
 
@@ -353,7 +353,7 @@ bool UPCGSkinnedMeshSpawnerDataProvider::PrepareForExecute_GameThread(UPCGDataBi
 					FText::FromName(SelectorName)));
 			}
 
-			InBinding->CompletedMeshSpawners.Add(GetProducerKernel());
+			InBinding->AddCompletedMeshSpawnerKernel(GetProducerKernel());
 			return true;
 		}
 
@@ -416,12 +416,12 @@ bool UPCGSkinnedMeshSpawnerDataProvider::PrepareForExecute_GameThread(UPCGDataBi
 		else
 		{
 			// No component set up means we have no more work to do.
-			InBinding->CompletedMeshSpawners.Add(GetProducerKernel());
+			InBinding->AddCompletedMeshSpawnerKernel(GetProducerKernel());
 			return true;
 		}
 	}
 
-	InBinding->CompletedMeshSpawners.Add(GetProducerKernel());
+	InBinding->AddCompletedMeshSpawnerKernel(GetProducerKernel());
 	return true;
 }
 
@@ -618,7 +618,7 @@ void UPCGSkinnedMeshSpawnerDataProvider::CreatePrimitiveDescriptors(FPCGContext*
 		PrimitiveSelectionCDF.SetNumZeroed(PrimitiveDescriptors.Num());
 	}
 
-	FPCGSpawnerPrimitives& Primitives = InBinding->MeshSpawnersToPrimitives.FindOrAdd(GetProducerKernel());
+	FPCGSpawnerPrimitives& Primitives = InBinding->FindOrAddMeshSpawnerPrimitives(GetProducerKernel());
 	Primitives.NumCustomFloats = CustomFloatCount;
 	Primitives.AttributeIdOffsetStrides = AttributeIdOffsetStrides;
 	Primitives.SelectorAttributeId = SelectorAttributeId;
@@ -645,7 +645,7 @@ bool UPCGSkinnedMeshSpawnerDataProvider::SetupPrimitives(FPCGContext* InContext,
 		return true;
 	}
 
-	FPCGSpawnerPrimitives* Primitives = InBinding->MeshSpawnersToPrimitives.Find(GetProducerKernel());
+	FPCGSpawnerPrimitives* Primitives = InBinding->FindMeshSpawnerPrimitives(GetProducerKernel());
 	if (!Primitives)
 	{
 		return true;
diff --git a/Engine/Plugins/PCG/Source/PCG/Private/Compute/DataInterfaces/Elements/PCGStaticMeshSpawnerDataInterface.cpp b/Engine/Plugins/PCG/Source/PCG/Private/Compute/DataInterfaces/Elements/PCGStaticMeshSpawnerDataInterface.cpp
--- a/Engine/Plugins/PCG/Source/PCG/Private/Compute/DataInterfaces/Elements/PCGStaticMeshSpawnerDataInterface.cpp
+++ b/Engine/Plugins/PCG/Source/PCG/Private/Compute/DataInterfaces/Elements/PCGStaticMeshSpawnerDataInterface.cpp
@@ -247,7 +247,7 @@ bool UPCGStaticMeshSpawnerDataProvider::PerformPreExecuteReadbacks_GameThread(UP
 		return false;
 	}
 
-	if (const UPCGParamData* AnalysisResultsData = Cast<UPCGParamData>(InBinding->InputDataCollection.TaggedData[AnalysisDataIndex].Data))
+	if (const UPCGParamData* AnalysisResultsData = Cast<UPCGParamData>(InBinding->GetInputDataCollection().TaggedData[AnalysisDataIndex].Data))
 	{
 		const UPCGMetadata* AnalysisMetadata = AnalysisResultsData->ConstMetadata();
 		check(AnalysisMetadata);
@@ -275,7 +275,7 @@ bool UPCGStaticMeshSpawnerDataProvider::PerformPreExecuteReadbacks_GameThread(UP
 		{
 			UE_LOG(LogPCG, Warning, TEXT("No analysis data received by static mesh spawner kernel, worst case instance allocations will be made."));
 
-			if (InBinding->InputDataCollection.TaggedData[AnalysisDataIndex].Data->IsA<UPCGProxyForGPUData>())
+			if (InBinding->GetInputDataCollection().TaggedData[AnalysisDataIndex].Data->IsA<UPCGProxyForGPUData>())
 			{
 				UE_LOG(LogPCG, Error, TEXT("Data was not read back."));
 			}
@@ -283,7 +283,7 @@ bool UPCGStaticMeshSpawnerDataProvider::PerformPreExecuteReadbacks_GameThread(UP
 			return true;
 		}
 	}
-	else if (const UPCGRawBufferData* RawAnalysisData = Cast<UPCGRawBufferData>(InBinding->InputDataCollection.TaggedData[AnalysisDataIndex].Data))
+	else if (const UPCGRawBufferData* RawAnalysisData = Cast<UPCGRawBufferData>(InBinding->GetInputDataCollection().TaggedData[AnalysisDataIndex].Data))
 	{
 		// Hardcoded data format - uint array containing pairs of (string key value, instance count) pairs.
 		const int32 NumUints = RawAnalysisData->GetNumUint32s();
@@ -312,21 +312,21 @@ bool UPCGStaticMeshSpawnerDataProvider::PrepareForExecute_GameThread(UPCGDataBin
 
 	if (!bRegisteredPrimitives)
 	{
-		InBinding->MeshSpawnersToPrimitives.Add(GetProducerKernel(), {});
+		InBinding->AddMeshSpawnerPrimitives(GetProducerKernel(), {});
 		bRegisteredPrimitives = true;
 	}
 
-	TSharedPtr<FPCGContextHandle> ContextHandle = InBinding->ContextHandle.Pin();
+	TSharedPtr<FPCGContextHandle> ContextHandle = InBinding->GetContextHandle().Pin();
 	FPCGContext* Context = ContextHandle ? ContextHandle->GetContext() : nullptr;
 	if (!ensure(Context))
 	{
-		InBinding->CompletedMeshSpawners.Add(GetProducerKernel());
+		InBinding->AddCompletedMeshSpawnerKernel(GetProducerKernel());
 		return true;
 	}
 
 	if (!ensure(Context->ExecutionSource.IsValid()))
 	{
-		InBinding->CompletedMeshSpawners.Add(GetProducerKernel());
+		InBinding->AddCompletedMeshSpawnerKernel(GetProducerKernel());
 		return true;
 	}
 
@@ -338,7 +338,7 @@ bool UPCGStaticMeshSpawnerDataProvider::PrepareForExecute_GameThread(UPCGDataBin
 
 		if (!ensure(InputDataDesc))
 		{
-			InBinding->CompletedMeshSpawners.Add(GetProducerKernel());
+			InBinding->AddCompletedMeshSpawnerKernel(GetProducerKernel());
 			return true;
 		}
 
@@ -378,7 +378,7 @@ bool UPCGStaticMeshSpawnerDataProvider::PrepareForExecute_GameThread(UPCGDataBin
 					FText::FromName(SelectorName)));
 			}
 
-			InBinding->CompletedMeshSpawners.Add(GetProducerKernel());
+			InBinding->AddCompletedMeshSpawnerKernel(GetProducerKernel());
 			return true;
 		}
 
@@ -452,7 +452,7 @@ bool UPCGStaticMeshSpawnerDataProvider::PrepareForExecute_GameThread(UPCGDataBin
 		else
 		{
 			// No component set up means we have no more work to do.
-			InBinding->CompletedMeshSpawners.Add(GetProducerKernel());
+			InBinding->AddCompletedMeshSpawnerKernel(GetProducerKernel());
 			return true;
 		}
 
@@ -467,7 +467,7 @@ bool UPCGStaticMeshSpawnerDataProvider::PrepareForExecute_GameThread(UPCGDataBin
 	SelectedMeshAttributeId = InBinding->GetAttributeId(Settings->OutAttributeName, EPCGKernelAttributeType::StringKey);
 	ensure(SelectedMeshAttributeId != INDEX_NONE);
 
-	InBinding->CompletedMeshSpawners.Add(GetProducerKernel());
+	InBinding->AddCompletedMeshSpawnerKernel(GetProducerKernel());
 	return true;
 }
 
@@ -740,7 +740,7 @@ void UPCGStaticMeshSpawnerDataProvider::CreatePrimitiveDescriptors(FPCGContext*
 		}
 	}
 
-	FPCGSpawnerPrimitives& Primitives = InBinding->MeshSpawnersToPrimitives.FindOrAdd(GetProducerKernel());
+	FPCGSpawnerPrimitives& Primitives = InBinding->FindOrAddMeshSpawnerPrimitives(GetProducerKernel());
 	Primitives.NumCustomFloats = CustomFloatCount;
 	Primitives.AttributeIdOffsetStrides = AttributeIdOffsetStrides;
 	Primitives.SelectorAttributeId = SelectorAttributeId;
@@ -761,7 +761,7 @@ bool UPCGStaticMeshSpawnerDataProvider::SetupPrimitives(FPCGContext* InContext,
 
 	const UPCGStaticMeshSpawnerSettings* Settings = CastChecked<UPCGStaticMeshSpawnerSettings>(GetProducerKernel()->GetSettings());
 
-	FPCGSpawnerPrimitives* Primitives = InBinding->MeshSpawnersToPrimitives.Find(GetProducerKernel());
+	FPCGSpawnerPrimitives* Primitives = InBinding->FindMeshSpawnerPrimitives(GetProducerKernel());
 	if (!Primitives)
 	{
 		return true;
diff --git a/Engine/Plugins/PCG/Source/PCG/Private/Compute/DataInterfaces/Elements/PCGTransformPointsDataInterface.cpp b/Engine/Plugins/PCG/Source/PCG/Private/Compute/DataInterfaces/Elements/PCGTransformPointsDataInterface.cpp
--- a/Engine/Plugins/PCG/Source/PCG/Private/Compute/DataInterfaces/Elements/PCGTransformPointsDataInterface.cpp
+++ b/Engine/Plugins/PCG/Source/PCG/Private/Compute/DataInterfaces/Elements/PCGTransformPointsDataInterface.cpp
@@ -178,7 +178,7 @@ bool UPCGTransformPointsDataProvider::PrepareForExecute_GameThread(UPCGDataBindi
 		return false;
 	}
 
-	TSharedPtr<FPCGContextHandle> ContextHandle = InBinding->ContextHandle.Pin();
+	TSharedPtr<FPCGContextHandle> ContextHandle = InBinding->GetContextHandle().Pin();
 	FPCGContext* Context = ContextHandle ? ContextHandle->GetContext() : nullptr;
 	if (!ensure(Context))
 	{
diff --git a/Engine/Plugins/PCG/Source/PCG/Private/Compute/DataInterfaces/PCGComputeDataInterface.cpp b/Engine/Plugins/PCG/Source/PCG/Private/Compute/DataInterfaces/PCGComputeDataInterface.cpp
--- a/Engine/Plugins/PCG/Source/PCG/Private/Compute/DataInterfaces/PCGComputeDataInterface.cpp
+++ b/Engine/Plugins/PCG/Source/PCG/Private/Compute/DataInterfaces/PCGComputeDataInterface.cpp
@@ -149,7 +149,7 @@ void UPCGComputeDataProvider::NotifyProducerUploadedData(UPCGDataBinding* InBind
 			}
 		}
 
-		TSharedPtr<FPCGContextHandle> ContextHandle = InBinding->ContextHandle.Pin();
+		TSharedPtr<FPCGContextHandle> ContextHandle = InBinding->GetContextHandle().Pin();
 		FPCGContext* Context = ContextHandle ? ContextHandle->GetContext() : nullptr;
 		if (ProducerNode && Context && Context->GetStack() && Context->ExecutionSource.IsValid())
 		{
diff --git a/Engine/Plugins/PCG/Source/PCG/Private/Compute/DataInterfaces/PCGDataLabelResolverDataInterface.cpp b/Engine/Plugins/PCG/Source/PCG/Private/Compute/DataInterfaces/PCGDataLabelResolverDataInterface.cpp
--- a/Engine/Plugins/PCG/Source/PCG/Private/Compute/DataInterfaces/PCGDataLabelResolverDataInterface.cpp
+++ b/Engine/Plugins/PCG/Source/PCG/Private/Compute/DataInterfaces/PCGDataLabelResolverDataInterface.cpp
@@ -92,7 +92,7 @@ bool UPCGDataLabelResolverDataProvider::PrepareForExecute_GameThread(UPCGDataBin
 	check(InBinding);
 	check(Kernel);
 
-	const UPCGComputeGraph* ComputeGraph = InBinding->Graph.Get();
+	const UPCGComputeGraph* ComputeGraph = InBinding->GetComputeGraph();
 	check(ComputeGraph);
 
 	const TArray<FString>& StringTable = InBinding->GetStringTable();
@@ -131,7 +131,7 @@ bool UPCGDataLabelResolverDataProvider::PrepareForExecute_GameThread(UPCGDataBin
 
 			if (DataIdToDataIndexMap[LabelIndex] == INDEX_NONE)
 			{
-				if (TSharedPtr<FPCGContextHandle> Context = InBinding->ContextHandle.Pin())
+				if (TSharedPtr<FPCGContextHandle> Context = InBinding->GetContextHandle().Pin())
 				{
 					PCG_KERNEL_VALIDATION_ERR(Context->GetContext(), Kernel->GetSettings(), FText::Format(
 						NSLOCTEXT("PCGDataLabelResolver", "LabelDoesNotExist", "Data label '{0}' was not present in the tags on the incoming data."),
diff --git a/Engine/Plugins/PCG/Source/PCG/Private/Compute/DataInterfaces/PCGInstanceDataInterface.cpp b/Engine/Plugins/PCG/Source/PCG/Private/Compute/DataInterfaces/PCGInstanceDataInterface.cpp
--- a/Engine/Plugins/PCG/Source/PCG/Private/Compute/DataInterfaces/PCGInstanceDataInterface.cpp
+++ b/Engine/Plugins/PCG/Source/PCG/Private/Compute/DataInterfaces/PCGInstanceDataInterface.cpp
@@ -148,13 +148,13 @@ bool UPCGInstanceDataProvider::PrepareForExecute_GameThread(UPCGDataBinding* InB
 	TRACE_CPUPROFILER_EVENT_SCOPE(UPCGInstanceDataProvider::PrepareForExecute_GameThread);
 	UPCGDataBinding* Binding = CastChecked<UPCGDataBinding>(InBinding);
 
-	if (!Binding->CompletedMeshSpawners.Contains(GetProducerKernel()))
+	if (!Binding->IsMeshSpawnerKernelComplete(GetProducerKernel()))
 	{
 		// The static mesh data interface(s) set this up, so wait until it is ready.
 		return false;
 	}
 
-	FPCGSpawnerPrimitives* FoundPrimitives = Binding->MeshSpawnersToPrimitives.Find(GetProducerKernel());
+	FPCGSpawnerPrimitives* FoundPrimitives = Binding->FindMeshSpawnerPrimitives(GetProducerKernel());
 	if (!FoundPrimitives)
 	{
 		return true;
@@ -190,7 +190,7 @@ bool UPCGInstanceDataProvider::PrepareForExecute_GameThread(UPCGDataBinding* InB
 
 	NumCustomFloatsPerInstance = FoundPrimitives->NumCustomFloats;
 
-	ContextHandle = InBinding->ContextHandle;
+	ContextHandle = InBinding->GetContextHandle();
 
 	return true;
 }
diff --git a/Engine/Plugins/PCG/Source/PCG/Private/Compute/DataInterfaces/PCGLandscapeDataInterface.cpp b/Engine/Plugins/PCG/Source/PCG/Private/Compute/DataInterfaces/PCGLandscapeDataInterface.cpp
--- a/Engine/Plugins/PCG/Source/PCG/Private/Compute/DataInterfaces/PCGLandscapeDataInterface.cpp
+++ b/Engine/Plugins/PCG/Source/PCG/Private/Compute/DataInterfaces/PCGLandscapeDataInterface.cpp
@@ -318,8 +318,8 @@ void UPCGLandscapeDataProvider::Initialize(const UComputeDataInterface* InDataIn
 	};
 
 	// Take any input pin label alias to obtain the data from the input data collection.
-	check(!GetDownstreamInputPinLabelAliases().IsEmpty());
-	const TArray<FPCGTaggedData> TaggedDatas = Binding->InputDataCollection.GetInputsByPin(GetDownstreamInputPinLabelAliases()[0]);
+	check(!DataInterface->GetDownstreamInputPinLabelAliases().IsEmpty());
+	const TArray<FPCGTaggedData> TaggedDatas = Binding->GetInputDataCollection().GetInputsByPin(DataInterface->GetDownstreamInputPinLabelAliases()[0]);
 
 	if (!TaggedDatas.IsEmpty())
 	{
diff --git a/Engine/Plugins/PCG/Source/PCG/Private/Compute/DataInterfaces/PCGRawBufferDataInterface.cpp b/Engine/Plugins/PCG/Source/PCG/Private/Compute/DataInterfaces/PCGRawBufferDataInterface.cpp
--- a/Engine/Plugins/PCG/Source/PCG/Private/Compute/DataInterfaces/PCGRawBufferDataInterface.cpp
+++ b/Engine/Plugins/PCG/Source/PCG/Private/Compute/DataInterfaces/PCGRawBufferDataInterface.cpp
@@ -131,7 +131,7 @@ bool UPCGRawBufferDataProvider::PrepareForExecute_GameThread(UPCGDataBinding* In
 	{
 		check(!GetDownstreamInputPinLabelAliases().IsEmpty());
 
-		const TArray<FPCGTaggedData> TaggedDatas = InBinding->InputDataCollection.GetInputsByPin(GetDownstreamInputPinLabelAliases()[0]);
+		const TArray<FPCGTaggedData> TaggedDatas = InBinding->GetInputDataCollection().GetInputsByPin(GetDownstreamInputPinLabelAliases()[0]);
 
 		if (!TaggedDatas.IsEmpty())
 		{
diff --git a/Engine/Plugins/PCG/Source/PCG/Private/Compute/DataInterfaces/PCGStaticMeshDataInterface.cpp b/Engine/Plugins/PCG/Source/PCG/Private/Compute/DataInterfaces/PCGStaticMeshDataInterface.cpp
--- a/Engine/Plugins/PCG/Source/PCG/Private/Compute/DataInterfaces/PCGStaticMeshDataInterface.cpp
+++ b/Engine/Plugins/PCG/Source/PCG/Private/Compute/DataInterfaces/PCGStaticMeshDataInterface.cpp
@@ -162,7 +162,7 @@ bool UPCGStaticMeshDataProvider::PrepareForExecute_GameThread(UPCGDataBinding* I
 		check(!GetDownstreamInputPinLabelAliases().IsEmpty());
 
 		const UPCGDataBinding* Binding = CastChecked<UPCGDataBinding>(InBinding);
-		const TArray<FPCGTaggedData> TaggedDatas = Binding->InputDataCollection.GetInputsByPin(GetDownstreamInputPinLabelAliases()[0]);
+		const TArray<FPCGTaggedData> TaggedDatas = Binding->GetInputDataCollection().GetInputsByPin(GetDownstreamInputPinLabelAliases()[0]);
 
 		if (!TaggedDatas.IsEmpty())
 		{
diff --git a/Engine/Plugins/PCG/Source/PCG/Private/Compute/DataInterfaces/PCGTextureDataInterface.cpp b/Engine/Plugins/PCG/Source/PCG/Private/Compute/DataInterfaces/PCGTextureDataInterface.cpp
--- a/Engine/Plugins/PCG/Source/PCG/Private/Compute/DataInterfaces/PCGTextureDataInterface.cpp
+++ b/Engine/Plugins/PCG/Source/PCG/Private/Compute/DataInterfaces/PCGTextureDataInterface.cpp
@@ -298,7 +298,7 @@ void UPCGTextureDataProvider::BuildInfosFromDataCollection(UPCGDataBinding* InBi
 
 	// Take any input pin label alias to obtain the data from the input data collection.
 	check(!GetDownstreamInputPinLabelAliases().IsEmpty());
-	const TArray<FPCGTaggedData> InputTaggedData = InBinding->InputDataCollection.GetInputsByPin(GetDownstreamInputPinLabelAliases()[0]);
+	const TArray<FPCGTaggedData> InputTaggedData = InBinding->GetInputDataCollection().GetInputsByPin(GetDownstreamInputPinLabelAliases()[0]);
 
 	for (int I = 0; I < InputTaggedData.Num(); ++I)
 	{
@@ -339,7 +339,7 @@ void UPCGTextureDataProvider::BuildInfosFromDataDescription(UPCGDataBinding* InB
 {
 	check(InBinding);
 
-	TSharedPtr<FPCGContextHandle> ContextHandle = InBinding->ContextHandle.Pin();
+	TSharedPtr<FPCGContextHandle> ContextHandle = InBinding->GetContextHandle().Pin();
 	FPCGContext* Context = ContextHandle ? ContextHandle->GetContext() : nullptr;
 
 	if (!ensure(Context))
diff --git a/Engine/Plugins/PCG/Source/PCG/Private/Compute/DataInterfaces/PCGVirtualTextureDataInterface.cpp b/Engine/Plugins/PCG/Source/PCG/Private/Compute/DataInterfaces/PCGVirtualTextureDataInterface.cpp
--- a/Engine/Plugins/PCG/Source/PCG/Private/Compute/DataInterfaces/PCGVirtualTextureDataInterface.cpp
+++ b/Engine/Plugins/PCG/Source/PCG/Private/Compute/DataInterfaces/PCGVirtualTextureDataInterface.cpp
@@ -133,8 +133,8 @@ void UPCGVirtualTextureDataProvider::Initialize(const UComputeDataInterface* InD
 	const UPCGComponent* SourceComponent = Cast<UPCGComponent>(Binding->GetExecutionSource());
 
 	// Take any input pin label alias to obtain the data from the input data collection.
-	check(!GetDownstreamInputPinLabelAliases().IsEmpty());
-	const TArray<FPCGTaggedData> InputTaggedData = Binding->InputDataCollection.GetInputsByPin(GetDownstreamInputPinLabelAliases()[0]);
+	check(!DataInterface->GetDownstreamInputPinLabelAliases().IsEmpty());
+	const TArray<FPCGTaggedData> InputTaggedData = Binding->GetInputDataCollection().GetInputsByPin(DataInterface->GetDownstreamInputPinLabelAliases()[0]);
 
 	// Sampling of virtual textures is not supported outside of runtime generation at this time, because the RuntimeGenScheduler is responsible for priming the virtual textures.
 	if (!SourceComponent || !SourceComponent->IsManagedByRuntimeGenSystem())
diff --git a/Engine/Plugins/PCG/Source/PCG/Private/Compute/Elements/PCGComputeGraphElement.cpp b/Engine/Plugins/PCG/Source/PCG/Private/Compute/Elements/PCGComputeGraphElement.cpp
--- a/Engine/Plugins/PCG/Source/PCG/Private/Compute/Elements/PCGComputeGraphElement.cpp
+++ b/Engine/Plugins/PCG/Source/PCG/Private/Compute/Elements/PCGComputeGraphElement.cpp
@@ -604,7 +604,7 @@ bool FPCGComputeGraphElement::ExecuteInternal(FPCGContext* InContext) const
 		// Currently we don't output anything if processing any readback data processing failed.
 		if (ensure(Context->bExecutionSuccess) && ensure(Context->DataBinding))
 		{
-			Context->OutputData = Context->DataBinding->OutputDataCollection;
+			Context->OutputData = Context->DataBinding->GetOutputDataCollection();
 		}
 
 #if WITH_EDITOR
@@ -669,7 +669,7 @@ void FPCGComputeGraphElement::PostExecuteInternal(FPCGContext* InContext) const
 #if WITH_EDITOR
 	if (Context->bExecutionSuccess)
 	{
-		for (TSoftObjectPtr<const UPCGNode> NodePtr : Context->DataBinding->Graph->KernelToNode)
+		for (TSoftObjectPtr<const UPCGNode> NodePtr : Context->DataBinding->GetComputeGraph()->KernelToNode)
 		{
 			const UPCGNode* Node = NodePtr.Get();
 			IPCGGraphExecutionSource* ExecutionSource = Context->ExecutionSource.Get();
@@ -753,12 +753,12 @@ void FPCGComputeGraphElement::CollectDebugDataPrepareActions(FPCGComputeGraphCon
 		}
 	};
 
-	for (FPCGDataToDebug& DataToDebug : InContext->DataBinding->DataToDebug)
+	for (FPCGDataToDebug& DataToDebug : InContext->DataBinding->GetDataToDebugMutable())
 	{
 		CollectDebugDataPrepareAction(DataToDebug, /*bIsInspecting=*/false);
 	}
 
-	for (FPCGDataToDebug& DataToDebug : InContext->DataBinding->DataToInspect)
+	for (FPCGDataToDebug& DataToDebug : InContext->DataBinding->GetDataToInspectMutable())
 	{
 		CollectDebugDataPrepareAction(DataToDebug, /*bIsInspecting=*/true);
 	}
@@ -773,9 +773,11 @@ void FPCGComputeGraphElement::ExecuteDebugDraw(FPCGComputeGraphContext* InContex
 
 	const FPCGDataVisualizationRegistry& DataVisRegistry = FPCGModule::GetConstPCGDataVisualizationRegistry();
 
-	for (int Index = InContext->DataBinding->DataToDebug.Num() - 1; Index >= 0; --Index)
+	TArray<FPCGDataToDebug>& AllDataToDebug = InContext->DataBinding->GetDataToDebugMutable();
+
+	for (int Index = AllDataToDebug.Num() - 1; Index >= 0; --Index)
 	{
-		FPCGDataToDebug& DataToDebug = InContext->DataBinding->DataToDebug[Index];
+		FPCGDataToDebug& DataToDebug = AllDataToDebug[Index];
 		const UPCGData* DataToDisplay = nullptr;
 
 		if (const UPCGProxyForGPUData* Proxy = Cast<UPCGProxyForGPUData>(DataToDebug.Data.Get()))
@@ -790,7 +792,7 @@ void FPCGComputeGraphElement::ExecuteDebugDraw(FPCGComputeGraphContext* InContex
 			DataToDisplay = DataToDebug.DataPendingInit;
 		}
 
-		InContext->DataBinding->DataToDebug.RemoveAtSwap(Index);
+		AllDataToDebug.RemoveAtSwap(Index);
 
 		const UPCGSettings* ProducerSettings = DataToDebug.ProducerSettings.Get();
 
@@ -814,7 +816,7 @@ void FPCGComputeGraphElement::StoreDataForInspection(FPCGComputeGraphContext* In
 	// Collect all data into a collection and store it.
 	TMap<const UPCGSettings*, FPCGDataCollection> SettingsToDataCollection;
 
-	for (FPCGDataToDebug& DataToInspect : InContext->DataBinding->DataToInspect)
+	for (FPCGDataToDebug& DataToInspect : InContext->DataBinding->GetDataToInspectMutable())
 	{
 		if (const UPCGSettings* ProducerSettings = DataToInspect.ProducerSettings.Get())
 		{
@@ -852,7 +854,7 @@ void FPCGComputeGraphElement::StoreDataForInspection(FPCGComputeGraphContext* In
 		InContext->ExecutionSource->GetExecutionState().GetInspection().StoreInspectionData(InContext->GetStack(), Node, /*InTimer=*/nullptr, /*InInputData=*/{}, SettingsAndData.Get<1>(), /*bUsedCache*/false);
 	}
 
-	InContext->DataBinding->DataToInspect.Empty();
+	InContext->DataBinding->GetDataToInspectMutable().Empty();
 }
 
 void FPCGComputeGraphElement::LogCompilationMessages(FPCGComputeGraphContext* InContext) const
diff --git a/Engine/Plugins/PCG/Source/PCG/Private/Compute/Elements/PCGCustomHLSLKernel.cpp b/Engine/Plugins/PCG/Source/PCG/Private/Compute/Elements/PCGCustomHLSLKernel.cpp
--- a/Engine/Plugins/PCG/Source/PCG/Private/Compute/Elements/PCGCustomHLSLKernel.cpp
+++ b/Engine/Plugins/PCG/Source/PCG/Private/Compute/Elements/PCGCustomHLSLKernel.cpp
@@ -1940,7 +1940,6 @@ const FPCGPinPropertiesGPU* UPCGCustomHLSLKernel::GetFirstOutputPin() const
 int UPCGCustomHLSLKernel::GetElementCountForInputPin(const FPCGPinProperties& InInputPinProps, const UPCGDataBinding* InBinding) const
 {
 	check(InBinding);
-	check(InBinding->Graph);
 
 	const TSharedPtr<const FPCGDataCollectionDesc> InputDesc = InBinding->GetCachedKernelPinDataDesc(this, InInputPinProps.Label, /*bIsInputPin=*/true);
 
diff --git a/Engine/Plugins/PCG/Source/PCG/Private/Compute/PCGComputeCommon.cpp b/Engine/Plugins/PCG/Source/PCG/Private/Compute/PCGComputeCommon.cpp
--- a/Engine/Plugins/PCG/Source/PCG/Private/Compute/PCGComputeCommon.cpp
+++ b/Engine/Plugins/PCG/Source/PCG/Private/Compute/PCGComputeCommon.cpp
@@ -320,14 +320,14 @@ namespace PCGComputeHelpers
 			return;
 		}
 
-		TSharedPtr<FPCGContextHandle> ContextHandle = InBinding->ContextHandle.Pin();
+		TSharedPtr<FPCGContextHandle> ContextHandle = InBinding->GetContextHandle().Pin();
 		if (FPCGContext* Context = ContextHandle ? ContextHandle->GetContext() : nullptr)
 		{
 			const UPCGNode* IndicatorNode = nullptr;
 
 			if (InKernel)
 			{
-				IndicatorNode = InBinding->Graph ? InBinding->Graph->GetKernelNode(InKernel) : nullptr;
+				IndicatorNode = InBinding->GetComputeGraph() ? InBinding->GetComputeGraph()->GetKernelNode(InKernel) : nullptr;
 			}
 			else if (InSettings)
 			{
diff --git a/Engine/Plugins/PCG/Source/PCG/Private/Compute/PCGComputeGraph.cpp b/Engine/Plugins/PCG/Source/PCG/Private/Compute/PCGComputeGraph.cpp
--- a/Engine/Plugins/PCG/Source/PCG/Private/Compute/PCGComputeGraph.cpp
+++ b/Engine/Plugins/PCG/Source/PCG/Private/Compute/PCGComputeGraph.cpp
@@ -221,7 +221,7 @@ void UPCGComputeGraph::DebugLogDataDescriptions(const UPCGDataBinding* InBinding
 
 			if (const FName* VirtualLabel = GetVirtualPinLabelFromBindingIndex(GetBindingIndex(InKernelPin)))
 			{
-				TArray<FPCGTaggedData> PinInputs = InBinding->InputDataCollection.GetInputsByPin(*VirtualLabel);
+				TArray<FPCGTaggedData> PinInputs = InBinding->GetInputDataCollection().GetInputsByPin(*VirtualLabel);
 
 				if (ensure(PinInputs.IsValidIndex(DataIndex)))
 				{
@@ -368,7 +368,7 @@ TSharedPtr<const FPCGDataCollectionDesc> UPCGComputeGraph::ComputeExternalPinDes
 
 	TSharedPtr<FPCGDataCollectionDesc> OutDataDesc = FPCGDataCollectionDesc::MakeShared();
 
-	const TArray<FPCGTaggedData> TaggedDatas = InBinding->InputDataCollection.GetInputsByPin(InVirtualLabel);
+	const TArray<FPCGTaggedData> TaggedDatas = InBinding->GetInputDataCollection().GetInputsByPin(InVirtualLabel);
 
 	for (const FPCGTaggedData& TaggedData : TaggedDatas)
 	{
diff --git a/Engine/Plugins/PCG/Source/PCG/Private/Compute/PCGDataBinding.cpp b/Engine/Plugins/PCG/Source/PCG/Private/Compute/PCGDataBinding.cpp
--- a/Engine/Plugins/PCG/Source/PCG/Private/Compute/PCGDataBinding.cpp
+++ b/Engine/Plugins/PCG/Source/PCG/Private/Compute/PCGDataBinding.cpp
@@ -319,6 +319,66 @@ const FPCGKernelParams* UPCGDataBinding::GetCachedKernelParams(int32 InKernelInd
 	return KernelParamsCache.GetCachedKernelParams(InKernelIndex);
 }
 
+TArray<FPCGDataToDebug>& UPCGDataBinding::GetDataToDebugMutable()
+{
+	return DataToDebug;
+}
+
+TArray<FPCGDataToDebug>& UPCGDataBinding::GetDataToInspectMutable()
+{
+	return DataToInspect;
+}
+
+void UPCGDataBinding::AddCompletedMeshSpawnerKernel(TObjectPtr<const UComputeKernel> InCompletedMeshSpawnerKernel)
+{
+	CompletedMeshSpawners.Add(InCompletedMeshSpawnerKernel);
+}
+
+bool UPCGDataBinding::IsMeshSpawnerKernelComplete(TObjectPtr<const UComputeKernel> InMeshSpawnerKernel) const
+{
+	return CompletedMeshSpawners.Contains(InMeshSpawnerKernel);
+}
+
+void UPCGDataBinding::AddMeshSpawnerPrimitives(TObjectPtr<const UComputeKernel> InMeshSpawnerKernel, const FPCGSpawnerPrimitives& Primitives)
+{
+	MeshSpawnersToPrimitives.Add(InMeshSpawnerKernel, Primitives);
+}
+
+void UPCGDataBinding::AddMeshSpawnerPrimitives(TObjectPtr<const UComputeKernel> InMeshSpawnerKernel, FPCGSpawnerPrimitives&& Primitives)
+{
+	MeshSpawnersToPrimitives.Add(InMeshSpawnerKernel, MoveTempIfPossible(Primitives));
+}
+
+FPCGSpawnerPrimitives& UPCGDataBinding::FindOrAddMeshSpawnerPrimitives(TObjectPtr<const UComputeKernel> InMeshSpawnerKernel)
+{
+	return MeshSpawnersToPrimitives.FindOrAdd(InMeshSpawnerKernel);
+}
+
+FPCGSpawnerPrimitives* UPCGDataBinding::FindMeshSpawnerPrimitives(TObjectPtr<const UComputeKernel> InMeshSpawnerKernel)
+{
+	return MeshSpawnersToPrimitives.Find(InMeshSpawnerKernel);
+}
+
+const UPCGComputeGraph* UPCGDataBinding::GetComputeGraph() const
+{
+	return Graph;
+}
+
+const FPCGDataCollection& UPCGDataBinding::GetInputDataCollection() const
+{
+	return InputDataCollection;
+}
+
+FPCGDataCollection& UPCGDataBinding::GetInputDataCollectionMutable()
+{
+	return InputDataCollection;
+}
+
+const FPCGDataCollection& UPCGDataBinding::GetOutputDataCollection() const
+{
+	return OutputDataCollection;
+}
+
 void UPCGDataBinding::AddInputDataAttributesToTable()
 {
 	for (const FPCGTaggedData& Data : InputDataCollection.TaggedData)
diff --git a/Engine/Plugins/PCG/Source/PCG/Private/Compute/PCGDataForGPU.cpp b/Engine/Plugins/PCG/Source/PCG/Private/Compute/PCGDataForGPU.cpp
--- a/Engine/Plugins/PCG/Source/PCG/Private/Compute/PCGDataForGPU.cpp
+++ b/Engine/Plugins/PCG/Source/PCG/Private/Compute/PCGDataForGPU.cpp
@@ -12,7 +12,6 @@
 #include "PCGPointPropertiesTraits.h"
 #include "PCGSettings.h"
 #include "Compute/PCGComputeCommon.h"
-#include "Compute/PCGComputeGraph.h"
 #include "Compute/PCGDataBinding.h"
 #include "Compute/PCGPinPropertiesGPU.h"
 #include "Compute/Data/PCGProxyForGPUData.h"
@@ -581,7 +580,6 @@ namespace PCGDataForGPUHelpers
 		TArray<FUint32Vector4>& OutAttributeIdOffsetStrides)
 	{
 		check(InBinding);
-		check(InBinding->Graph);
 		check(InDataCollectionDescription);
 
 		uint32 OffsetFloats = 0;
diff --git a/Engine/Plugins/PCG/Source/PCG/Private/Compute/PCGDataUploadAdaptor.cpp b/Engine/Plugins/PCG/Source/PCG/Private/Compute/PCGDataUploadAdaptor.cpp
--- a/Engine/Plugins/PCG/Source/PCG/Private/Compute/PCGDataUploadAdaptor.cpp
+++ b/Engine/Plugins/PCG/Source/PCG/Private/Compute/PCGDataUploadAdaptor.cpp
@@ -29,9 +29,9 @@ FPCGDataUploadAdaptor::FPCGDataUploadAdaptor(UPCGDataBinding* InDataBinding, con
 	// Are all our data items coming direct from GPU - and importantly are a complete, single data collection?
 	bool bMultipleBuffersEncountered = false;
 	TSharedPtr<const FPCGProxyForGPUDataCollection> UniqueDataCollection;
-	for (int InputDataIndex = 0; InputDataIndex < DataBinding->InputDataCollection.TaggedData.Num(); ++InputDataIndex)
+	for (int InputDataIndex = 0; InputDataIndex < DataBinding->GetInputDataCollection().TaggedData.Num(); ++InputDataIndex)
 	{
-		const FName DataPin = DataBinding->InputDataCollection.TaggedData[InputDataIndex].Pin;
+		const FName DataPin = DataBinding->GetInputDataCollection().TaggedData[InputDataIndex].Pin;
 
 		if (InInputPinLabel != DataPin)
 		{
@@ -40,7 +40,7 @@ FPCGDataUploadAdaptor::FPCGDataUploadAdaptor(UPCGDataBinding* InDataBinding, con
 
 		++NumDataForThisPin;
 
-		const UPCGProxyForGPUData* DataGPU = Cast<UPCGProxyForGPUData>(DataBinding->InputDataCollection.TaggedData[InputDataIndex].Data);
+		const UPCGProxyForGPUData* DataGPU = Cast<UPCGProxyForGPUData>(DataBinding->GetInputDataCollection().TaggedData[InputDataIndex].Data);
 		if (!DataGPU)
 		{
 			// Encountered CPU data, no reuse for now.
@@ -152,7 +152,7 @@ bool FPCGDataUploadAdaptor::PrepareData_GameThread(const UPCGComputeKernel* InPr
 
 		bool bAllDataReady = true;
 
-		for (FPCGTaggedData& TaggedData : Binding->InputDataCollection.TaggedData)
+		for (FPCGTaggedData& TaggedData : Binding->GetInputDataCollectionMutable().TaggedData)
 		{
 			if (TaggedData.Pin != InputPinLabel)
 			{
@@ -188,7 +188,7 @@ bool FPCGDataUploadAdaptor::PrepareData_GameThread(const UPCGComputeKernel* InPr
 		PCGComputeHelpers::NotifyCPUToGPUUpload(Binding, InProducerKernel, InProducerSettings);
 #endif
 
-		TargetDataCollectionDesc->PackDataCollection(Binding->InputDataCollection, InputPinLabel, Binding, PackedDataCollection);
+		TargetDataCollectionDesc->PackDataCollection(Binding->GetInputDataCollection(), InputPinLabel, Binding, PackedDataCollection);
 	}
 
 	return true;
diff --git a/Engine/Plugins/PCG/Source/PCG/Private/Compute/PCGKernelParamsCache.cpp b/Engine/Plugins/PCG/Source/PCG/Private/Compute/PCGKernelParamsCache.cpp
--- a/Engine/Plugins/PCG/Source/PCG/Private/Compute/PCGKernelParamsCache.cpp
+++ b/Engine/Plugins/PCG/Source/PCG/Private/Compute/PCGKernelParamsCache.cpp
@@ -223,14 +223,16 @@ void FPCGKernelParamsCache::CacheOverridableParamDataIndices(FPCGContext* InCont
 
 	if (OverridableParamDataIndices.IsEmpty())
 	{
-		if (!ensure(InDataBinding->Graph))
+		const UPCGComputeGraph* ComputeGraph = InDataBinding->GetComputeGraph();
+
+		if (!ensure(ComputeGraph))
 		{
 			return;
 		}
 
-		for (int32 KernelIndex = 0; KernelIndex < InDataBinding->Graph->GetNumKernels(); ++KernelIndex)
+		for (int32 KernelIndex = 0; KernelIndex < ComputeGraph->GetNumKernels(); ++KernelIndex)
 		{
-			const UPCGComputeKernel* Kernel = Cast<UPCGComputeKernel>(InDataBinding->Graph->GetKernel(KernelIndex));
+			const UPCGComputeKernel* Kernel = Cast<UPCGComputeKernel>(ComputeGraph->GetKernel(KernelIndex));
 
 			if (ensure(Kernel))
 			{
@@ -290,14 +292,16 @@ void FPCGKernelParamsCache::CreateAndCacheKernelParams(FPCGContext* InContext, c
 	TRACE_CPUPROFILER_EVENT_SCOPE(FPCGKernelParamsCache::CreateAndCacheKernelParams);
 	check(InDataBinding);
 
-	if (!ensure(InDataBinding->Graph))
+	const UPCGComputeGraph* ComputeGraph = InDataBinding->GetComputeGraph();
+
+	if (!ensure(ComputeGraph))
 	{
 		return;
 	}
 
-	for (int32 KernelIndex = 0; KernelIndex < InDataBinding->Graph->GetNumKernels(); ++KernelIndex)
+	for (int32 KernelIndex = 0; KernelIndex < ComputeGraph->GetNumKernels(); ++KernelIndex)
 	{
-		const UPCGComputeKernel* Kernel = Cast<UPCGComputeKernel>(InDataBinding->Graph->GetKernel(KernelIndex));
+		const UPCGComputeKernel* Kernel = Cast<UPCGComputeKernel>(ComputeGraph->GetKernel(KernelIndex));
 		const UPCGSettings* Settings = ensure(Kernel) ? Kernel->GetSettings() : nullptr;
 
 		if (!Settings)
@@ -325,7 +329,7 @@ void FPCGKernelParamsCache::CreateAndCacheKernelParams(FPCGContext* InContext, c
 			for (const TPair<FPCGKernelOverridableParam, int32>& OverridableParamAndDataIndex : *OverridableParams)
 			{
 				const FPCGKernelOverridableParam& OverridableParam = OverridableParamAndDataIndex.Key;
-				const UPCGData* Data = InDataBinding->InputDataCollection.TaggedData[OverridableParamAndDataIndex.Value].Data;
+				const UPCGData* Data = InDataBinding->GetInputDataCollection().TaggedData[OverridableParamAndDataIndex.Value].Data;
 				const FProperty* OverrideProperty = KernelParamsStructType ? KernelParamsStructType->FindPropertyByName(OverridableParam.Label) : nullptr;
 
 				if (!OverrideProperty)
diff --git a/Engine/Plugins/PCG/Source/PCG/Private/Elements/Metadata/PCGMetadataPartitionKernel.cpp b/Engine/Plugins/PCG/Source/PCG/Private/Elements/Metadata/PCGMetadataPartitionKernel.cpp
--- a/Engine/Plugins/PCG/Source/PCG/Private/Elements/Metadata/PCGMetadataPartitionKernel.cpp
+++ b/Engine/Plugins/PCG/Source/PCG/Private/Elements/Metadata/PCGMetadataPartitionKernel.cpp
@@ -127,10 +127,10 @@ TSharedPtr<const FPCGDataCollectionDesc> UPCGMetadataPartitionKernel::ComputeOut
 	{
 		const int32 AnalysisDataIndex = AnalysisDataIndices[DataIndex];
 		
-		check(InBinding->InputDataCollection.TaggedData.IsValidIndex(AnalysisDataIndex));
+		check(InBinding->GetInputDataCollection().TaggedData.IsValidIndex(AnalysisDataIndex));
 
 		// We can expect the analysis results data to already have been readback from the GPU because data description caching runs after the pre-execute readback stage of the ComputeGraph execution.
-		const UPCGParamData* AnalysisResultsData = Cast<UPCGParamData>(InBinding->InputDataCollection.TaggedData[AnalysisDataIndex].Data);
+		const UPCGParamData* AnalysisResultsData = Cast<UPCGParamData>(InBinding->GetInputDataCollection().TaggedData[AnalysisDataIndex].Data);
 		const UPCGMetadata* AnalysisMetadata = AnalysisResultsData ? AnalysisResultsData->ConstMetadata() : nullptr;
 
 		const FPCGMetadataAttributeBase* AnalysisCountAttributeBase = AnalysisMetadata ? AnalysisMetadata->GetConstAttribute(PCGCountUniqueAttributeValuesConstants::ValueCountAttributeName) : nullptr;
@@ -155,7 +155,7 @@ TSharedPtr<const FPCGDataCollectionDesc> UPCGMetadataPartitionKernel::ComputeOut
 		{
 			UE_LOG(LogPCG, Warning, TEXT("UPCGMetadataPartitionKernel: No analysis data received."));
 
-			if (InBinding->InputDataCollection.TaggedData[AnalysisDataIndex].Data->IsA<UPCGProxyForGPUData>())
+			if (InBinding->GetInputDataCollection().TaggedData[AnalysisDataIndex].Data->IsA<UPCGProxyForGPUData>())
 			{
 				UE_LOG(LogPCG, Error, TEXT("Data was not read back."));
 			}
diff --git a/Engine/Plugins/PCG/Source/PCG/Private/Elements/PCGCopyPointsKernel.cpp b/Engine/Plugins/PCG/Source/PCG/Private/Elements/PCGCopyPointsKernel.cpp
--- a/Engine/Plugins/PCG/Source/PCG/Private/Elements/PCGCopyPointsKernel.cpp
+++ b/Engine/Plugins/PCG/Source/PCG/Private/Elements/PCGCopyPointsKernel.cpp
@@ -66,7 +66,7 @@ TSharedPtr<const FPCGDataCollectionDesc> UPCGCopyPointsKernel::ComputeOutputBind
 
 		if (AnalysisDataIndex != INDEX_NONE)
 		{
-			const UPCGParamData* AnalysisResultsData = Cast<UPCGParamData>(InBinding->InputDataCollection.TaggedData[AnalysisDataIndex].Data);
+			const UPCGParamData* AnalysisResultsData = Cast<UPCGParamData>(InBinding->GetInputDataCollection().TaggedData[AnalysisDataIndex].Data);
 			const UPCGMetadata* AnalysisMetadata = AnalysisResultsData ? AnalysisResultsData->ConstMetadata() : nullptr;
 			const FPCGMetadataAttributeBase* AnalysisAttributeBase = AnalysisMetadata ? AnalysisMetadata->GetConstAttribute(PCGCopyPointsConstants::SelectedFlagAttributeName) : nullptr;
 
@@ -89,7 +89,7 @@ TSharedPtr<const FPCGDataCollectionDesc> UPCGCopyPointsKernel::ComputeOutputBind
 			else
 			{
 				UE_LOG(LogPCG, Warning, TEXT("No analysis data received by copy points kernel, no points will be copied."));
-				if (InBinding->InputDataCollection.TaggedData[AnalysisDataIndex].Data->IsA<UPCGProxyForGPUData>())
+				if (InBinding->GetInputDataCollection().TaggedData[AnalysisDataIndex].Data->IsA<UPCGProxyForGPUData>())
 				{
 					UE_LOG(LogPCG, Error, TEXT("Data was not read back."));
 				}
diff --git a/Engine/Plugins/PCG/Source/PCG/Public/Compute/PCGDataBinding.h b/Engine/Plugins/PCG/Source/PCG/Public/Compute/PCGDataBinding.h
--- a/Engine/Plugins/PCG/Source/PCG/Public/Compute/PCGDataBinding.h
+++ b/Engine/Plugins/PCG/Source/PCG/Public/Compute/PCGDataBinding.h
@@ -11,11 +11,8 @@
 
 class IPCGRuntimePrimitiveFactory;
 class UComputeKernel;
-class UPCGComponent;
 class UPCGComputeGraph;
 class UPCGComputeKernel;
-class UPCGProxyForGPUData;
-struct FPCGComputeGraphContext;
 struct FPCGContextHandle;
 struct FPCGKernelPin;
 
@@ -103,8 +100,8 @@ public:
 	PCG_API TSharedPtr<const FPCGDataCollectionDesc> GetCachedKernelPinDataDesc(const UPCGComputeKernel* InKernel, FName InPinLabel, bool bIsInput) const;
 	TSharedPtr<const FPCGDataCollectionDesc> GetCachedKernelPinDataDesc(int32 InGraphBindingIndex) const;
 
-	void ReceiveDataFromGPU_GameThread(UPCGData* InData, const UPCGSettings* InProducerSettings, EPCGExportMode InExportMode, FName InPinLabel, FName InPinLabelAlias);
-	void ReceiveDataFromGPU_GameThread(UPCGData* InData, const UPCGSettings* InProducerSettings, EPCGExportMode InExportMode, FName InPinLabel, FName InPinLabelAlias, const TSet<FString>& AdditionalTags);
+	PCG_API void ReceiveDataFromGPU_GameThread(UPCGData* InData, const UPCGSettings* InProducerSettings, EPCGExportMode InExportMode, FName InPinLabel, FName InPinLabelAlias);
+	PCG_API void ReceiveDataFromGPU_GameThread(UPCGData* InData, const UPCGSettings* InProducerSettings, EPCGExportMode InExportMode, FName InPinLabel, FName InPinLabelAlias, const TSet<FString>& AdditionalTags);
 
 	using FSourceBufferAndAttributeIndex = TPair<TSharedPtr<const FPCGProxyForGPUDataCollection>, /*Source Metadata Attribute Index*/int32>;
 	using FGraphAttributeIndexAndSourceBuffer = TPair</*Graph Metadata Attribute Index*/int32, TRefCountPtr<FRDGPooledBuffer>>;
@@ -123,18 +120,18 @@ public:
 	/** Map from string key value in source buffer to the key value usable in this graph. */
 	const TMap<FSourceBufferAndStringKey, /*Source Buffer String Key*/int32>& GetGraphToSourceBufferStringKey() const { return GraphToSourceBufferStringKey; }
 
-	IPCGGraphExecutionSource* GetExecutionSource() const;
+	PCG_API IPCGGraphExecutionSource* GetExecutionSource() const;
 
 	/** Helper to get index in input data collection of the first data item for the given kernel and input pin label. */
-	int32 GetFirstInputDataIndex(const UPCGComputeKernel* InKernel, FName InPinLabel) const;
+	PCG_API int32 GetFirstInputDataIndex(const UPCGComputeKernel* InKernel, FName InPinLabel) const;
 
 	/** Returns indices of data in the input data collection arriving on the given kernel pin. */
-	TArray<int32> GetPinInputDataIndices(const UPCGComputeKernel* InKernel, FName InPinLabel) const;
+	PCG_API TArray<int32> GetPinInputDataIndices(const UPCGComputeKernel* InKernel, FName InPinLabel) const;
 
 	/** If there is a GPU proxy at the given index of the input data collection, triggers a readback and replaces the data item with CPU data if readback succeeds.
 	* Returns false while readback is in progress.
 	*/
-	bool ReadbackInputDataToCPU(int32 InInputDataIndex);
+	PCG_API bool ReadbackInputDataToCPU(int32 InInputDataIndex);
 
 	void DebugLogDataDescriptions() const;
 
@@ -144,21 +141,36 @@ public:
 	bool InitializeKernelParams(FPCGContext* InContext);
 
 	/** Look up the cached parameter struct for a kernel. */
-	const FPCGKernelParams* GetCachedKernelParams(const UPCGComputeKernel* InKernel) const;
+	PCG_API const FPCGKernelParams* GetCachedKernelParams(const UPCGComputeKernel* InKernel) const;
 	const FPCGKernelParams* GetCachedKernelParams(int32 InKernelIndex) const;
 
+	TWeakPtr<FPCGContextHandle> GetContextHandle() const { return ContextHandle; }
+
+	TArray<FPCGDataToDebug>& GetDataToDebugMutable();
+	TArray<FPCGDataToDebug>& GetDataToInspectMutable();
+
+	void AddCompletedMeshSpawnerKernel(TObjectPtr<const UComputeKernel> InCompletedMeshSpawnerKernel);
+	bool IsMeshSpawnerKernelComplete(TObjectPtr<const UComputeKernel> InMeshSpawnerKernel) const;
+
+	void AddMeshSpawnerPrimitives(TObjectPtr<const UComputeKernel> InMeshSpawnerKernel, const FPCGSpawnerPrimitives& Primitives);
+	void AddMeshSpawnerPrimitives(TObjectPtr<const UComputeKernel> InMeshSpawnerKernel, FPCGSpawnerPrimitives&& Primitives);
+	FPCGSpawnerPrimitives& FindOrAddMeshSpawnerPrimitives(TObjectPtr<const UComputeKernel> InMeshSpawnerKernel);
+	FPCGSpawnerPrimitives* FindMeshSpawnerPrimitives(TObjectPtr<const UComputeKernel> InMeshSpawnerKernel);
+
+	const UPCGComputeGraph* GetComputeGraph() const;
+
+	PCG_API const FPCGDataCollection& GetInputDataCollection() const;
+	FPCGDataCollection& GetInputDataCollectionMutable();
+
+	const FPCGDataCollection& GetOutputDataCollection() const;
+
 private:
 	/** Loop over all metadata attributes in all input data ensure all attributes are registered in attribute table. */
 	void AddInputDataAttributesToTable();
 	void AddInputDataStringsToTable();
 	void AddInputDataTagsToTable();
 
-public:
-	/** Compute graph element output data. Data items are labelled with unique virtual output pin labels so that the can be routed correctly by the graph
-	 * executor to downstream nodes. */
-	UPROPERTY()
-	FPCGDataCollection OutputDataCollection;
-
+private:
 	UPROPERTY()
 	TObjectPtr<const UPCGComputeGraph> Graph = nullptr;
 
@@ -166,10 +178,15 @@ public:
 	UPROPERTY()
 	FPCGDataCollection InputDataCollection;
 
+	/** Compute graph element output data. Data items are labeled with unique virtual output pin labels so that the can be routed correctly by the graph
+	 * executor to downstream nodes. */
+	UPROPERTY()
+	FPCGDataCollection OutputDataCollection;
+
 	UPROPERTY()
 	TMap<TObjectPtr<const UComputeKernel>, FPCGSpawnerPrimitives> MeshSpawnersToPrimitives;
 
-	UPROPERTY()	
+	UPROPERTY()
 	TArray<TObjectPtr<const UComputeKernel>> CompletedMeshSpawners;
 
 	UPROPERTY()
@@ -180,9 +197,9 @@ public:
 
 	TWeakPtr<FPCGContextHandle> ContextHandle;
 
-protected:
 	/** All attributes present in graph at execution time. Coherent across all graph branches. Seeded using statically known created attributes at compile time, then
-	 * augmented with incoming attributes from input data collection at runtime. */
+	 * augmented with incoming attributes from input data collection at runtime.
+	 */
 	UPROPERTY()
 	FPCGKernelAttributeTable AttributeTable;
 
diff --git a/Engine/Plugins/PCG/Source/PCG/Public/Compute/PCGKernelParamsCache.h b/Engine/Plugins/PCG/Source/PCG/Public/Compute/PCGKernelParamsCache.h
--- a/Engine/Plugins/PCG/Source/PCG/Public/Compute/PCGKernelParamsCache.h
+++ b/Engine/Plugins/PCG/Source/PCG/Public/Compute/PCGKernelParamsCache.h
@@ -14,9 +14,9 @@ struct FPCGContext;
 struct FPCGKernelParams
 {
 public:
-	bool GetValueBool(FName InParamName) const;
-	int GetValueInt(FName InParamName) const;
-	FName GetValueName(FName InParamName) const;
+	PCG_API bool GetValueBool(FName InParamName) const;
+	PCG_API int GetValueInt(FName InParamName) const;
+	PCG_API FName GetValueName(FName InParamName) const;
 
 	template<typename EnumValueType>
 	EnumValueType GetValueEnum(FName InParamName) const
