[PCG] Allow to delete generated actors in tool mode

#rnx

diff --git a/Engine/Plugins/PCG/Source/PCGEditor/Private/EditorMode/Tools/Helpers/PCGEdModeActorHelpers.cpp b/Engine/Plugins/PCG/Source/PCGEditor/Private/EditorMode/Tools/Helpers/PCGEdModeActorHelpers.cpp
--- a/Engine/Plugins/PCG/Source/PCGEditor/Private/EditorMode/Tools/Helpers/PCGEdModeActorHelpers.cpp
+++ b/Engine/Plugins/PCG/Source/PCGEditor/Private/EditorMode/Tools/Helpers/PCGEdModeActorHelpers.cpp
@@ -37,7 +37,8 @@ namespace UE::PCG::EditorMode
 			FActorSpawnParameters SpawnParams;
 			SpawnParams.Name = DefaultWorkingActorName;
 			SpawnParams.NameMode = FActorSpawnParameters::ESpawnActorNameMode::Requested;
-			SpawnParams.ObjectFlags = RF_Transient;
+			// Actor needs to be transactional, otherwise it can't be deleted.
+			SpawnParams.ObjectFlags = RF_Transient | RF_Transactional;
 
 			UPCGActorHelpers::FSpawnDefaultActorParams SpawnDefaultActorParams(World, ActorClass, Transform, SpawnParams);
 
