[PCG] Disallow actor selection when the tool is enabled

#rnx
#okforigithub public

diff --git a/Engine/Plugins/PCG/Source/PCGEditor/Private/EditorMode/PCGEdMode.cpp b/Engine/Plugins/PCG/Source/PCGEditor/Private/EditorMode/PCGEdMode.cpp
--- a/Engine/Plugins/PCG/Source/PCGEditor/Private/EditorMode/PCGEdMode.cpp
+++ b/Engine/Plugins/PCG/Source/PCGEditor/Private/EditorMode/PCGEdMode.cpp
@@ -3,8 +3,12 @@
 #include "EditorMode/PCGEdMode.h"
 
 #include "EditorMode/PCGEdModeCommands.h"
+#include "EditorMode/PCGEdModeSettings.h"
+#include "EditorMode/PCGEdModeStyle.h"
 #include "EditorMode/PCGEdModeToolkit.h"
+#include "EditorMode/Tools/PCGInteractiveToolSettings.h"
 #include "EditorMode/Tools/Line/PCGDrawSplineTool.h"
+#include "EditorMode/Tools/Paint/PCGPaintTool.h"
 #include "EditorMode/Tools/Volume/PCGVolumeTool.h"
 
 #include "ContextObjectStore.h"
@@ -28,9 +32,6 @@
 
 #include "ToolContextInterfaces.h"
 #include "ToolTargetManager.h"
-#include "EditorMode/PCGEdModeSettings.h"
-#include "EditorMode/PCGEdModeStyle.h"
-#include "EditorMode/Tools/Paint/PCGPaintTool.h"
 #include "ToolTargets/PrimitiveComponentToolTarget.h"
 #include "ViewportToolbar/UnrealEdViewportToolbar.h"
 
@@ -86,11 +87,6 @@ void UPCGEditorMode::UnregisterEditorMode()
 	FPCGEditorModeStyle::Unregister();
 }
 
-void UPCGEditorMode::ActorSelectionChangeNotify()
-{
-	// @todo_pcg: Validate and accept/deny selection.
-}
-
 bool UPCGEditorMode::ProcessEditDelete()
 {
 	if (UEdMode::ProcessEditDelete())
@@ -280,6 +276,21 @@ FBox UPCGEditorMode::ComputeCustomViewportFocus() const
 	return FBox();
 }
 
+bool UPCGEditorMode::IsSelectionAllowed(AActor* InActor, bool bInSelection) const
+{
+	if (!bIsToolActive)
+	{
+		return true;
+	}
+
+	if (const UPCGInteractiveToolSettings* ToolSettings = GetCurrentToolSettings())
+	{
+		return ToolSettings->IsSelectionAllowed(InActor, bInSelection);
+	}
+
+	return false;
+}
+
 void UPCGEditorMode::CreateToolkit()
 {
 	Toolkit = MakeShared<FPCGEditorModeToolkit>();
@@ -402,6 +413,31 @@ void UPCGEditorMode::ConfigureRealTimeViewportsOverride(const bool bEnable)
 	}
 }
 
+const UPCGInteractiveToolSettings* UPCGEditorMode::GetCurrentToolSettings() const
+{
+	UInteractiveToolManager* ToolManager = GetToolManager();
+
+	if (ToolManager && ToolManager->HasAnyActiveTool())
+	{
+		if (UInteractiveTool* Tool = ToolManager->GetActiveTool(EToolSide::Mouse))
+		{
+			TArray<UObject*> AllToolProperties = Tool->GetToolProperties(false);
+
+			UObject** FoundProperties = AllToolProperties.FindByPredicate([](const UObject* Candidate)
+			{
+				return Candidate->IsA<UPCGInteractiveToolSettings>();
+			});
+
+			if (FoundProperties)
+			{
+				return Cast<UPCGInteractiveToolSettings>(*FoundProperties);
+			}
+		}
+	}
+
+	return nullptr;
+}
+
 // @todo_pcg: Activate when custom tools are supported.
 // void UPCGEditorMode::RegisterCustomTool(
 // 	TSharedPtr<FUICommandInfo> UICommand,
diff --git a/Engine/Plugins/PCG/Source/PCGEditor/Private/EditorMode/Tools/PCGInteractiveToolSettings.cpp b/Engine/Plugins/PCG/Source/PCGEditor/Private/EditorMode/Tools/PCGInteractiveToolSettings.cpp
--- a/Engine/Plugins/PCG/Source/PCGEditor/Private/EditorMode/Tools/PCGInteractiveToolSettings.cpp
+++ b/Engine/Plugins/PCG/Source/PCGEditor/Private/EditorMode/Tools/PCGInteractiveToolSettings.cpp
@@ -88,7 +88,7 @@ bool UPCGInteractiveToolSettings::Initialize(UInteractiveTool* InTool)
 	else
 	{
 		AActor* SpawnedActor = Actor::SpawnWorking(World, FTransform::Identity, GetWorkingActorClass());
-		// our spawned actor is also the actor we are editing
+		// Our spawned actor is also the actor we are editing
 		EditActor = SpawnedActor;
 
 		GeneratedResources.SpawnedActor = SpawnedActor;
@@ -210,6 +210,11 @@ void UPCGInteractiveToolSettings::Shutdown()
 	}
 }
 
+bool UPCGInteractiveToolSettings::IsSelectionAllowed(AActor* InActor, bool bInSelection) const
+{
+	return InActor == EditActor;
+}
+
 void UPCGInteractiveToolSettings::InitializeWorkingData(FName DataInstanceIdentifier)
 {
 	UPCGComponent* PCGComponent = GetWorkingPCGComponent();
diff --git a/Engine/Plugins/PCG/Source/PCGEditor/Public/EditorMode/PCGEdMode.h b/Engine/Plugins/PCG/Source/PCGEditor/Public/EditorMode/PCGEdMode.h
--- a/Engine/Plugins/PCG/Source/PCGEditor/Public/EditorMode/PCGEdMode.h
+++ b/Engine/Plugins/PCG/Source/PCGEditor/Public/EditorMode/PCGEdMode.h
@@ -7,6 +7,8 @@
 
 #include "PCGEdMode.generated.h"
 
+class UPCGInteractiveToolSettings;
+
 /** The unique state for the editor to be in, specific to PCG's toolkits. */
 UCLASS(Transient)
 class UPCGEditorMode : public UEdMode
@@ -23,7 +25,6 @@ public:
 	static void UnregisterEditorMode();
 
 	//~Begin UEdMode interface
-	virtual void ActorSelectionChangeNotify() override;
 	virtual bool ProcessEditDelete() override;
 	virtual bool ProcessEditCut() override;
 	virtual bool CanAutoSave() const override;
@@ -37,6 +38,7 @@ public:
 	virtual bool ShouldToolStartBeAllowed(const FString& ToolIdentifier) const override;
 	virtual bool HasCustomViewportFocus() const override;
 	virtual FBox ComputeCustomViewportFocus() const override;
+	virtual bool IsSelectionAllowed(AActor* InActor, bool bInSelection) const override;
 	//~End UEdMode interface
 
 protected:
@@ -62,5 +64,7 @@ private:
 	// 	const TFunction<bool(UInteractiveToolManager*, EToolSide)>& CanExecuteAction,
 	// 	const TFunction<bool(UInteractiveToolManager*, EToolSide)>& IsActionChecked) const;
 
+	const UPCGInteractiveToolSettings* GetCurrentToolSettings() const;
+
 	bool bIsToolActive = false;
 };
diff --git a/Engine/Plugins/PCG/Source/PCGEditor/Public/EditorMode/Tools/PCGInteractiveToolSettings.h b/Engine/Plugins/PCG/Source/PCGEditor/Public/EditorMode/Tools/PCGInteractiveToolSettings.h
--- a/Engine/Plugins/PCG/Source/PCGEditor/Public/EditorMode/Tools/PCGInteractiveToolSettings.h
+++ b/Engine/Plugins/PCG/Source/PCGEditor/Public/EditorMode/Tools/PCGInteractiveToolSettings.h
@@ -36,6 +36,9 @@ public:
 	virtual void Cancel(UInteractiveTool* OwningTool);
 	virtual void Shutdown();
 
+	// By default, when a tool is active, we can only select the Edit Actor in the Outliner.
+	virtual bool IsSelectionAllowed(AActor* InActor, bool bInSelection) const;
+
 	/** Initializes or finds the working data associated with the data instance identifier.
 	 * If NAME_None, the ToolTag is used as identifier. */
 	void InitializeWorkingData(FName DataInstanceIdentifier = NAME_None);
