[PCG Editor Mode] Prevent re-initializing editor mode tools if its already selected

#rnx

diff --git a/Engine/Plugins/PCG/Source/PCGEditor/Private/EditorMode/PCGEdMode.cpp b/Engine/Plugins/PCG/Source/PCGEditor/Private/EditorMode/PCGEdMode.cpp
--- a/Engine/Plugins/PCG/Source/PCGEditor/Private/EditorMode/PCGEdMode.cpp
+++ b/Engine/Plugins/PCG/Source/PCGEditor/Private/EditorMode/PCGEdMode.cpp
@@ -235,10 +235,12 @@ void UPCGEditorMode::Exit()
 
 bool UPCGEditorMode::ShouldToolStartBeAllowed(const FString& ToolIdentifier) const
 {
-	if (UInteractiveToolManager* ToolManager = GetToolManager())
+	UInteractiveToolManager* ToolManager = GetToolManager();
+	if (ToolManager && ToolManager->HasActiveTool(EToolSide::Left) && ToolManager->GetActiveToolName(EToolSide::Left) == ToolIdentifier)
 	{
-		// @todo_pcg: As more tools are added, some may require scene compatibility.
+		return false;
 	}
+	// @todo_pcg: As more tools are added, some may require scene compatibility.
 
 	return Super::ShouldToolStartBeAllowed(ToolIdentifier);
 }
