[PCG Editor Mode] Make PCG tools tooltip use the ToolTag defined in code.

#rnx

diff --git a/Engine/Plugins/PCG/Source/PCGEditor/Private/EditorMode/PCGEdModeCommands.cpp b/Engine/Plugins/PCG/Source/PCGEditor/Private/EditorMode/PCGEdModeCommands.cpp
--- a/Engine/Plugins/PCG/Source/PCGEditor/Private/EditorMode/PCGEdModeCommands.cpp
+++ b/Engine/Plugins/PCG/Source/PCGEditor/Private/EditorMode/PCGEdModeCommands.cpp
@@ -69,12 +69,12 @@ FPCGEditorModeToolCommands::FPCGEditorModeToolCommands()
  * @param ToolTip The tooltip for the tool
  * @param IconStyleName The style name for the icon to be used for the button
  */
-#define PCG_TOOL_COMMAND_INFO(CommandInfoSP, Name, Label, ToolTip, IconStyleName)\
+#define PCG_TOOL_COMMAND_INFO(CommandInfoSP, Name, Label, ToolTip, ToolTag, IconStyleName)\
 FUICommandInfo::MakeCommandInfo(AsShared()\
 	, CommandInfoSP\
 	, TEXT(Name"Tool")\
 	, LOCTEXT(Name"ToolLabel", Label)\
-	, LOCTEXT(Name"ToolTip", ToolTip)\
+	, FText::Format(LOCTEXT(Name"ToolTip", ToolTip), FFormatNamedArguments({{TEXT("ToolTag"), FText::FromName(ToolTag)}}))\
 	, FSlateIcon(FPCGEditorModeStyle::Get().GetStyleSetName(), IconStyleName)\
 	, EUserInterfaceActionType::ToggleButton\
 	, FInputChord())
@@ -84,10 +84,11 @@ void FPCGEditorModeToolCommands::RegisterCommands()
 	TInteractiveToolCommands::RegisterCommands();
 
 	/*** PCG EDITOR MODE TOOLS **/
-	PCG_TOOL_COMMAND_INFO(EnableDrawSplineTool, "DrawSpline", "Draw Spline", "Use the 'Draw Spline' tool to build a spline for PCG.", "PCGEditorMode.Tools.DrawSpline");
-	PCG_TOOL_COMMAND_INFO(EnableDrawSurfaceTool, "DrawSurface", "Draw Surface", "Use the 'Draw Surface' tool to build a spline for PCG.", "PCGEditorMode.Tools.DrawSurface");
-	PCG_TOOL_COMMAND_INFO(EnablePaintTool, "Paint", "Paint", "Use the 'Paint' tool to apply PCG points with a brush.", "PCGEditorMode.Context.Paint");
-	PCG_TOOL_COMMAND_INFO(EnableVolumeTool, "Volume", "Volume", "Create or updates a volume in the scene.", "PCGEditorMode.Tools.Volume");
+
+	PCG_TOOL_COMMAND_INFO(EnableDrawSplineTool, "DrawSpline", "Draw Spline", "Use the '{ToolTag}' tool to build a spline for PCG.", UPCGInteractiveToolSettings_Spline::StaticGetToolTag(), "PCGEditorMode.Tools.DrawSpline");
+	PCG_TOOL_COMMAND_INFO(EnableDrawSurfaceTool, "DrawSurface", "Draw Surface", "Use the '{ToolTag}' tool to build a closed spline representing an area for PCG.", UPCGInteractiveToolSettings_SplineSurface::StaticGetToolTag(), "PCGEditorMode.Tools.DrawSurface");
+	PCG_TOOL_COMMAND_INFO(EnablePaintTool, "Paint", "Paint", "Use the '{ToolTag}' tool to apply PCG points with a brush.", UPCGInteractiveToolSettings_PaintTool::StaticGetToolTag(), "PCGEditorMode.Context.Paint");
+	PCG_TOOL_COMMAND_INFO(EnableVolumeTool, "Volume", "Volume", "Create or updates a volume in the scene using the '{ToolTag}' tool.", UPCGInteractiveToolSettings_Volume::StaticGetToolTag(), "PCGEditorMode.Tools.Volume");
 }
 
 void FPCGEditorModeToolCommands::RegisterAllToolCommands()
diff --git a/Engine/Plugins/PCG/Source/PCGEditor/Private/EditorMode/Tools/Line/PCGDrawSplineTool.cpp b/Engine/Plugins/PCG/Source/PCGEditor/Private/EditorMode/Tools/Line/PCGDrawSplineTool.cpp
--- a/Engine/Plugins/PCG/Source/PCGEditor/Private/EditorMode/Tools/Line/PCGDrawSplineTool.cpp
+++ b/Engine/Plugins/PCG/Source/PCGEditor/Private/EditorMode/Tools/Line/PCGDrawSplineTool.cpp
@@ -236,6 +236,11 @@ TArray<USplineComponent*> UPCGInteractiveToolSettings_SplineBase::GetWorkingSpli
 	return Result;
 }
 
+FName UPCGInteractiveToolSettings_Spline::StaticGetToolTag()
+{
+	return PCGDrawSplineTool::Constants::SplineToolTag;
+}
+
 const UScriptStruct* UPCGInteractiveToolSettings_Spline::GetWorkingDataType() const
 {
 	return FPCGInteractiveToolWorkingData_Spline::StaticStruct();
@@ -243,7 +248,7 @@ const UScriptStruct* UPCGInteractiveToolSettings_Spline::GetWorkingDataType() co
 
 FName UPCGInteractiveToolSettings_Spline::GetToolTag() const
 {
-	return PCGDrawSplineTool::Constants::SplineToolTag;
+	return StaticGetToolTag();
 }
 
 void UPCGInteractiveToolSettings_Spline::OnWorkingDataInitialized(FPCGInteractiveToolWorkingData* WorkingData) const
@@ -254,6 +259,11 @@ void UPCGInteractiveToolSettings_Spline::OnWorkingDataInitialized(FPCGInteractiv
 	}
 }
 
+FName UPCGInteractiveToolSettings_SplineSurface::StaticGetToolTag()
+{
+	return PCGDrawSplineTool::Constants::SplineSurfaceToolTag;
+}
+
 const UScriptStruct* UPCGInteractiveToolSettings_SplineSurface::GetWorkingDataType() const
 {
 	return FPCGInteractiveToolWorkingData_SplineSurface::StaticStruct();
@@ -261,7 +271,7 @@ const UScriptStruct* UPCGInteractiveToolSettings_SplineSurface::GetWorkingDataTy
 
 FName UPCGInteractiveToolSettings_SplineSurface::GetToolTag() const
 {
-	return PCGDrawSplineTool::Constants::SplineSurfaceToolTag;
+	return StaticGetToolTag();
 }
 
 void UPCGInteractiveToolSettings_Spline::OnPropertyModified(UObject* Object, FProperty* Property)
diff --git a/Engine/Plugins/PCG/Source/PCGEditor/Private/EditorMode/Tools/Paint/PCGPaintTool.cpp b/Engine/Plugins/PCG/Source/PCGEditor/Private/EditorMode/Tools/Paint/PCGPaintTool.cpp
--- a/Engine/Plugins/PCG/Source/PCGEditor/Private/EditorMode/Tools/Paint/PCGPaintTool.cpp
+++ b/Engine/Plugins/PCG/Source/PCGEditor/Private/EditorMode/Tools/Paint/PCGPaintTool.cpp
@@ -40,11 +40,16 @@ UPCGInteractiveToolSettings_PaintTool::UPCGInteractiveToolSettings_PaintTool()
 	RaycastRuleCollection.RaycastRules.Add(IgnorePCGGeneratedComponentsRule);
 } 
 
-FName UPCGInteractiveToolSettings_PaintTool::GetToolTag() const
+FName UPCGInteractiveToolSettings_PaintTool::StaticGetToolTag()
 {
 	return PCGPaintTool::Constants::ToolTag;
 }
 
+FName UPCGInteractiveToolSettings_PaintTool::GetToolTag() const
+{
+	return StaticGetToolTag();
+}
+
 const UScriptStruct* UPCGInteractiveToolSettings_PaintTool::GetWorkingDataType() const
 {
 	return FPCGInteractiveToolWorkingData_PointArrayData::StaticStruct();
diff --git a/Engine/Plugins/PCG/Source/PCGEditor/Private/EditorMode/Tools/Volume/PCGVolumeTool.cpp b/Engine/Plugins/PCG/Source/PCGEditor/Private/EditorMode/Tools/Volume/PCGVolumeTool.cpp
--- a/Engine/Plugins/PCG/Source/PCGEditor/Private/EditorMode/Tools/Volume/PCGVolumeTool.cpp
+++ b/Engine/Plugins/PCG/Source/PCGEditor/Private/EditorMode/Tools/Volume/PCGVolumeTool.cpp
@@ -31,11 +31,16 @@ UPCGInteractiveToolSettings_Volume::UPCGInteractiveToolSettings_Volume()
 	RaycastRuleCollection.RaycastRules.Add(IgnorePCGGeneratedComponentsRule);
 }
 
-FName UPCGInteractiveToolSettings_Volume::GetToolTag() const
+FName UPCGInteractiveToolSettings_Volume::StaticGetToolTag()
 {
 	return PCGVolumeTool::Constants::ToolTag;
 }
 
+FName UPCGInteractiveToolSettings_Volume::GetToolTag() const
+{
+	return StaticGetToolTag();
+}
+
 const UScriptStruct* UPCGInteractiveToolSettings_Volume::GetWorkingDataType() const
 {
 	return FPCGInteractiveToolWorkingData_Volume::StaticStruct();
diff --git a/Engine/Plugins/PCG/Source/PCGEditor/Public/EditorMode/Tools/Line/PCGDrawSplineTool.h b/Engine/Plugins/PCG/Source/PCGEditor/Public/EditorMode/Tools/Line/PCGDrawSplineTool.h
--- a/Engine/Plugins/PCG/Source/PCGEditor/Public/EditorMode/Tools/Line/PCGDrawSplineTool.h
+++ b/Engine/Plugins/PCG/Source/PCGEditor/Public/EditorMode/Tools/Line/PCGDrawSplineTool.h
@@ -69,6 +69,8 @@ protected:
 public:
 	UPROPERTY(EditAnywhere, Category = "Spline")
 	bool bClosedSpline = false;
+
+	static FName StaticGetToolTag();
 };
 
 
@@ -80,6 +82,9 @@ class PCGEDITOR_API UPCGInteractiveToolSettings_SplineSurface : public UPCGInter
 protected:
 	virtual const UScriptStruct* GetWorkingDataType() const override;
 	virtual FName GetToolTag() const override;
+
+public:
+	static FName StaticGetToolTag();
 };
 
 UCLASS(Abstract, Experimental)
diff --git a/Engine/Plugins/PCG/Source/PCGEditor/Public/EditorMode/Tools/Paint/PCGPaintTool.h b/Engine/Plugins/PCG/Source/PCGEditor/Public/EditorMode/Tools/Paint/PCGPaintTool.h
--- a/Engine/Plugins/PCG/Source/PCGEditor/Public/EditorMode/Tools/Paint/PCGPaintTool.h
+++ b/Engine/Plugins/PCG/Source/PCGEditor/Public/EditorMode/Tools/Paint/PCGPaintTool.h
@@ -20,6 +20,8 @@ class PCGEDITOR_API UPCGInteractiveToolSettings_PaintTool : public UPCGInteracti
 public:
 	UPCGInteractiveToolSettings_PaintTool();
 
+	static FName StaticGetToolTag();
+
 	// ~Begin UPCGInteractiveToolSettings interface
 	virtual FName GetToolTag() const override;
 	virtual const UScriptStruct* GetWorkingDataType() const override;
diff --git a/Engine/Plugins/PCG/Source/PCGEditor/Public/EditorMode/Tools/Volume/PCGVolumeTool.h b/Engine/Plugins/PCG/Source/PCGEditor/Public/EditorMode/Tools/Volume/PCGVolumeTool.h
--- a/Engine/Plugins/PCG/Source/PCGEditor/Public/EditorMode/Tools/Volume/PCGVolumeTool.h
+++ b/Engine/Plugins/PCG/Source/PCGEditor/Public/EditorMode/Tools/Volume/PCGVolumeTool.h
@@ -20,6 +20,8 @@ class PCGEDITOR_API UPCGInteractiveToolSettings_Volume : public UPCGInteractiveT
 public:
 	UPCGInteractiveToolSettings_Volume();
 
+	static FName StaticGetToolTag();
+
 	// ~Begin UPCGInteractiveToolSettings interface
 	virtual FName GetToolTag() const override;
 	virtual const UScriptStruct* GetWorkingDataType() const override;
