[PCG] Made sure spline data is always constructed from valid data.

#rnx

diff --git a/Engine/Plugins/PCG/Source/PCG/Private/Data/PCGPolyLineData.cpp b/Engine/Plugins/PCG/Source/PCG/Private/Data/PCGPolyLineData.cpp
--- a/Engine/Plugins/PCG/Source/PCG/Private/Data/PCGPolyLineData.cpp
+++ b/Engine/Plugins/PCG/Source/PCG/Private/Data/PCGPolyLineData.cpp
@@ -36,7 +36,13 @@ FVector::FReal UPCGPolyLineData::GetLength() const
 
 int UPCGPolyLineData::GetNumVertices() const
 {
-	return GetNumSegments() + (IsClosed() ? 0 : 1);
+	int NumPoints = GetNumSegments();
+	if (NumPoints > 0 && !IsClosed())
+	{
+		++NumPoints;
+	}
+	
+	return NumPoints;
 }
 
 FTransform UPCGPolyLineData::K2_GetTransformAtDistance(int SegmentIndex, double Distance, FBox& OutBounds, bool bWorldSpace) const
diff --git a/Engine/Plugins/PCG/Source/PCG/Private/Data/Registry/PCGGetDataFunctionRegistry.cpp b/Engine/Plugins/PCG/Source/PCG/Private/Data/Registry/PCGGetDataFunctionRegistry.cpp
--- a/Engine/Plugins/PCG/Source/PCG/Private/Data/Registry/PCGGetDataFunctionRegistry.cpp
+++ b/Engine/Plugins/PCG/Source/PCG/Private/Data/Registry/PCGGetDataFunctionRegistry.cpp
@@ -258,7 +258,7 @@ int FPCGGetDataFunctionRegistry::DefaultDataFromComponent(FPCGContext* InContext
 			SplineData->Initialize(LandscapeSplineComponent);
 			TaggedData.Data = SplineData;
 		}
-		else if (USplineComponent* SplineComponent = Cast<USplineComponent>(InComponent))
+		else if (USplineComponent* SplineComponent = Cast<USplineComponent>(InComponent); SplineComponent && SplineComponent->GetNumberOfSplinePoints() > 0)
 		{
 			UPCGSplineData* SplineData = FPCGContext::NewObject_AnyThread<UPCGSplineData>(InContext);
 			SplineData->Initialize(SplineComponent);
diff --git a/Engine/Plugins/PCG/Source/PCG/Private/Elements/PCGGetSplineControlPoints.cpp b/Engine/Plugins/PCG/Source/PCG/Private/Elements/PCGGetSplineControlPoints.cpp
--- a/Engine/Plugins/PCG/Source/PCG/Private/Elements/PCGGetSplineControlPoints.cpp
+++ b/Engine/Plugins/PCG/Source/PCG/Private/Elements/PCGGetSplineControlPoints.cpp
@@ -66,8 +66,7 @@ bool FPCGGetSplineControlPointsElement::ExecuteInternal(FPCGContext* InContext)
 			continue;
 		}
 
-		const bool bIsClosed = PolylineData->IsClosed();
-		const int32 NumPoints = bIsClosed ? PolylineData->GetNumSegments() : (PolylineData->GetNumSegments() + 1);
+		const int32 NumPoints = PolylineData->GetNumVertices();
 		const TConstArrayView<PCGMetadataEntryKey> MetadataEntries = PolylineData->GetConstVerticesEntryKeys();
 
 		UPCGBasePointData* PointData = FPCGContext::NewPointData_AnyThread(InContext);	
