[PCG Ed Mode] Improvements

- Moved spline wipe from ToolStart to Begin Drag in Free Draw mode
- Added "Reset" tool data API. Exposes button to reset tool data. Disabled for volume tool.
- Renamed ToolDrawTarget to RaycastSettings to differentiate the two better. RaycastRules determine validity of raycasts, RaycastSettings determine changes to raycast results (normals, offsets etc.)
- Shuffled some properties & categories
- Removed EditActor from UI. If it's editable, the tool would need to go through the initialization steps again to ensure PCG Component etc. is found. VisibleAnywhere is not necessary anymore with selection being unchanging. Something for the future perhaps.
- Fixed Actor Label not being applied to the spawned actor
- Fixed new pcg component name not being applied to a generated pcg component
- Actor Label and PCG Component Name are only displayed when a new PCG component is generated or a new actor is spawned.
They update post creation as well.

diff --git a/Engine/Plugins/PCG/Source/PCG/Private/Data/PCGToolData.cpp b/Engine/Plugins/PCG/Source/PCG/Private/Data/PCGToolData.cpp
--- a/Engine/Plugins/PCG/Source/PCG/Private/Data/PCGToolData.cpp
+++ b/Engine/Plugins/PCG/Source/PCG/Private/Data/PCGToolData.cpp
@@ -132,6 +132,17 @@ void FPCGInteractiveToolWorkingData_PointArrayData::OnToolCancel(const FPCGInter
 	
 	FPCGInteractiveToolWorkingData::OnToolCancel(Context);
 }
+
+void FPCGInteractiveToolWorkingData_PointArrayData::OnResetToolDataRequested(const FPCGInteractiveToolWorkingDataContext& Context)
+{
+	if (GeneratedPointData)
+	{
+		GeneratedPointData->SetNumPoints(0);
+
+		FPropertyChangedEvent EmptyPropertyChangedEvent(nullptr);
+		FCoreUObjectDelegates::OnObjectPropertyChanged.Broadcast(GeneratedPointData, EmptyPropertyChangedEvent);
+	}
+}
 #endif
 
 void FPCGInteractiveToolWorkingData_PointArrayData::InitializeRuntimeElementData(FPCGContext* Context) const
@@ -184,12 +195,6 @@ void FPCGInteractiveToolWorkingData_Spline::OnToolStart(const FPCGInteractiveToo
 	if(SplineComponent.IsValid())
 	{
 		OnToolStartSplineComponent = Cast<USplineComponent>(StaticDuplicateObject(SplineComponent.Get(), GetTransientPackage()));
-
-		// We wipe the spline points on tool start so it becomes easy to edit them after the fact.
-		// @todo_pcg
-		// When the Spline Tool exposes more controls, such as selecting spline points to move them, resetting the spline etc., we will remove this.
-		SplineComponent->bSplineHasBeenEdited = true;
-		SplineComponent->ClearSplinePoints(true);
 	}
 }
 
@@ -214,6 +219,17 @@ void FPCGInteractiveToolWorkingData_Spline::OnToolCancel(const FPCGInteractiveTo
 	FPCGInteractiveToolWorkingData::OnToolCancel(Context);
 }
 
+void FPCGInteractiveToolWorkingData_Spline::OnResetToolDataRequested(const FPCGInteractiveToolWorkingDataContext& Context)
+{
+	if (SplineComponent.IsValid())
+	{
+		SplineComponent->ClearSplinePoints();
+
+		FPropertyChangedEvent EmptyPropertyChangedEvent(nullptr);
+		FCoreUObjectDelegates::OnObjectPropertyChanged.Broadcast(SplineComponent.Get(), EmptyPropertyChangedEvent);
+	}
+}
+
 USplineComponent* FPCGInteractiveToolWorkingData_Spline::FindOrGenerateSplineComponent(const FPCGInteractiveToolWorkingDataContext& Context)
 {
 	USplineComponent* MatchingSplineComponent = FindMatchingSplineComponent(Context);
diff --git a/Engine/Plugins/PCG/Source/PCG/Public/Data/PCGToolData.h b/Engine/Plugins/PCG/Source/PCG/Public/Data/PCGToolData.h
--- a/Engine/Plugins/PCG/Source/PCG/Public/Data/PCGToolData.h
+++ b/Engine/Plugins/PCG/Source/PCG/Public/Data/PCGToolData.h
@@ -108,6 +108,9 @@ public:
 	/** Called when the tool is canceled. Revert changes you have made here. */
 	PCG_API virtual void OnToolCancel(const FPCGInteractiveToolWorkingDataContext& Context);
 
+	/** Called when reset data was specifically requested, for example by the user. */
+	virtual void OnResetToolDataRequested(const FPCGInteractiveToolWorkingDataContext& Context) {}
+	
 	PCG_API FOnInitialized& OnInitialized() { return OnInitializedDelegate; }
 private:
 	/** Implement your initialization logic here. Called by Initialize. */
@@ -147,6 +150,8 @@ public:
 	PCG_API virtual void OnToolStart(const FPCGInteractiveToolWorkingDataContext& Context) override;
 	
 	PCG_API virtual void OnToolCancel(const FPCGInteractiveToolWorkingDataContext& Context) override;
+
+	PCG_API virtual void OnResetToolDataRequested(const FPCGInteractiveToolWorkingDataContext& Context) override;
 #endif
 	
 protected:
@@ -177,7 +182,7 @@ public:
 public:
 	virtual void OnToolStart(const FPCGInteractiveToolWorkingDataContext& Context) override;
 	virtual void OnToolCancel(const FPCGInteractiveToolWorkingDataContext& Context) override;
-	
+	virtual void OnResetToolDataRequested(const FPCGInteractiveToolWorkingDataContext& Context) override;
 private:
 	virtual void InitializeInternal(const FPCGInteractiveToolWorkingDataContext& Context) override;
 	
diff --git a/Engine/Plugins/PCG/Source/PCGEditor/Private/Details/PCGInteractiveToolSettingsDetails.cpp b/Engine/Plugins/PCG/Source/PCGEditor/Private/Details/PCGInteractiveToolSettingsDetails.cpp
--- a/Engine/Plugins/PCG/Source/PCGEditor/Private/Details/PCGInteractiveToolSettingsDetails.cpp
+++ b/Engine/Plugins/PCG/Source/PCGEditor/Private/Details/PCGInteractiveToolSettingsDetails.cpp
@@ -8,6 +8,8 @@
 #include "DetailLayoutBuilder.h"
 #include "DetailWidgetRow.h"
 #include "EditorMode/Tools/PCGInteractiveToolSettings.h"
+#include "Widgets/Layout/SWidgetSwitcher.h"
+#include "Framework/Application/SlateApplication.h"
 
 #define LOCTEXT_NAMESPACE "PCGGraphInstanceDetails"
 
@@ -52,29 +54,61 @@ void FPCGInteractiveToolSettingsDetails::CustomizeDetails(IDetailLayoutBuilder&
 	}
 	
 	IDetailCategoryBuilder& CategoryBuilder = DetailBuilder.EditCategory("PCG");
-	UPCGGraphInstance* GraphInstance = ToolSettings.IsValid() ? ToolSettings->GetGraphInstance() : nullptr;
-	if(GraphInstance != nullptr && GraphInstance->GetUserParametersStruct()->GetNumPropertiesInBag() > 0)
-	{
-		TArray<TSharedRef<IPropertyHandle>> PropertyHandles;
-		CategoryBuilder.GetDefaultProperties(PropertyHandles);
 	
-		for(TSharedRef<IPropertyHandle> PropertyHandle : PropertyHandles)
-		{
-			CategoryBuilder.AddProperty(PropertyHandle);
+	TArray<TSharedRef<IPropertyHandle>> PropertyHandles;
+	CategoryBuilder.GetDefaultProperties(PropertyHandles);
+
+	for(TSharedRef<IPropertyHandle> PropertyHandle : PropertyHandles)
+	{
+		IDetailPropertyRow& PropertyRow = CategoryBuilder.AddProperty(PropertyHandle);
 
-			// We add the parameters right behind the graph property
-			if(PropertyHandle->GetProperty()->GetFName() == GET_MEMBER_NAME_CHECKED(UPCGInteractiveToolSettings, ToolGraph))
+		// We add the parameters right behind the graph property
+		if(PropertyHandle->GetProperty()->GetFName() == GET_MEMBER_NAME_CHECKED(UPCGInteractiveToolSettings, ToolGraph))
+		{
+			UPCGGraphInstance* GraphInstance = ToolSettings.IsValid() ? ToolSettings->GetGraphInstance() : nullptr;
+			if(GraphInstance != nullptr && GraphInstance->GetUserParametersStruct()->GetNumPropertiesInBag() > 0)
 			{
-				IDetailPropertyRow* PropertyRow = CategoryBuilder.AddExternalObjectProperty(
+				IDetailPropertyRow* ParametersPropertyRow = CategoryBuilder.AddExternalObjectProperty(
 				{GraphInstance}, 
 				GET_MEMBER_NAME_CHECKED(UPCGGraphInstance, ParametersOverrides),
 				 EPropertyLocation::Default,
 				 FAddPropertyParams());
 		
-				PropertyRow->ShouldAutoExpand(true);
+				ParametersPropertyRow->ShouldAutoExpand(true);
 			}
-		}		
-	}
+		}
+
+		if(PropertyHandle->GetProperty()->GetFName() == GET_MEMBER_NAME_CHECKED(UPCGInteractiveToolSettings, DataInstance))
+		{
+			TSharedPtr<SWidget> NameWidget;
+			TSharedPtr<SWidget> ValueWidget;
+			PropertyRow.GetDefaultWidgets(NameWidget, ValueWidget, true);
+
+			PropertyRow.CustomWidget()
+			.NameContent()
+			[
+				NameWidget.ToSharedRef()
+			]
+			.ValueContent()
+			[
+				SNew(SHorizontalBox)
+				+ SHorizontalBox::Slot()
+				.FillWidth(1.f)
+				[
+					ValueWidget.ToSharedRef()
+				]
+				+ SHorizontalBox::Slot()
+				.AutoWidth()
+				.Padding(2.f)
+				[
+					SNew(SButton)
+					.Text(this, &FPCGInteractiveToolSettingsDetails::GetResetToolDataButtonText)
+					.Visibility(this, &FPCGInteractiveToolSettingsDetails::IsResetToolDataButtonVisible)
+					.OnClicked(this, &FPCGInteractiveToolSettingsDetails::ResetToolData)
+				]	
+			];
+		}
+	}	
 }
 
 void FPCGInteractiveToolSettingsDetails::CustomizeDetails(const TSharedPtr<IDetailLayoutBuilder>& DetailBuilder)
@@ -97,4 +131,74 @@ void FPCGInteractiveToolSettingsDetails::Rebuild(UObject* Object, const FPropert
 	}
 }
 
+FText FPCGInteractiveToolSettingsDetails::GetResetToolDataButtonText() const
+{
+	if (FSlateApplication::Get().GetModifierKeys().IsControlDown())
+	{
+		return LOCTEXT("ResetAllToolDataButton_Text", "Reset All");	
+	}
+	
+	return LOCTEXT("ResetToolDataButton_Text", "Reset");
+}
+
+FReply FPCGInteractiveToolSettingsDetails::ResetToolData() const
+{
+	if (FSlateApplication::Get().GetModifierKeys().IsControlDown())
+	{
+		FPCGInteractiveToolWorkingDataContext Context;
+		Context.PCGSettings = ToolSettings;
+		Context.InteractiveTool = ToolSettings->GetTypedOuter<UInteractiveTool>();
+		Context.OwningPCGComponent = ToolSettings->GetWorkingPCGComponent();
+		Context.OwningActor = Context.OwningPCGComponent.IsValid() ? Context.OwningPCGComponent->GetOwner() : nullptr;
+		for (const auto& WorkingDataTuple : ToolSettings->GetMutableTypedWorkingDataMap())
+		{
+			Context.DataInstanceIdentifier = WorkingDataTuple.Key;
+
+			if(ToolSettings->CanResetToolData(Context.DataInstanceIdentifier))
+			{
+				WorkingDataTuple.Value->OnResetToolDataRequested(Context);
+			}
+		}
+		
+	}
+	else
+	{
+		if (FPCGInteractiveToolWorkingData* WorkingData = ToolSettings->GetMutableTypedWorkingData(ToolSettings->DataInstance))
+		{
+			FPCGInteractiveToolWorkingDataContext Context;
+			Context.DataInstanceIdentifier = ToolSettings->DataInstance;
+			Context.PCGSettings = ToolSettings;
+			Context.InteractiveTool = ToolSettings->GetTypedOuter<UInteractiveTool>();
+			Context.OwningPCGComponent = ToolSettings->GetWorkingPCGComponent();
+			Context.OwningActor = Context.OwningPCGComponent.IsValid() ? Context.OwningPCGComponent->GetOwner() : nullptr;
+			
+			if(ToolSettings->CanResetToolData(Context.DataInstanceIdentifier))
+			{
+				WorkingData->OnResetToolDataRequested(Context);
+			}
+		}
+	}
+
+	return FReply::Handled();
+}
+
+EVisibility FPCGInteractiveToolSettingsDetails::IsResetToolDataButtonVisible() const
+{
+	if (FSlateApplication::Get().GetModifierKeys().IsControlDown())
+	{
+		bool bAny = false;
+		for (const auto& WorkingDataTuple : ToolSettings->GetMutableTypedWorkingDataMap())
+		{			
+			if(ToolSettings->CanResetToolData(WorkingDataTuple.Key))
+			{
+				bAny = true;
+			}
+		}
+
+		return bAny ? EVisibility::Visible : EVisibility::Collapsed;
+	}
+	
+	return ToolSettings->CanResetToolData(ToolSettings->DataInstance) ? EVisibility::Visible : EVisibility::Collapsed;
+}
+
 #undef LOCTEXT_NAMESPACE
diff --git a/Engine/Plugins/PCG/Source/PCGEditor/Private/EditorMode/Tools/Line/PCGDrawSplineTool.cpp b/Engine/Plugins/PCG/Source/PCGEditor/Private/EditorMode/Tools/Line/PCGDrawSplineTool.cpp
--- a/Engine/Plugins/PCG/Source/PCGEditor/Private/EditorMode/Tools/Line/PCGDrawSplineTool.cpp
+++ b/Engine/Plugins/PCG/Source/PCGEditor/Private/EditorMode/Tools/Line/PCGDrawSplineTool.cpp
@@ -303,7 +303,6 @@ void UPCGDrawSplineToolBase::Setup()
 	ClickOrDragBehavior->Initialize(this, this);
 	AddInputBehavior(ClickOrDragBehavior);
 	
-	// @todo_pcg: Make adjustments to this and check when properties change.
 	CachedRayParams =
 	{
 		.World = GetWorld(),
@@ -335,14 +334,14 @@ void UPCGDrawSplineToolBase::AddSplinePoint(const FVector& HitLocation, const FV
 	const FVector NewNormal = GetNewPointNormal(HitNormal);
 
 	FVector FinalLocation = HitLocation;
-	switch (ToolSettings->Target.OffsetMode)
+	switch (ToolSettings->Settings.OffsetMode)
 	{
 		case EPCGToolDrawTargetOffset::DistanceOffset:
 			check(HitNormal.IsNormalized());
-			FinalLocation += ToolSettings->Target.OffsetDistance * NewNormal;
+			FinalLocation += ToolSettings->Settings.OffsetDistance * NewNormal;
 			break;
 		case EPCGToolDrawTargetOffset::ExplicitOffset:
-			FinalLocation += ToolSettings->Target.OffsetVector;
+			FinalLocation += ToolSettings->Settings.OffsetVector;
 			break;
 		default: // fall-through
 		case EPCGToolDrawTargetOffset::HitLocation:
@@ -369,7 +368,7 @@ FVector UPCGDrawSplineToolBase::GetNewPointNormal(const FVector& HitNormal) cons
 	
 	USplineComponent* WorkingSpline = ToolSettings->GetWorkingSplineComponent(ToolSettings->DataInstance);
 	const int32 NumSplinePoints = WorkingSpline->GetNumberOfSplinePoints();
-	switch (ToolSettings->Target.NormalMode)
+	switch (ToolSettings->Settings.NormalMode)
 	{
 		case EPCGToolDrawTargetNormal::AlignToPrevious:
 		{
@@ -396,7 +395,7 @@ FVector UPCGDrawSplineToolBase::GetNewPointNormal(const FVector& HitNormal) cons
 		case EPCGToolDrawTargetNormal::WorldUp:
 			return FVector::UnitZ();
 		case EPCGToolDrawTargetNormal::Explicit:
-			return ToolSettings->Target.Normal;
+			return ToolSettings->Settings.Normal;
 		default: // fall-through
 		case EPCGToolDrawTargetNormal::HitNormal:
 			return HitNormal;
@@ -480,6 +479,16 @@ void UPCGDrawSplineToolBase::OnClickPress(const FInputDeviceRay& ClickPosition)
 	if (TOptional<FViewHitResult> Result = ViewportRay(ClickPosition, CachedRayParams))
 	{
 		UPCGInteractiveToolSettings_SplineBase* ToolSettings = GetSplineSettings();
+
+		// For "continuous" draw modes like FreeDraw, we reset the spline data so so the user can re-drag until happy.
+		// @todo_pcg
+		// When the Spline Tool exposes more controls, such as selecting spline points to move them, resetting the spline etc., we will remove this.
+		if (ToolSettings->DrawMode == EPCGDrawSplineDrawMode::FreeDraw)
+		{
+			USplineComponent* SplineComponent = ToolSettings->GetWorkingSplineComponent(ToolSettings->DataInstance);
+			SplineComponent->bSplineHasBeenEdited = true;
+			SplineComponent->ClearSplinePoints(true);
+		}
 		
 		AddSplinePoint(Result->ImpactPosition, Result->ImpactNormal);
 		
@@ -544,14 +553,14 @@ void UPCGDrawSplineToolBase::OnClickDrag(const FInputDeviceRay& DragPosition)
 					FVector FinalLocation = Result->ImpactPosition;
 					const FVector Normal = GetNewPointNormal(Result->ImpactNormal);
 
-					switch (ToolSettings->Target.OffsetMode)
+					switch (ToolSettings->Settings.OffsetMode)
 					{
 						case EPCGToolDrawTargetOffset::DistanceOffset:
 							check(Result->ImpactNormal.IsNormalized());
-							FinalLocation += ToolSettings->Target.OffsetDistance * Normal;
+							FinalLocation += ToolSettings->Settings.OffsetDistance * Normal;
 							break;
 						case EPCGToolDrawTargetOffset::ExplicitOffset:
-							FinalLocation += ToolSettings->Target.OffsetVector;
+							FinalLocation += ToolSettings->Settings.OffsetVector;
 							break;
 						default: // fall-through
 						case EPCGToolDrawTargetOffset::HitLocation:
diff --git a/Engine/Plugins/PCG/Source/PCGEditor/Private/EditorMode/Tools/PCGInteractiveToolSettings.cpp b/Engine/Plugins/PCG/Source/PCGEditor/Private/EditorMode/Tools/PCGInteractiveToolSettings.cpp
--- a/Engine/Plugins/PCG/Source/PCGEditor/Private/EditorMode/Tools/PCGInteractiveToolSettings.cpp
+++ b/Engine/Plugins/PCG/Source/PCGEditor/Private/EditorMode/Tools/PCGInteractiveToolSettings.cpp
@@ -75,6 +75,8 @@ bool UPCGInteractiveToolSettings::Initialize(UInteractiveTool* InTool)
 		return false;
 	}
 
+	RegisterPropertyWatchers();
+	
 	// If we have selected one actor, that's the actor we want to edit
 	// This actor might:
 	// - Not have a PCG Component
@@ -88,6 +90,7 @@ bool UPCGInteractiveToolSettings::Initialize(UInteractiveTool* InTool)
 	else
 	{
 		AActor* SpawnedActor = Actor::SpawnWorking(World, FTransform::Identity, GetWorkingActorClass());
+		SpawnedActor->SetActorLabel(NewActorName.ToString(), false);
 		// Our spawned actor is also the actor we are editing
 		EditActor = SpawnedActor;
 
@@ -215,6 +218,42 @@ bool UPCGInteractiveToolSettings::IsSelectionAllowed(AActor* InActor, bool bInSe
 	return InActor == EditActor;
 }
 
+bool UPCGInteractiveToolSettings::CanResetToolData(FName DataInstanceIdentifier)
+{
+	return true;
+}
+
+void UPCGInteractiveToolSettings::RegisterPropertyWatchers()
+{
+	WatchProperty(NewActorName, [this](FName InName)
+	{
+		if(InName.IsNone() == false && GeneratedResources.SpawnedActor.IsValid() && GeneratedResources.SpawnedActor->GetActorLabelView().Equals(InName.ToString()) == false)
+		{
+			FActorLabelUtilities::RenameExistingActor(GeneratedResources.SpawnedActor.Get(), InName.ToString());
+		}
+	});
+
+	WatchProperty(NewPCGComponentName, [this](FName InName)
+	{
+		if(InName.IsNone() == false && GeneratedResources.GeneratedPCGComponent.IsValid())
+		{
+			// Only rename if the desired name is unique
+			UActorComponent* Component = GeneratedResources.GeneratedPCGComponent.Get();
+			UObject* Outer = Component->GetOuter();
+			UObject* ExistingObjectWithName = StaticFindObject(UObject::StaticClass(), Outer, *InName.ToString());
+			if (ExistingObjectWithName == nullptr)
+			{
+				const ERenameFlags Flags = REN_DontCreateRedirectors;
+				Component->Rename(*InName.ToString(), nullptr, Flags);
+			}
+			else if (ExistingObjectWithName != GeneratedResources.GeneratedPCGComponent.Get())
+			{
+				UE_LOGFMT(LogPCGEditor, Warning, "Could not rename PCG Component to name {InName}. Subobject with that name already exists.", InName.ToString());
+			}
+		}
+	});
+}
+
 void UPCGInteractiveToolSettings::InitializeWorkingData(FName DataInstanceIdentifier)
 {
 	UPCGComponent* PCGComponent = GetWorkingPCGComponent();
@@ -319,7 +358,7 @@ UPCGComponent* UPCGInteractiveToolSettings::FindOrGeneratePCGComponent(AActor& O
 
 	if(PCGComponent == nullptr)
 	{
-		PCGComponent = PCGComponent::Create(&OwningActor, "ToolPCGComponent");
+		PCGComponent = PCGComponent::Create(&OwningActor, NewPCGComponentName);
 		PCGComponent->SetFlags(RF_Transient);
 		GeneratedResources.GeneratedPCGComponent = PCGComponent;
 	}
@@ -419,19 +458,14 @@ void UPCGInteractiveToolSettings::OnRefresh()
 	}
 }
 
-bool UPCGInteractiveToolSettings::CanEditChange(const FProperty* InProperty) const
+bool UPCGInteractiveToolSettings::HasSpawnedActor() const
 {
-	bool bSuperResult =  Super::CanEditChange(InProperty);
-
-	if(bSuperResult == true)
-	{
-		if(InProperty->GetFName() == GET_MEMBER_NAME_CHECKED(UPCGInteractiveToolSettings, ActorClassToSpawn))
-		{
-			return EditActor.IsValid();
-		}
-	}
+	return GeneratedResources.SpawnedActor.IsValid();
+}
 
-	return bSuperResult;
+bool UPCGInteractiveToolSettings::HasGeneratedPCGComponent() const
+{
+	return GeneratedResources.GeneratedPCGComponent.IsValid();
 }
 
 void UPCGInteractiveToolSettings::CommitPropertyChangedEvents()
diff --git a/Engine/Plugins/PCG/Source/PCGEditor/Public/Details/PCGInteractiveToolSettingsDetails.h b/Engine/Plugins/PCG/Source/PCGEditor/Public/Details/PCGInteractiveToolSettingsDetails.h
--- a/Engine/Plugins/PCG/Source/PCGEditor/Public/Details/PCGInteractiveToolSettingsDetails.h
+++ b/Engine/Plugins/PCG/Source/PCGEditor/Public/Details/PCGInteractiveToolSettingsDetails.h
@@ -5,6 +5,7 @@
 #include "IDetailCustomization.h"
 #include "Input/Reply.h" // IWYU pragma: keep
 #include "UObject/WeakObjectPtrTemplates.h"
+#include "Layout/Visibility.h"
 
 class UPCGInteractiveToolSettings;
 class UPCGGraphInstance;
@@ -19,11 +20,15 @@ public:
 
 	void RegisterDelegates();
 	void UnregisterDelegates();
-	
+
 	virtual void CustomizeDetails(IDetailLayoutBuilder& DetailBuilder) override;
 	virtual void CustomizeDetails(const TSharedPtr<IDetailLayoutBuilder>& DetailBuilder) override;
 private:
 	void Rebuild(UObject* Object, const FPropertyChangedChainEvent& PropertyChangedChainEvent);
+	
+	FText GetResetToolDataButtonText() const;
+	EVisibility IsResetToolDataButtonVisible() const;
+	FReply ResetToolData() const;
 
 private:
 	TWeakObjectPtr<UPCGInteractiveToolSettings> ToolSettings;
diff --git a/Engine/Plugins/PCG/Source/PCGEditor/Public/EditorMode/Tools/Line/PCGDrawSplineTool.h b/Engine/Plugins/PCG/Source/PCGEditor/Public/EditorMode/Tools/Line/PCGDrawSplineTool.h
--- a/Engine/Plugins/PCG/Source/PCGEditor/Public/EditorMode/Tools/Line/PCGDrawSplineTool.h
+++ b/Engine/Plugins/PCG/Source/PCGEditor/Public/EditorMode/Tools/Line/PCGDrawSplineTool.h
@@ -45,15 +45,17 @@ public:
 			EditConditionHides))
 	double MinPointSpacing = 200.0;
 
-	/** The abstract target upon which to draw the spline. By default, it will be a raycast to the world. */
-	UPROPERTY(EditAnywhere, Category = "Target", meta = (ShowOnlyInnerProperties))
-	FPCGToolDrawTarget Target{};
-
-	UPROPERTY(EditAnywhere, Category = "PCG", meta = (TransientToolProperty))
+	UPROPERTY(EditAnywhere, Category = "Spline", meta = (TransientToolProperty))
 	FName SplineName;
 
+	/** The raycast rules that determine whether a raycast is valid or not.
+	 *  Used to determine different raycast targets such as landscapes, meshes, selected actor etc.	 */
 	UPROPERTY(EditAnywhere, Category= "Raycast")
 	FPCGRaycastFilterRuleCollection RaycastRuleCollection;
+
+	/** The settings to manipulate raycast results, such as normals or offsets of a spline point. */
+	UPROPERTY(EditAnywhere, Category = "Raycast", meta = (ShowOnlyInnerProperties))
+	FPCGToolRaycastSettings Settings{};
 };
 
 UCLASS(BlueprintType)
diff --git a/Engine/Plugins/PCG/Source/PCGEditor/Public/EditorMode/Tools/PCGInteractiveToolCommon.h b/Engine/Plugins/PCG/Source/PCGEditor/Public/EditorMode/Tools/PCGInteractiveToolCommon.h
--- a/Engine/Plugins/PCG/Source/PCGEditor/Public/EditorMode/Tools/PCGInteractiveToolCommon.h
+++ b/Engine/Plugins/PCG/Source/PCGEditor/Public/EditorMode/Tools/PCGInteractiveToolCommon.h
@@ -31,7 +31,7 @@ enum class EPCGToolDrawTargetNormal : uint8
 };
 
 USTRUCT()
-struct PCGEDITOR_API FPCGToolDrawTarget
+struct PCGEDITOR_API FPCGToolRaycastSettings
 {
 	GENERATED_BODY()
 
@@ -54,8 +54,4 @@ struct PCGEDITOR_API FPCGToolDrawTarget
 	/** Manual vector to offset drawn points from the clicked surface. */
 	UPROPERTY(EditAnywhere, Category = "Draw Target", meta = (UIMin = 0, UIMax = 100, EditCondition = "NormalMode == EPCGToolDrawTargetNormal::Explicit", EditConditionHides))
 	FVector Normal = FVector::ZAxisVector;
-
-	// @todo_pcg: Target object (such as a mesh)
-
-	// @todo_pcg: Target plane
 };
diff --git a/Engine/Plugins/PCG/Source/PCGEditor/Public/EditorMode/Tools/PCGInteractiveToolSettings.h b/Engine/Plugins/PCG/Source/PCGEditor/Public/EditorMode/Tools/PCGInteractiveToolSettings.h
--- a/Engine/Plugins/PCG/Source/PCGEditor/Public/EditorMode/Tools/PCGInteractiveToolSettings.h
+++ b/Engine/Plugins/PCG/Source/PCGEditor/Public/EditorMode/Tools/PCGInteractiveToolSettings.h
@@ -36,9 +36,15 @@ public:
 	virtual void Cancel(UInteractiveTool* OwningTool);
 	virtual void Shutdown();
 
-	// By default, when a tool is active, we can only select the Edit Actor in the Outliner.
+	/** By default, when a tool is active, we can only select the Edit Actor in the Outliner. */
 	virtual bool IsSelectionAllowed(AActor* InActor, bool bInSelection) const;
 
+	/** Whether tool data can be reset. Controls visibility of the Reset button in the UI. */
+	virtual bool CanResetToolData(FName DataInstanceIdentifier);
+	
+	/** Register property watchers and callbacks in here. */
+	virtual void RegisterPropertyWatchers();
+	
 	/** Initializes or finds the working data associated with the data instance identifier.
 	 * If NAME_None, the ToolTag is used as identifier. */
 	void InitializeWorkingData(FName DataInstanceIdentifier = NAME_None);
@@ -181,14 +187,14 @@ protected:
 	virtual void OnWorkingDataInitialized(FPCGInteractiveToolWorkingData* WorkingData) const;
 	virtual void OnPropertyModified(UObject* Object, FProperty* Property);
 	virtual void OnRefreshInternal();
+	
 private:
 	UPCGComponent* FindOrGeneratePCGComponent(AActor& OwningActor) const;
 	
 	void DeleteGeneratedResources() const;
 	
 	void OnRefresh();
-
-	virtual bool CanEditChange(const FProperty* InProperty) const override;
+	
 public:
 	UPROPERTY(EditAnywhere, BlueprintReadOnly, Category = "PCG", meta = (NoResetToDefault, GetAssetFilter = "GraphAssetFilter"))
 	TObjectPtr<UPCGGraphInterface> ToolGraph = nullptr;
@@ -197,16 +203,16 @@ public:
 	FName DataInstance = NAME_None;
 	
 	// @todo_pcg: Convert to a target object, rather than actor, to ease transition into SceneGraph
-	/** The selected actor to modify with the tool. */
-	UPROPERTY(EditAnywhere, Category = "PCG", meta = (TransientToolProperty))
+	/** The actor to modify with the tool. Either a selected actor, or a newly spawned one. */
+	UPROPERTY(meta = (TransientToolProperty))
 	TWeakObjectPtr<AActor> EditActor = nullptr;
 
 	/** The name to be issued to a newly created actor when the tool is applied. */
-	UPROPERTY(EditAnywhere, DisplayName="Actor Label", Category = "PCG", meta = (TransientToolProperty, EditConditionHides))
+	UPROPERTY(EditAnywhere, DisplayName="Actor Label", Category = "PCG", meta = (TransientToolProperty, EditCondition="HasSpawnedActor()", EditConditionHides))
 	FName NewActorName;
 
 	/** The name to be issued to a newly created PCG component when the tool is applied. */
-	UPROPERTY(EditAnywhere, DisplayName="Component Name", Category = "PCG", meta = (TransientToolProperty))
+	UPROPERTY(EditAnywhere, DisplayName="Component Name", Category = "PCG", meta = (TransientToolProperty, EditCondition="HasGeneratedPCGComponent()", EditConditionHides))
 	FName NewPCGComponentName;
 
 	UPROPERTY(EditAnywhere, Category = "PCG", meta=(TransientToolProperty, EditConditionHides))
@@ -219,9 +225,6 @@ public:
 	UPROPERTY(meta = (TransientToolProperty))
 	mutable TWeakObjectPtr<UPCGComponent> WorkingComponentCache;
 	
-	UPROPERTY(meta = (TransientToolProperty))
-	mutable bool bLockActorToSpawn = false;
-	
 	TMap<TWeakObjectPtr<UObject>, FPropertyChangedEvent> PropertyChangedEventQueue;
 
 	struct FGeneratedResources
@@ -235,4 +238,11 @@ public:
 	mutable FGeneratedResources GeneratedResources;
 
 	bool bTryToGenerate = false;
+
+private:
+	UFUNCTION()
+	bool HasSpawnedActor() const;
+
+	UFUNCTION()
+	bool HasGeneratedPCGComponent() const;
 };
diff --git a/Engine/Plugins/PCG/Source/PCGEditor/Public/EditorMode/Tools/Volume/PCGVolumeTool.h b/Engine/Plugins/PCG/Source/PCGEditor/Public/EditorMode/Tools/Volume/PCGVolumeTool.h
--- a/Engine/Plugins/PCG/Source/PCGEditor/Public/EditorMode/Tools/Volume/PCGVolumeTool.h
+++ b/Engine/Plugins/PCG/Source/PCGEditor/Public/EditorMode/Tools/Volume/PCGVolumeTool.h
@@ -25,6 +25,7 @@ public:
 	// ~Begin UPCGInteractiveToolSettings interface
 	virtual FName GetToolTag() const override;
 	virtual const UScriptStruct* GetWorkingDataType() const override;
+	virtual bool CanResetToolData(FName DataInstanceIdentifier) override { return false; }
 	// ~End UPCGInteractiveToolSettings interface
 
 	AVolume* GetWorkingVolume(FName DataInstanceIdentifier) const;
