[PCG] A few editor mode behavior improvements/fixes - Editor mode was not opening if using a non-default data instance. Volume mode not behaving well on non-AVolume actors, also was not behaving well under snapping.  PCG component selection not behaving properly on APCGVolumes.

#rnx

diff --git a/Engine/Plugins/PCG/Source/PCGEditor/Private/EditorMode/Tools/Helpers/PCGEdModeActorHelpers.cpp b/Engine/Plugins/PCG/Source/PCGEditor/Private/EditorMode/Tools/Helpers/PCGEdModeActorHelpers.cpp
--- a/Engine/Plugins/PCG/Source/PCGEditor/Private/EditorMode/Tools/Helpers/PCGEdModeActorHelpers.cpp
+++ b/Engine/Plugins/PCG/Source/PCGEditor/Private/EditorMode/Tools/Helpers/PCGEdModeActorHelpers.cpp
@@ -60,8 +60,10 @@ namespace UE::PCG::EditorMode
 
 			// Priority order:
 			// 1. A tool-marked pcg component with a graph that has the same tool tags
-			// 2. A pcg component with a graph that has the same tool tags
+			// 2. A pcg component with a graph that has the same tool tags ("tentative").
+			// 3. Any pcg component with no graph - especially useful to pickup PCG components on PCG Volumes ("fallback").
 			UPCGComponent* TentativeComponent = nullptr;
+			UPCGComponent* FallbackComponent = nullptr;
 
 			for (UPCGComponent* Component : Components)
 			{
@@ -69,6 +71,7 @@ namespace UE::PCG::EditorMode
 				{
 					const bool bComponentHasToolTag = Component->ComponentHasTag(Tags::ToolGeneratedTag);
 					const bool bComponentHasCompatibleGraph = Component->GetGraph() && Component->GetGraph()->ToolData.CompatibleToolTags.Contains(ToolTag);
+					const bool bComponentHasNoGraph = !Component->GetGraph();
 
 					if (bComponentHasToolTag && bComponentHasCompatibleGraph)
 					{
@@ -78,10 +81,14 @@ namespace UE::PCG::EditorMode
 					{
 						TentativeComponent = Component;
 					}
+					else if (bComponentHasNoGraph && !FallbackComponent)
+					{
+						FallbackComponent = Component;
+					}
 				}
 			}
 
-			return TentativeComponent;
+			return TentativeComponent ? TentativeComponent : FallbackComponent;
 		}
 
 		UPCGComponent* Create(AActor* Owner, const FName ComponentName)
diff --git a/Engine/Plugins/PCG/Source/PCGEditor/Private/EditorMode/Tools/PCGInteractiveToolSettings.cpp b/Engine/Plugins/PCG/Source/PCGEditor/Private/EditorMode/Tools/PCGInteractiveToolSettings.cpp
--- a/Engine/Plugins/PCG/Source/PCGEditor/Private/EditorMode/Tools/PCGInteractiveToolSettings.cpp
+++ b/Engine/Plugins/PCG/Source/PCGEditor/Private/EditorMode/Tools/PCGInteractiveToolSettings.cpp
@@ -213,6 +213,18 @@ void UPCGInteractiveToolSettings::Shutdown()
 	}
 }
 
+void UPCGInteractiveToolSettings::RestoreProperties(UInteractiveTool* RestoreToTool, const FString& CacheIdentifier)
+{
+	Super::RestoreProperties(RestoreToTool, CacheIdentifier);
+
+	// Reset the data instance if it is not valid anymore.
+	TArray<FName> DataInstanceNames = GetDataInstanceNamesForGraph();
+	if (!DataInstanceNames.Contains(DataInstance) && !DataInstanceNames.IsEmpty())
+	{
+		SetDataInstance(GetDefaultDataInstanceName());
+	}
+}
+
 bool UPCGInteractiveToolSettings::IsSelectionAllowed(AActor* InActor, bool bInSelection) const
 {
 	return InActor == EditActor;
@@ -380,7 +392,11 @@ void UPCGInteractiveToolSettings::DeleteGeneratedResources() const
 		FPCGTaskId CleanupTaskId = GeneratedResources.GeneratedPCGComponent->GetSubsystem()->ScheduleCleanup(GeneratedResources.GeneratedPCGComponent.Get(), true, {});
 		GeneratedResources.GeneratedPCGComponent->GetSubsystem()->ScheduleGeneric([PCGComponent = GeneratedResources.GeneratedPCGComponent]() -> bool
 		{
-			PCGComponent->DestroyComponent();
+			if (PCGComponent.IsValid())
+			{
+				PCGComponent->DestroyComponent();
+			}
+			
 			return true;
 		},
 		GeneratedResources.GeneratedPCGComponent.Get(),
@@ -405,8 +421,10 @@ bool UPCGInteractiveToolSettings::TryLoadDefaults()
 	check(!HasAnyFlags(RF_ClassDefaultObject));
 	if (UPCGGraphInterface* DefaultGraphInterface = CastChecked<UPCGGraphInterface>(GetDefaultGraph().TryLoad(), ECastCheckedType::NullAllowed))
 	{
+		// At this point, the delegates (OnPropertyModified) aren't setup, so we need to set things up manually.
 		SetToolGraph(DefaultGraphInterface);
 		SetActorClassToSpawn(DefaultGraphInterface->ToolData.InitialActorClassToSpawn);
+		SetDataInstance(GetDefaultDataInstanceName());
 
 		return true;
 	}
diff --git a/Engine/Plugins/PCG/Source/PCGEditor/Private/EditorMode/Tools/Volume/PCGVolumeTool.cpp b/Engine/Plugins/PCG/Source/PCGEditor/Private/EditorMode/Tools/Volume/PCGVolumeTool.cpp
--- a/Engine/Plugins/PCG/Source/PCGEditor/Private/EditorMode/Tools/Volume/PCGVolumeTool.cpp
+++ b/Engine/Plugins/PCG/Source/PCGEditor/Private/EditorMode/Tools/Volume/PCGVolumeTool.cpp
@@ -2,14 +2,19 @@
 
 #include "EditorMode/Tools/Volume/PCGVolumeTool.h"
 
+#include "PCGEditorModule.h"
+#include "PCGVolume.h"
 #include "EditorMode/PCGEdModeStyle.h"
 
+#include "InteractiveToolManager.h"
+#include "SnappingUtils.h"
+#include "ToolTargetManager.h"
 #include "ActorFactories/ActorFactory.h"
 #include "BaseBehaviors/SingleClickOrDragBehavior.h"
 #include "BaseBehaviors/MouseHoverBehavior.h"
 #include "Builders/CubeBuilder.h"
 #include "GameFramework/Volume.h"
-#include "InteractiveToolManager.h"
+#include "ToolTargets/ToolTarget.h"
 
 #include UE_INLINE_GENERATED_CPP_BY_NAME(PCGVolumeTool)
 
@@ -56,11 +61,24 @@ AVolume* UPCGInteractiveToolSettings_Volume::GetWorkingVolume(FName DataInstance
 	return nullptr;
 }
 
+void UPCGInteractiveToolSettings_Volume::SetActorClassToSpawn(TSubclassOf<AActor> InClass)
+{
+	TSubclassOf<AActor> Class = InClass;
+
+	if (!Class || !Class->IsChildOf(AVolume::StaticClass()))
+	{
+		UE_LOG(LogPCGEditor, Error, TEXT("Invalid class used in volume tool. Will use PCGVolume instead."));
+		Class = APCGVolume::StaticClass();
+	}
+
+	Super::SetActorClassToSpawn(Class);
+}
+
 void UPCGVolumeTool::Setup()
 {
 	Super::Setup();
 
-	DefaultHeight = 100; // @todo_pcg: get from existing volume if any
+	DefaultHalfHeight = 100; // @todo_pcg: get from existing volume if any
 
 	ClickOrDragBehavior = NewObject<USingleClickOrDragInputBehavior>();
 	ClickOrDragBehavior->Initialize(this, this);
@@ -115,11 +133,23 @@ void UPCGVolumeTool::UpdateVolumeFromBox()
 	if (AVolume* WorkingVolume = ToolSettings ? ToolSettings->GetWorkingVolume(ToolSettings->DataInstance) : nullptr)
 	{
 		const FVector BoxCenter = LastBox.GetCenter();
-		const FVector Extent = 2.0 * LastBox.GetExtent();
+		FVector Extent = 2.0 * LastBox.GetExtent();
 
 		// Move actor to bounds center
 		WorkingVolume->SetActorTransform(FTransform(BoxCenter));
 
+		// if the user is using snapping, we still need to make sure that the final location would cover the full extent.
+		// This is because the brush builder will enforce snapping.
+		if (FSnappingUtils::IsSnapToGridEnabled())
+		{
+			FVector SnappedCenter = BoxCenter;
+			FRotator DummyRotator = FRotator::ZeroRotator;
+			FSnappingUtils::SnapToBSPVertex(SnappedCenter, FVector::ZeroVector, DummyRotator);
+
+			// We will not change the X/Y coordinates, but will adjust the Z coordinate so that we're always covering the bottom of the original box
+			Extent.Z += 2.0 * FMath::Abs(SnappedCenter.Z - BoxCenter.Z);
+		}
+
 		// setup cube builder
 		UCubeBuilder* Builder = NewObject<UCubeBuilder>();
 		Builder->X = Extent.X;
@@ -178,7 +208,7 @@ FBox UPCGVolumeTool::GetCurrentBox() const
 	else
 	{
 		BoxCenter.Z = FirstPoint.Z - Bias;
-		HalfExtents.Z = DefaultHeight;
+		HalfExtents.Z = DefaultHalfHeight;
 	}
 
 	return FBox(BoxCenter - HalfExtents, BoxCenter + HalfExtents);
@@ -201,8 +231,6 @@ FInputRayHit UPCGVolumeTool::IsHitByClick(const FInputDeviceRay& ClickPosition)
 
 void UPCGVolumeTool::OnClicked(const FInputDeviceRay& ClickPosition)
 {
-	UPCGInteractiveToolSettings_Volume* Settings = CastChecked<UPCGInteractiveToolSettings_Volume>(ToolSettings);
-
 	bIsDraggingBase = false;
 
 	if (bBaseSelectionDone)
@@ -292,6 +320,28 @@ bool UPCGVolumeTool::OnUpdateHover(const FInputDeviceRay& DevicePos)
 
 bool UPCGVolumeToolBuilder::CanBuildTool(const FToolBuilderState& SceneState) const
 {
+	FToolTargetTypeRequirements AcceptAllReqs{};
+
+	// Accepts either no selection, or selection from an AVolume actor.
+	const int32 NumTargets = SceneState.TargetManager->CountSelectedAndTargetable(SceneState, AcceptAllReqs);
+	if (NumTargets > 0)
+	{
+		bool bHasValidTargets = false;
+		SceneState.TargetManager->EnumerateSelectedAndTargetableComponents(SceneState, AcceptAllReqs,
+			[&bHasValidTargets](UActorComponent* Component)
+			{
+				if (Component && Component->GetOwner() && Component->GetOwner()->IsA<AVolume>())
+				{
+					bHasValidTargets = true;
+				}
+			});
+
+		if (!bHasValidTargets)
+		{
+			return false;
+		}
+	}
+
 	return true;
 }
 
diff --git a/Engine/Plugins/PCG/Source/PCGEditor/Public/EditorMode/Tools/PCGInteractiveToolSettings.h b/Engine/Plugins/PCG/Source/PCGEditor/Public/EditorMode/Tools/PCGInteractiveToolSettings.h
--- a/Engine/Plugins/PCG/Source/PCGEditor/Public/EditorMode/Tools/PCGInteractiveToolSettings.h
+++ b/Engine/Plugins/PCG/Source/PCGEditor/Public/EditorMode/Tools/PCGInteractiveToolSettings.h
@@ -36,6 +36,10 @@ public:
 	virtual void Cancel(UInteractiveTool* OwningTool);
 	virtual void Shutdown();
 
+	//~Begin interface UInteractiveToolPropertySet
+	virtual void RestoreProperties(UInteractiveTool* RestoreToTool, const FString& CacheIdentifier = TEXT("")) override;
+	//~End interface UInteractiveToolPropertySet
+
 	/** By default, when a tool is active, we can only select the Edit Actor in the Outliner. */
 	virtual bool IsSelectionAllowed(AActor* InActor, bool bInSelection) const;
 
@@ -54,7 +58,7 @@ public:
 
 	void SetDataInstance(FName DataInstanceIdentifier);
 
-	void SetActorClassToSpawn(TSubclassOf<AActor> Class);
+	virtual void SetActorClassToSpawn(TSubclassOf<AActor> Class);
 	
 	/** A soft object path to the graph asset to load by default for this tool. */
 	virtual FSoftObjectPath GetDefaultGraph() const;
diff --git a/Engine/Plugins/PCG/Source/PCGEditor/Public/EditorMode/Tools/Volume/PCGVolumeTool.h b/Engine/Plugins/PCG/Source/PCGEditor/Public/EditorMode/Tools/Volume/PCGVolumeTool.h
--- a/Engine/Plugins/PCG/Source/PCGEditor/Public/EditorMode/Tools/Volume/PCGVolumeTool.h
+++ b/Engine/Plugins/PCG/Source/PCGEditor/Public/EditorMode/Tools/Volume/PCGVolumeTool.h
@@ -25,6 +25,7 @@ public:
 	// ~Begin UPCGInteractiveToolSettings interface
 	virtual FName GetToolTag() const override;
 	virtual const UScriptStruct* GetWorkingDataType() const override;
+	virtual void SetActorClassToSpawn(TSubclassOf<AActor> Class) override;
 	virtual bool CanResetToolData(FName DataInstanceIdentifier) override { return false; }
 	// ~End UPCGInteractiveToolSettings interface
 
@@ -99,14 +100,14 @@ private:
 	FVector FirstPoint = FVector::ZeroVector;
 	FVector SecondPoint = FVector::ZeroVector;
 	FVector InteractiveSecondPoint = FVector::ZeroVector;
-	double DefaultHeight = 0;
+	double DefaultHalfHeight = 0;
 
 	FRay LastRay;
 	FBox LastBox = FBox(EForceInit::ForceInit);
 };
 
 UCLASS(Transient, MinimalAPI)
-class UPCGVolumeToolBuilder : public UPCGInteractiveToolBuilder
+class UPCGVolumeToolBuilder : public UPCGInteractiveToolWithToolTargetsBuilder
 {
 	GENERATED_BODY()
 
