Aborting tool activation when the default graph doesn't actually have a compatible tool tag. This can otherwise lead to a crash

diff --git a/Engine/Plugins/PCG/Source/PCGEditor/Private/EditorMode/Tools/PCGInteractiveToolSettings.cpp b/Engine/Plugins/PCG/Source/PCGEditor/Private/EditorMode/Tools/PCGInteractiveToolSettings.cpp
--- a/Engine/Plugins/PCG/Source/PCGEditor/Private/EditorMode/Tools/PCGInteractiveToolSettings.cpp
+++ b/Engine/Plugins/PCG/Source/PCGEditor/Private/EditorMode/Tools/PCGInteractiveToolSettings.cpp
@@ -421,6 +421,12 @@ bool UPCGInteractiveToolSettings::TryLoadDefaults()
 	check(!HasAnyFlags(RF_ClassDefaultObject));
 	if (UPCGGraphInterface* DefaultGraphInterface = CastChecked<UPCGGraphInterface>(GetDefaultGraph().TryLoad(), ECastCheckedType::NullAllowed))
 	{
+		if (DefaultGraphInterface->ToolData.CompatibleToolTags.Contains(GetToolTag()) == false)
+		{
+			UE_LOGFMT(LogPCGEditor, Error, "Graph {Graph} does not have compatible tool tag {ToolTag}", GetPathNameSafe(DefaultGraphInterface), GetToolTag());
+			return false;
+		}
+		
 		// At this point, the delegates (OnPropertyModified) aren't setup, so we need to set things up manually.
 		SetToolGraph(DefaultGraphInterface);
 		SetActorClassToSpawn(DefaultGraphInterface->ToolData.InitialActorClassToSpawn);
