[PCG Editor mode] Make graph ToolData visibility not depending on 'Expose in library'. Also Bughawk fix.

#rnx

diff --git a/Engine/Plugins/PCG/Source/PCG/Public/PCGGraph.h b/Engine/Plugins/PCG/Source/PCG/Public/PCGGraph.h
--- a/Engine/Plugins/PCG/Source/PCG/Public/PCGGraph.h
+++ b/Engine/Plugins/PCG/Source/PCG/Public/PCGGraph.h
@@ -180,7 +180,7 @@ public:
 	bool bExposeToLibrary = false;
 
 	/** Contains the data relevant for PCG Editor Mode usage. */
-	UPROPERTY(BlueprintReadWrite, EditAnywhere, Category = AssetInfo, AssetRegistrySearchable, meta=(EditCondition = "bExposeToLibrary", EditConditionHides, PCGNoHash))
+	UPROPERTY(BlueprintReadWrite, EditAnywhere, Category = AssetInfo, AssetRegistrySearchable, meta=(EditConditionHides, PCGNoHash))
 	FPCGGraphToolData ToolData;
 
 	/** When enabled, this graph can be executed outside of the world using an editor execution source. */
diff --git a/Engine/Plugins/PCG/Source/PCGEditor/Private/EditorMode/PCGEdModeCommands.cpp b/Engine/Plugins/PCG/Source/PCGEditor/Private/EditorMode/PCGEdModeCommands.cpp
--- a/Engine/Plugins/PCG/Source/PCGEditor/Private/EditorMode/PCGEdModeCommands.cpp
+++ b/Engine/Plugins/PCG/Source/PCGEditor/Private/EditorMode/PCGEdModeCommands.cpp
@@ -66,7 +66,8 @@ FPCGEditorModeToolCommands::FPCGEditorModeToolCommands()
   * @param CommandInfoSP The shared pointer for the FUICommandInfo
  * @param Name The name of the tool
  * @param Label The user-forward label for the tool
- * @param ToolTip The tooltip for the tool
+ * @param ToolTip The tooltip for the tool: Format can include {ToolTag}.
+ * @param ToolTag The tag for the tool: can be used in the tooltip.
  * @param IconStyleName The style name for the icon to be used for the button
  */
 #define PCG_TOOL_COMMAND_INFO(CommandInfoSP, Name, Label, ToolTip, ToolTag, IconStyleName)\
