[PCG Ed Mode] Tool Setting Initialization can now return an error message to display

diff --git a/Engine/Plugins/PCG/Source/PCGEditor/Private/EditorMode/Tools/PCGInteractiveToolBuilder.cpp b/Engine/Plugins/PCG/Source/PCGEditor/Private/EditorMode/Tools/PCGInteractiveToolBuilder.cpp
--- a/Engine/Plugins/PCG/Source/PCGEditor/Private/EditorMode/Tools/PCGInteractiveToolBuilder.cpp
+++ b/Engine/Plugins/PCG/Source/PCGEditor/Private/EditorMode/Tools/PCGInteractiveToolBuilder.cpp
@@ -43,13 +43,23 @@ bool UE::PCG::EditorMode::Tool::BuildTool(UInteractiveTool* InTool)
 	
 	UPCGInteractiveToolSettings* Settings = NewObject<UPCGInteractiveToolSettings>(InTool, SettingsClass);
 	SettingsProperty->SetObjectPropertyValue(const_cast<void*>(PropertyAddress), Settings);
-	if(Settings->Initialize(InTool) == false)
+	if(TValueOrError<bool, FText> Result = Settings->Initialize(InTool); Result.HasError() || (Result.HasValue() && Result.GetValue() == false))
 	{
 		Settings->Cancel(InTool);
 		Settings->Shutdown();
+
+		FText Message = FText::FormatOrdered(LOCTEXT("PCGSettingsInitializationFailedMessage", "Tool {0} failed initialization."), InTool->GetClass()->GetDisplayNameText());
+		
+		if(Result.HasError())
+		{
+			Message = Result.GetError();
+		}
+
+		if(Message.IsEmpty() == false)
+		{
+			InTool->GetToolManager()->DisplayMessage(Message, EToolMessageLevel::UserError);
+		}
 		
-		FText Message = FText::FormatOrdered(LOCTEXT("PCGSettingsInitializationFailedMessage", "Tool {0} failed initialization on settings {1}"), InTool->GetClass()->GetDisplayNameText(), Settings->GetClass()->GetDisplayNameText());
-		InTool->GetToolManager()->DisplayMessage(Message, EToolMessageLevel::UserError);
 		return false;
 	}
 	
diff --git a/Engine/Plugins/PCG/Source/PCGEditor/Private/EditorMode/Tools/PCGInteractiveToolSettings.cpp b/Engine/Plugins/PCG/Source/PCGEditor/Private/EditorMode/Tools/PCGInteractiveToolSettings.cpp
--- a/Engine/Plugins/PCG/Source/PCGEditor/Private/EditorMode/Tools/PCGInteractiveToolSettings.cpp
+++ b/Engine/Plugins/PCG/Source/PCGEditor/Private/EditorMode/Tools/PCGInteractiveToolSettings.cpp
@@ -59,7 +59,7 @@ UPCGInteractiveToolSettings::UPCGInteractiveToolSettings()
 	
 }
 
-bool UPCGInteractiveToolSettings::Initialize(UInteractiveTool* InTool)
+TValueOrError<bool, FText> UPCGInteractiveToolSettings::Initialize(UInteractiveTool* InTool)
 {
 	using namespace UE::PCG::EditorMode;
 
@@ -70,9 +70,9 @@ bool UPCGInteractiveToolSettings::Initialize(UInteractiveTool* InTool)
 	TArray<AActor*> SelectedActors = State.SelectedActors;
 
 	// In order to constrain component selection and actor creation, we need to make sure to setup the default settings first.
-	if (!TryLoadDefaults())
+	if (TValueOrError<bool, FText> Result = TryLoadDefaults(); Result.HasError() || (Result.HasValue() && Result.GetValue() == false))
 	{
-		return false;
+		return Result;
 	}
 
 	RegisterPropertyWatchers();
@@ -99,13 +99,13 @@ bool UPCGInteractiveToolSettings::Initialize(UInteractiveTool* InTool)
 
 	if(!ensureMsgf(EditActor.IsValid(), TEXT("The Edit Actor has to be valid at this point. Either as a previously existing actor, or a new spawned actor.")))
 	{
-		return false;
+		return MakeValue(false);
 	}
 	
 	TWeakObjectPtr<UPCGComponent> PCGComponent = FindOrGeneratePCGComponent(*EditActor.Get());
 	if(!ensureMsgf(PCGComponent.IsValid(), TEXT("The PCG Component has to be valid at this point. Either as a previously existing PCG Component, or a transient generated PCG Component.")))
 	{
-		return false;
+		return MakeValue(false);
 	}
 
 	OnModified.AddUObject(this, &UPCGInteractiveToolSettings::OnPropertyModified);
@@ -115,7 +115,7 @@ bool UPCGInteractiveToolSettings::Initialize(UInteractiveTool* InTool)
 	{
 		if (!ensureMsgf(ToolGraph, TEXT("The PCG component has no graph and no default graph for this tool was found in the editor preferences. Shutting down the tool.")))
 		{
-			return false;
+			return MakeValue(false);
 		}
 
 		PCGComponent->SetGraph(ToolGraph);
@@ -137,7 +137,7 @@ bool UPCGInteractiveToolSettings::Initialize(UInteractiveTool* InTool)
 	KeyInputBehavior->bRequireAllKeys = false;
 	InTool->AddInputBehavior(KeyInputBehavior);
 
-	return true;
+	return MakeValue(true);
 }
 
 void UPCGInteractiveToolSettings::Apply(UInteractiveTool* OwningTool)
@@ -416,15 +416,14 @@ void UPCGInteractiveToolSettings::DeleteGeneratedResources() const
 	}
 }
 
-bool UPCGInteractiveToolSettings::TryLoadDefaults()
+TValueOrError<bool, FText> UPCGInteractiveToolSettings::TryLoadDefaults()
 {
 	check(!HasAnyFlags(RF_ClassDefaultObject));
 	if (UPCGGraphInterface* DefaultGraphInterface = CastChecked<UPCGGraphInterface>(GetDefaultGraph().TryLoad(), ECastCheckedType::NullAllowed))
 	{
 		if (DefaultGraphInterface->ToolData.CompatibleToolTags.Contains(GetToolTag()) == false)
 		{
-			UE_LOGFMT(LogPCGEditor, Error, "Graph {Graph} does not have compatible tool tag {ToolTag}", GetPathNameSafe(DefaultGraphInterface), GetToolTag());
-			return false;
+			return MakeError(FText::FormatOrdered(LOCTEXT("DefaultGraph_IncompatibleToolTag", "Graph {Graph} does not have compatible tool tag {ToolTag}"), FText::FromString(GetPathNameSafe(DefaultGraphInterface)), FText::FromName(GetToolTag())));
 		}
 		
 		// At this point, the delegates (OnPropertyModified) aren't setup, so we need to set things up manually.
@@ -432,12 +431,10 @@ bool UPCGInteractiveToolSettings::TryLoadDefaults()
 		SetActorClassToSpawn(DefaultGraphInterface->ToolData.InitialActorClassToSpawn);
 		SetDataInstance(GetDefaultDataInstanceName());
 
-		return true;
-	}
-	else
-	{
-		return false;
+		return MakeValue(true);
 	}
+	
+	return MakeValue(false);
 }
 
 void UPCGInteractiveToolSettings::Tick(float DeltaTime)
diff --git a/Engine/Plugins/PCG/Source/PCGEditor/Public/EditorMode/Tools/PCGInteractiveToolSettings.h b/Engine/Plugins/PCG/Source/PCGEditor/Public/EditorMode/Tools/PCGInteractiveToolSettings.h
--- a/Engine/Plugins/PCG/Source/PCGEditor/Public/EditorMode/Tools/PCGInteractiveToolSettings.h
+++ b/Engine/Plugins/PCG/Source/PCGEditor/Public/EditorMode/Tools/PCGInteractiveToolSettings.h
@@ -31,7 +31,7 @@ public:
 	UPCGInteractiveToolSettings();
 	virtual ~UPCGInteractiveToolSettings() override {}
 
-	virtual bool Initialize(UInteractiveTool* InTool);
+	virtual TValueOrError<bool, FText> Initialize(UInteractiveTool* InTool);
 	virtual void Apply(UInteractiveTool* OwningTool);
 	virtual void Cancel(UInteractiveTool* OwningTool);
 	virtual void Shutdown();
@@ -176,7 +176,7 @@ public:
 	}
 protected:
 	/** Try to load the defaults (graph, actor class, ...) into the tool settings (ToolGraph, etc.). */
-	virtual bool TryLoadDefaults();
+	virtual TValueOrError<bool, FText> TryLoadDefaults();
 	
 	virtual void Tick(float DeltaTime) override;
 	virtual bool IsTickable() const override;
