Added error message to TryLoadDefaults if default graph wasn't found.

diff --git a/Engine/Plugins/PCG/Source/PCGEditor/Private/EditorMode/Tools/PCGInteractiveToolSettings.cpp b/Engine/Plugins/PCG/Source/PCGEditor/Private/EditorMode/Tools/PCGInteractiveToolSettings.cpp
--- a/Engine/Plugins/PCG/Source/PCGEditor/Private/EditorMode/Tools/PCGInteractiveToolSettings.cpp
+++ b/Engine/Plugins/PCG/Source/PCGEditor/Private/EditorMode/Tools/PCGInteractiveToolSettings.cpp
@@ -433,8 +433,9 @@ TValueOrError<bool, FText> UPCGInteractiveToolSettings::TryLoadDefaults()
 
 		return MakeValue(true);
 	}
-	
-	return MakeValue(false);
+
+	FText Message = LOCTEXT("DefaultGraph_NotFound", "No default graph found for tool settings {ToolSettings}. Please specify a graph with tool tag {ToolTag} under Editor Preferences -> PCG Editor Mode Settings");
+	return MakeError(FText::FormatOrdered(Message, GetClass()->GetDisplayNameText(), FText::FromName(GetToolTag())));
 }
 
 void UPCGInteractiveToolSettings::Tick(float DeltaTime)
