[PCG] Export KernelParams DI and Exportable DI

#rnx

diff --git a/Engine/Plugins/PCG/Source/PCG/Public/Compute/DataInterfaces/PCGComputeDataInterface.h b/Engine/Plugins/PCG/Source/PCG/Public/Compute/DataInterfaces/PCGComputeDataInterface.h
--- a/Engine/Plugins/PCG/Source/PCG/Public/Compute/DataInterfaces/PCGComputeDataInterface.h
+++ b/Engine/Plugins/PCG/Source/PCG/Public/Compute/DataInterfaces/PCGComputeDataInterface.h
@@ -154,12 +154,12 @@ class UPCGExportableDataProvider : public UPCGComputeDataProvider
 
 public:
 	//~ Begin UComputeDataInterface Interface
-	virtual void Initialize(const UComputeDataInterface* InDataInterface, UObject* InBinding, uint64 InInputMask, uint64 InOutputMask) override;
-	virtual void Reset() override;
+	PCG_API virtual void Initialize(const UComputeDataInterface* InDataInterface, UObject* InBinding, uint64 InInputMask, uint64 InOutputMask) override;
+	PCG_API virtual void Reset() override;
 	//~ End UComputeDataInterface Interface
 
 	//~ Begin UPCGComputeDataProvider Interface
-	virtual bool PrepareForExecute_GameThread(UPCGDataBinding* InBinding) override;
+	PCG_API virtual bool PrepareForExecute_GameThread(UPCGDataBinding* InBinding) override;
 	//~ End UPCGComputeDataProvider Interface
 
 	/** Returns true if the generated buffer should be exported to a persistent buffer to be passed to downstream nodes. */
@@ -188,8 +188,8 @@ class UPCGKernelParamsDataInterface : public UPCGComputeDataInterface
 
 public:
 	//~ Begin UComputeDataInterface Interface
-	void GetSupportedInputs(TArray<FShaderFunctionDefinition>& OutFunctions) const override;
-	void GetHLSL(FString& OutHLSL, FString const& InDataInterfaceName) const override;
+	PCG_API void GetSupportedInputs(TArray<FShaderFunctionDefinition>& OutFunctions) const override;
+	PCG_API void GetHLSL(FString& OutHLSL, FString const& InDataInterfaceName) const override;
 	//~ End UComputeDataInterface Interface
 };
 
