[PCG] Fix noisy attribute validation error when running HLSL nodes with no data

#rnx

diff --git a/Engine/Plugins/PCG/Source/PCG/Private/Compute/Elements/PCGCustomHLSLKernel.cpp b/Engine/Plugins/PCG/Source/PCG/Private/Compute/Elements/PCGCustomHLSLKernel.cpp
--- a/Engine/Plugins/PCG/Source/PCG/Private/Compute/Elements/PCGCustomHLSLKernel.cpp
+++ b/Engine/Plugins/PCG/Source/PCG/Private/Compute/Elements/PCGCustomHLSLKernel.cpp
@@ -1883,7 +1883,7 @@ bool UPCGCustomHLSLKernel::AreAttributesValid(FPCGContext* InContext, FText* Out
 					}
 				}
 
-				if (!AttrDesc)
+				if (!AttrDesc && !PinDesc->DataDescs.IsEmpty())
 				{
 #if PCG_KERNEL_LOGGING_ENABLED
 					if (OutErrorText)
