[PCG] Fix inaccurate comment - Bughawk

#rnx

diff --git a/Engine/Plugins/PCG/Source/PCG/Public/Compute/DataInterfaces/PCGComputeDataInterface.h b/Engine/Plugins/PCG/Source/PCG/Public/Compute/DataInterfaces/PCGComputeDataInterface.h
--- a/Engine/Plugins/PCG/Source/PCG/Public/Compute/DataInterfaces/PCGComputeDataInterface.h
+++ b/Engine/Plugins/PCG/Source/PCG/Public/Compute/DataInterfaces/PCGComputeDataInterface.h
@@ -76,8 +76,10 @@ class UPCGComputeDataProvider : public UComputeDataProvider
 	GENERATED_BODY()
 
 public:
+	//~ Begin UComputeDataProvider Interface
 	PCG_API virtual void Initialize(const UComputeDataInterface* InDataInterface, UObject* InBinding, uint64 InInputMask, uint64 InOutputMask) override;
 	PCG_API virtual void Reset() override;
+	//~ End UComputeDataProvider Interface
 
 	/** Called before PrepareForExecute to allow data providers to do any readbacks and finalize data descriptions. */
 	virtual bool PerformPreExecuteReadbacks_GameThread(UPCGDataBinding* InBinding) { return true; }
@@ -153,10 +155,10 @@ class UPCGExportableDataProvider : public UPCGComputeDataProvider
 	GENERATED_BODY()
 
 public:
-	//~ Begin UComputeDataInterface Interface
+	//~ Begin UComputeDataProvider Interface
 	PCG_API virtual void Initialize(const UComputeDataInterface* InDataInterface, UObject* InBinding, uint64 InInputMask, uint64 InOutputMask) override;
 	PCG_API virtual void Reset() override;
-	//~ End UComputeDataInterface Interface
+	//~ End UComputeDataProvider Interface
 
 	//~ Begin UPCGComputeDataProvider Interface
 	PCG_API virtual bool PrepareForExecute_GameThread(UPCGDataBinding* InBinding) override;
