Changed max dimensions of the PCG Paint Brush to be fixed

diff --git a/Engine/Plugins/PCG/Source/PCGEditor/Private/EditorMode/Tools/Paint/PCGPaintTool.cpp b/Engine/Plugins/PCG/Source/PCGEditor/Private/EditorMode/Tools/Paint/PCGPaintTool.cpp
--- a/Engine/Plugins/PCG/Source/PCGEditor/Private/EditorMode/Tools/Paint/PCGPaintTool.cpp
+++ b/Engine/Plugins/PCG/Source/PCGEditor/Private/EditorMode/Tools/Paint/PCGPaintTool.cpp
@@ -267,8 +267,8 @@ void UPCGPaintTool::RestartBrushStampIndicator()
 
 double UPCGPaintTool::EstimateMaximumTargetDimension()
 {
-	// We return the target actor's max dimension
-	return ToolSettings->EditActor->GetComponentsBoundingBox().GetExtent().GetAbsMax();
+	// 400 will mean brush radius at default Brush Size 0.25 will be 50, meaning 1m in diameter 
+	return 400.f;
 }
 
 void UPCGPaintTool::AddSinglePointFromBrush()
