[PCG] Cleanup of unused data and unspecific namespace.

#rnx

diff --git a/Engine/Plugins/PCG/Source/PCG/Private/Compute/DataInterfaces/Elements/PCGSkinnedMeshSpawnerDataInterface.cpp b/Engine/Plugins/PCG/Source/PCG/Private/Compute/DataInterfaces/Elements/PCGSkinnedMeshSpawnerDataInterface.cpp
--- a/Engine/Plugins/PCG/Source/PCG/Private/Compute/DataInterfaces/Elements/PCGSkinnedMeshSpawnerDataInterface.cpp
+++ b/Engine/Plugins/PCG/Source/PCG/Private/Compute/DataInterfaces/Elements/PCGSkinnedMeshSpawnerDataInterface.cpp
@@ -623,7 +623,6 @@ void UPCGSkinnedMeshSpawnerDataProvider::CreatePrimitiveDescriptors(FPCGContext*
 	Primitives.AttributeIdOffsetStrides = AttributeIdOffsetStrides;
 	Primitives.SelectorAttributeId = SelectorAttributeId;
 	Primitives.SelectionCDF = PrimitiveSelectionCDF;
-	Primitives.PrimitiveStringKeys = PrimitiveStringKeys;
 	Primitives.PrimitiveMeshBounds = PrimitiveMeshBounds;
 
 	// Useful for debugging instance counts.
diff --git a/Engine/Plugins/PCG/Source/PCG/Private/Compute/DataInterfaces/Elements/PCGStaticMeshSpawnerDataInterface.cpp b/Engine/Plugins/PCG/Source/PCG/Private/Compute/DataInterfaces/Elements/PCGStaticMeshSpawnerDataInterface.cpp
--- a/Engine/Plugins/PCG/Source/PCG/Private/Compute/DataInterfaces/Elements/PCGStaticMeshSpawnerDataInterface.cpp
+++ b/Engine/Plugins/PCG/Source/PCG/Private/Compute/DataInterfaces/Elements/PCGStaticMeshSpawnerDataInterface.cpp
@@ -232,7 +232,7 @@ bool UPCGStaticMeshSpawnerDataProvider::PerformPreExecuteReadbacks_GameThread(UP
 			return true;
 		}
 
-		AnalysisDataIndex = GetProducerKernel() ? InBinding->GetFirstInputDataIndex(GetProducerKernel(), PCGStaticMeshSpawnerConstants::InstanceCountsPinLabel) : INDEX_NONE;
+		AnalysisDataIndex = GetProducerKernel() ? InBinding->GetFirstInputDataIndex(GetProducerKernel(), PCGStaticMeshSpawnerKernelConstants::InstanceCountsPinLabel) : INDEX_NONE;
 
 		if (AnalysisDataIndex == INDEX_NONE)
 		{
@@ -745,7 +745,6 @@ void UPCGStaticMeshSpawnerDataProvider::CreatePrimitiveDescriptors(FPCGContext*
 	Primitives.AttributeIdOffsetStrides = AttributeIdOffsetStrides;
 	Primitives.SelectorAttributeId = SelectorAttributeId;
 	Primitives.SelectionCDF = PrimitiveSelectionCDF;
-	Primitives.PrimitiveStringKeys = PrimitiveStringKeys;
 	Primitives.PrimitiveMeshBounds = PrimitiveMeshBounds;
 
 	// Useful for debugging instance counts.
diff --git a/Engine/Plugins/PCG/Source/PCG/Private/Elements/PCGStaticMeshSpawner.cpp b/Engine/Plugins/PCG/Source/PCG/Private/Elements/PCGStaticMeshSpawner.cpp
--- a/Engine/Plugins/PCG/Source/PCG/Private/Elements/PCGStaticMeshSpawner.cpp
+++ b/Engine/Plugins/PCG/Source/PCG/Private/Elements/PCGStaticMeshSpawner.cpp
@@ -76,7 +76,7 @@ void UPCGStaticMeshSpawnerSettings::CreateKernels(FPCGGPUCompilationContext& InO
 		CountKernel->SetEmitPerDataCounts(false);
 		CountKernel->SetOutputRawBuffer(bUseRawBuffer);
 
-		OutEdges.Emplace(FPCGPinReference(CountKernel, PCGPinConstants::DefaultOutputLabel), FPCGPinReference(SpawnerKernel, PCGStaticMeshSpawnerConstants::InstanceCountsPinLabel));
+		OutEdges.Emplace(FPCGPinReference(CountKernel, PCGPinConstants::DefaultOutputLabel), FPCGPinReference(SpawnerKernel, PCGStaticMeshSpawnerKernelConstants::InstanceCountsPinLabel));
 	}
 }
 
diff --git a/Engine/Plugins/PCG/Source/PCG/Private/Elements/PCGStaticMeshSpawnerKernel.cpp b/Engine/Plugins/PCG/Source/PCG/Private/Elements/PCGStaticMeshSpawnerKernel.cpp
--- a/Engine/Plugins/PCG/Source/PCG/Private/Elements/PCGStaticMeshSpawnerKernel.cpp
+++ b/Engine/Plugins/PCG/Source/PCG/Private/Elements/PCGStaticMeshSpawnerKernel.cpp
@@ -164,7 +164,7 @@ void UPCGStaticMeshSpawnerKernel::GetKernelAttributeKeys(TArray<FPCGKernelAttrib
 void UPCGStaticMeshSpawnerKernel::GetInputPins(TArray<FPCGPinProperties>& OutPins) const
 {
 	OutPins.Emplace(PCGPinConstants::DefaultInputLabel, EPCGDataType::Point);
-	OutPins.Emplace(PCGStaticMeshSpawnerConstants::InstanceCountsPinLabel, bUseRawInstanceCountsBuffer ? FPCGDataTypeInfoRawBuffer::AsId() : FPCGDataTypeInfoParam::AsId());
+	OutPins.Emplace(PCGStaticMeshSpawnerKernelConstants::InstanceCountsPinLabel, bUseRawInstanceCountsBuffer ? FPCGDataTypeInfoRawBuffer::AsId() : FPCGDataTypeInfoParam::AsId());
 }
 
 void UPCGStaticMeshSpawnerKernel::GetOutputPins(TArray<FPCGPinPropertiesGPU>& OutPins) const
diff --git a/Engine/Plugins/PCG/Source/PCG/Public/Compute/PCGDataBinding.h b/Engine/Plugins/PCG/Source/PCG/Public/Compute/PCGDataBinding.h
--- a/Engine/Plugins/PCG/Source/PCG/Public/Compute/PCGDataBinding.h
+++ b/Engine/Plugins/PCG/Source/PCG/Public/Compute/PCGDataBinding.h
@@ -35,8 +35,6 @@ struct FPCGSpawnerPrimitives
 	TArray<FUintVector4> AttributeIdOffsetStrides;
 
 	int32 SelectorAttributeId = INDEX_NONE;
-
-	TArray<int32> PrimitiveStringKeys;
 };
 
 USTRUCT()
diff --git a/Engine/Plugins/PCG/Source/PCG/Public/Elements/PCGStaticMeshSpawnerKernel.h b/Engine/Plugins/PCG/Source/PCG/Public/Elements/PCGStaticMeshSpawnerKernel.h
--- a/Engine/Plugins/PCG/Source/PCG/Public/Elements/PCGStaticMeshSpawnerKernel.h
+++ b/Engine/Plugins/PCG/Source/PCG/Public/Elements/PCGStaticMeshSpawnerKernel.h
@@ -8,7 +8,7 @@
 
 class UPCGStaticMeshSpawnerSettings;
 
-namespace PCGStaticMeshSpawnerConstants
+namespace PCGStaticMeshSpawnerKernelConstants
 {
 	const FName InstanceCountsPinLabel = TEXT("InstanceCounts");
 }
