[PCG] Revert aggressive housekeeping in the ALV

* We need to keep a reference to data provided by the DataVisualization, otherwise it can be garbage collected

#rnx

diff --git a/Engine/Plugins/PCG/Source/PCGEditor/Private/Widgets/SPCGEditorGraphAttributeListView.cpp b/Engine/Plugins/PCG/Source/PCGEditor/Private/Widgets/SPCGEditorGraphAttributeListView.cpp
--- a/Engine/Plugins/PCG/Source/PCGEditor/Private/Widgets/SPCGEditorGraphAttributeListView.cpp
+++ b/Engine/Plugins/PCG/Source/PCGEditor/Private/Widgets/SPCGEditorGraphAttributeListView.cpp
@@ -900,6 +900,7 @@ void SPCGEditorGraphAttributeListView::RefreshAttributeList()
 	if (const IPCGDataVisualization* DataVisualization = DataVisRegistry.GetDataVisualization(PCGData->GetClass()))
 	{
 		const FPCGTableVisualizerInfo TableVisualizerInfo = DataVisualization->GetTableVisualizerInfoWithDomain(PCGData, MetadataDomainID);
+		DataToVisualize = TableVisualizerInfo.Data;
 
 		int32 NumEntries = 0;
 
diff --git a/Engine/Plugins/PCG/Source/PCGEditor/Private/Widgets/SPCGEditorGraphAttributeListView.h b/Engine/Plugins/PCG/Source/PCGEditor/Private/Widgets/SPCGEditorGraphAttributeListView.h
--- a/Engine/Plugins/PCG/Source/PCGEditor/Private/Widgets/SPCGEditorGraphAttributeListView.h
+++ b/Engine/Plugins/PCG/Source/PCGEditor/Private/Widgets/SPCGEditorGraphAttributeListView.h
@@ -342,9 +342,8 @@ private:
 	TSharedPtr<FPCGListViewUpdater> CurrentUpdateTask = nullptr;
 
 	/** In some cases, like when inspecting temporary collapsed point data, there is no owner of the data.
-	* As a temporary solution, since the ALV needs the data for visualization and double-click functionality,
-	* it can store a strong reference to the data to keep it alive during inspection. 
-	* @todo_pcg: Replace this with AddReferencedObjects
+	* Since the ALV needs the data for visualization and double-click functionality,
+	* it stores a reference to the data to keep it alive during inspection.
 	*/
 	TObjectPtr<const UPCGData> DataPtr = nullptr;
 
