[PCG] Split & moved tool data to private folder

#rnx

diff --git a/Engine/Plugins/PCG/Source/PCG/Internal/Data/Tool/PCGToolBaseData.h b/Engine/Plugins/PCG/Source/PCG/Internal/Data/Tool/PCGToolBaseData.h
new file mode 100644
--- /dev/null
+++ b/Engine/Plugins/PCG/Source/PCG/Internal/Data/Tool/PCGToolBaseData.h
@@ -0,0 +1,120 @@
+// Copyright Epic Games, Inc. All Rights Reserved.
+
+#pragma once
+
+#include "Data/PCGToolData.h"
+
+#include "PCGToolBaseData.generated.h"
+
+class UInteractiveTool;
+class UInteractiveToolPropertySet;
+struct FPCGContext;
+
+/** A transient struct that contains resources generated during tool usage for a specific tool data struct. Used to finalize or clean up when a tool is applied or cancelled. */
+USTRUCT()
+struct FPCGInteractiveToolWorkingDataGeneratedResources
+{
+	GENERATED_BODY()
+
+	FPCGInteractiveToolWorkingDataGeneratedResources() = default;
+	virtual ~FPCGInteractiveToolWorkingDataGeneratedResources() = default;
+
+#if WITH_EDITORONLY_DATA
+	/** Finalizes generated resources (e.g. remove Transient flag). */
+	PCG_API void Finalize();
+
+	/** Remove and destroy generated resources. */
+	PCG_API void Cleanup();
+
+	UPROPERTY(Transient)
+	TArray<TWeakObjectPtr<UActorComponent>> GeneratedComponents;
+#endif
+};
+
+/** A transient context struct that is passed in so the tool data can manage its data properly during tool usage. Used at editor-time only. */
+USTRUCT()
+struct FPCGInteractiveToolWorkingDataContext
+{
+	GENERATED_BODY()
+
+	FPCGInteractiveToolWorkingDataContext() = default;
+	virtual ~FPCGInteractiveToolWorkingDataContext() = default;
+
+#if WITH_EDITORONLY_DATA
+	PCG_API FPCGInteractiveToolWorkingDataContext(UPCGComponent* OwningComponent, UInteractiveTool* InteractiveTool, UInteractiveToolPropertySet* PropertySet, FName DataInstanceIdentifier);
+
+	/** The owning PCG Component. Should always be valid. Can be used as an Outer for new objects. */
+	UPROPERTY(Transient)
+	TWeakObjectPtr<UPCGComponent> OwningPCGComponent;
+
+	/** An optional DataInstanceIdentifier. Can be used to find or create a specific object, such as a Spline Component with a certain tag. */
+	UPROPERTY(Transient)
+	FName DataInstanceIdentifier;
+
+	/** The owning actor. The owning PCG component's outer. Should always be valid. Convenience accessor. */
+	UPROPERTY(Transient)
+	TWeakObjectPtr<AActor> OwningActor;
+
+	/** The interactive tool currently editing the tool data. */
+	UPROPERTY(Transient)
+	TWeakObjectPtr<UInteractiveTool> InteractiveTool;
+
+	/** The PCG Setting properties of the Interactive Tool. */
+	UPROPERTY(Transient)
+	TWeakObjectPtr<UInteractiveToolPropertySet> PCGSettings;
+#endif
+};
+
+/** The base struct for working data for a tool. This data is created by the tool, stored on the PCG Component,
+ * then read back by the tool on subsequent tool invocations. */
+USTRUCT()
+struct FPCGInteractiveToolWorkingData : public FPCGInteractiveToolWorkingBaseData
+{
+	GENERATED_BODY()
+
+	PCG_API bool IsInitialized() const;
+
+	PCG_API virtual bool IsValid() const;
+
+	/** Construct and add pcg data to the runtime context.
+	 * This will translate data from this struct into UPCGData to be processed by the graph.
+	 * Add dynamic tracking keys in here if you want to refresh when a change to your source data happens. */
+	PCG_API virtual void InitializeRuntimeElementData(FPCGContext* InContext) const;
+
+#if WITH_EDITORONLY_DATA
+public:
+	DECLARE_DELEGATE_OneParam(FOnInitialized, FPCGInteractiveToolWorkingData*)
+
+	/** Initializes data used by the tool working data. Initialize is called once when the working data is created. */
+	PCG_API void Initialize(const FPCGInteractiveToolWorkingDataContext& Context);
+
+	/** OnToolStart is called every time the tool is started up, or when the data is created.
+	 * Can be used to save the initial state of persistent data to restore to on Cancel, or to Modify existing data. */
+	PCG_API virtual void OnToolStart(const FPCGInteractiveToolWorkingDataContext& Context);
+
+	/** Called when the tool is applied. Finalize changes here. */
+	PCG_API virtual void OnToolApply(const FPCGInteractiveToolWorkingDataContext& Context);
+
+	/** Called when the tool is canceled. Revert changes you have made here. */
+	PCG_API virtual void OnToolCancel(const FPCGInteractiveToolWorkingDataContext& Context);
+
+	/** Called when reset data was specifically requested, for example by the user. */
+	virtual void OnResetToolDataRequested(const FPCGInteractiveToolWorkingDataContext& Context) {}
+
+	PCG_API FOnInitialized& OnInitialized() { return OnInitializedDelegate; }
+private:
+	/** Implement your initialization logic here. Called by Initialize. */
+	PCG_API virtual void InitializeInternal(const FPCGInteractiveToolWorkingDataContext& Context);
+protected:
+	/** Each working data struct keeps track of its generated resources. They have to be reverted during OnToolCancel. */
+	UPROPERTY(Transient)
+	FPCGInteractiveToolWorkingDataGeneratedResources GeneratedResources;
+
+	/** OnInitialized delegate to allow the settings class to do more work. */
+	FOnInitialized OnInitializedDelegate;
+#endif
+
+private:
+	UPROPERTY()
+	bool bInitialized = false;
+};
\ No newline at end of file
diff --git a/Engine/Plugins/PCG/Source/PCG/Internal/Data/Tool/PCGToolPointData.h b/Engine/Plugins/PCG/Source/PCG/Internal/Data/Tool/PCGToolPointData.h
new file mode 100644
--- /dev/null
+++ b/Engine/Plugins/PCG/Source/PCG/Internal/Data/Tool/PCGToolPointData.h
@@ -0,0 +1,41 @@
+// Copyright Epic Games, Inc. All Rights Reserved.
+
+#pragma once
+
+#include "Data/Tool/PCGToolBaseData.h"
+
+#include "PCGToolPointData.generated.h"
+
+/** Tool Data that stores points. */
+USTRUCT(BlueprintType, DisplayName = "Point Array Tool Data")
+struct FPCGInteractiveToolWorkingData_PointArrayData : public FPCGInteractiveToolWorkingData
+{
+	GENERATED_BODY()
+
+public:
+	FPCGInteractiveToolWorkingData_PointArrayData() = default;
+
+	PCG_API virtual void InitializeRuntimeElementData(FPCGContext* InContext) const override;
+
+	PCG_API virtual bool IsValid() const override;
+
+	PCG_API UPCGPointArrayData* GetPointArrayData() const;
+
+#if WITH_EDITORONLY_DATA
+	PCG_API virtual void InitializeInternal(const FPCGInteractiveToolWorkingDataContext& Context) override;
+
+	PCG_API virtual void OnToolStart(const FPCGInteractiveToolWorkingDataContext& Context) override;
+
+	PCG_API virtual void OnToolCancel(const FPCGInteractiveToolWorkingDataContext& Context) override;
+
+	PCG_API virtual void OnResetToolDataRequested(const FPCGInteractiveToolWorkingDataContext& Context) override;
+#endif
+
+protected:
+	UPROPERTY(VisibleAnywhere, BlueprintReadOnly, Category = "Working Data")
+	TObjectPtr<UPCGPointArrayData> GeneratedPointData;
+
+	/** A transient copy of the point array data, created on tool start/data access. We revert back to these points on Cancel. */
+	UPROPERTY(Transient)
+	TObjectPtr<UPCGPointArrayData> OnToolStartPointArray;
+};
\ No newline at end of file
diff --git a/Engine/Plugins/PCG/Source/PCG/Internal/Data/Tool/PCGToolSplineData.h b/Engine/Plugins/PCG/Source/PCG/Internal/Data/Tool/PCGToolSplineData.h
new file mode 100644
--- /dev/null
+++ b/Engine/Plugins/PCG/Source/PCG/Internal/Data/Tool/PCGToolSplineData.h
@@ -0,0 +1,62 @@
+// Copyright Epic Games, Inc. All Rights Reserved.
+
+#pragma once
+
+#include "Data/PCGSplineData.h"
+#include "Data/Tool/PCGToolBaseData.h"
+
+#include "PCGToolSplineData.generated.h"
+
+/** The working data for splines. */
+USTRUCT(BlueprintType, DisplayName = "Spline Tool Data")
+struct FPCGInteractiveToolWorkingData_Spline : public FPCGInteractiveToolWorkingData
+{
+	GENERATED_BODY()
+public:
+	FPCGInteractiveToolWorkingData_Spline() = default;
+	virtual ~FPCGInteractiveToolWorkingData_Spline() override = default;
+
+	PCG_API virtual void InitializeRuntimeElementData(FPCGContext* InContext) const override;
+
+	PCG_API virtual bool IsValid() const override;
+
+	PCG_API USplineComponent* GetSplineComponent() const;
+
+#if WITH_EDITOR
+public:
+	virtual void OnToolStart(const FPCGInteractiveToolWorkingDataContext& Context) override;
+	virtual void OnToolCancel(const FPCGInteractiveToolWorkingDataContext& Context) override;
+	virtual void OnResetToolDataRequested(const FPCGInteractiveToolWorkingDataContext& Context) override;
+private:
+	virtual void InitializeInternal(const FPCGInteractiveToolWorkingDataContext& Context) override;
+
+private:
+	USplineComponent* FindOrGenerateSplineComponent(const FPCGInteractiveToolWorkingDataContext& Context);
+
+	virtual USplineComponent* FindMatchingSplineComponent(const FPCGInteractiveToolWorkingDataContext& Context);
+	virtual USplineComponent* GenerateMatchingSplineComponent(const FPCGInteractiveToolWorkingDataContext& Context);
+#endif
+
+protected:
+	UPROPERTY(VisibleAnywhere, BlueprintReadOnly, Category = "Working Data")
+	TWeakObjectPtr<USplineComponent> SplineComponent;
+
+	UPROPERTY(Transient)
+	TObjectPtr<USplineComponent> OnToolStartSplineComponent;
+};
+
+USTRUCT(BlueprintType, DisplayName = "Spline Surface Tool Data")
+struct FPCGInteractiveToolWorkingData_SplineSurface : public FPCGInteractiveToolWorkingData_Spline
+{
+	GENERATED_BODY()
+
+	FPCGInteractiveToolWorkingData_SplineSurface() = default;
+	virtual ~FPCGInteractiveToolWorkingData_SplineSurface() override = default;
+
+#if WITH_EDITOR
+private:
+	/** We look for and generate a closed spline component. */
+	virtual USplineComponent* FindMatchingSplineComponent(const FPCGInteractiveToolWorkingDataContext& Context) override;
+	virtual USplineComponent* GenerateMatchingSplineComponent(const FPCGInteractiveToolWorkingDataContext& Context) override;
+#endif
+};
\ No newline at end of file
diff --git a/Engine/Plugins/PCG/Source/PCG/Internal/Data/Tool/PCGToolVolumeData.h b/Engine/Plugins/PCG/Source/PCG/Internal/Data/Tool/PCGToolVolumeData.h
new file mode 100644
--- /dev/null
+++ b/Engine/Plugins/PCG/Source/PCG/Internal/Data/Tool/PCGToolVolumeData.h
@@ -0,0 +1,48 @@
+// Copyright Epic Games, Inc. All Rights Reserved.
+
+#pragma once
+
+#include "Data/Tool/PCGToolBaseData.h"
+
+#include "PCGToolVolumeData.generated.h"
+
+class AVolume;
+class UBrushComponent;
+class UBrushBuilder;
+
+USTRUCT(BlueprintType, DisplayName = "Volume Tool Data")
+struct FPCGInteractiveToolWorkingData_Volume : public FPCGInteractiveToolWorkingData
+{
+	GENERATED_BODY()
+public:
+	UPROPERTY(VisibleAnywhere, BlueprintReadOnly, Category = "Working Data")
+	TWeakObjectPtr<AVolume> Volume = nullptr;
+
+	UPROPERTY(VisibleAnywhere, BlueprintReadOnly, Category = "Working Data")
+	TWeakObjectPtr<UBrushComponent> BrushComponent = nullptr;
+
+#if WITH_EDITOR
+public:
+	virtual void OnToolStart(const FPCGInteractiveToolWorkingDataContext& Context) override;
+	virtual void OnToolCancel(const FPCGInteractiveToolWorkingDataContext& Context) override;
+
+private:
+	virtual void InitializeInternal(const FPCGInteractiveToolWorkingDataContext& Context) override;
+
+	UBrushComponent* FindOrCreateBrushComponent(const FPCGInteractiveToolWorkingDataContext& Context);
+	virtual UBrushComponent* FindMatchingBrushComponent(const FPCGInteractiveToolWorkingDataContext& Context);
+	virtual UBrushComponent* GenerateMatchingBrushComponent(const FPCGInteractiveToolWorkingDataContext& Context);
+#endif
+
+protected:
+#if WITH_EDITORONLY_DATA
+	UPROPERTY(Transient)
+	TObjectPtr<UBrushComponent> OnToolStartBrushComponent = nullptr;
+
+	UPROPERTY(Transient)
+	TObjectPtr<UBrushBuilder> OnToolStartBrushBuilder = nullptr;
+
+	UPROPERTY(Transient)
+	FTransform OnToolStartVolumeTransform;
+#endif
+};
\ No newline at end of file
diff --git a/Engine/Plugins/PCG/Source/PCG/Private/Data/PCGToolData.cpp b/Engine/Plugins/PCG/Source/PCG/Private/Data/PCGToolData.cpp
--- a/Engine/Plugins/PCG/Source/PCG/Private/Data/PCGToolData.cpp
+++ b/Engine/Plugins/PCG/Source/PCG/Private/Data/PCGToolData.cpp
@@ -2,517 +2,7 @@
 
 #include "Data/PCGToolData.h"
 
-#include "PCGComponent.h"
-#include "PCGContext.h"
-#include "Data/PCGPointArrayData.h"
-#include "Helpers/PCGDynamicTrackingHelpers.h"
-
-#include "Components/BrushComponent.h"
-#include "Components/SplineComponent.h"
-#include "Components/BoxComponent.h"
-#include "Engine/Engine.h"
-#include "Engine/BrushBuilder.h"
-#include "GameFramework/Volume.h"
-#include "Logging/StructuredLog.h"
-#include "UObject/UObjectGlobals.h"
-#include "UObject/Package.h"
-
-#if WITH_EDITOR
-#include "InteractiveTool.h"
-#include "ActorFactories/ActorFactory.h"
-#endif
-
-#if WITH_EDITORONLY_DATA
-void FPCGInteractiveToolWorkingData::Initialize(const FPCGInteractiveToolWorkingDataContext& Context)
-{
-	ensureMsgf(Context.OwningPCGComponent.IsValid() && Context.OwningActor.IsValid(), TEXT("Owning PCG Component or Owning Actor invalid! This should never happen so an assumption was hurt."));
-
-	InitializeInternal(Context);
-	
-	OnInitializedDelegate.ExecuteIfBound(this);
-
-	bInitialized = true;
-}
-
-void FPCGInteractiveToolWorkingData::InitializeInternal(const FPCGInteractiveToolWorkingDataContext& Context)
-{
-}
-
-void FPCGInteractiveToolWorkingData::OnToolStart(const FPCGInteractiveToolWorkingDataContext& Context)
-{
-}
-
-void FPCGInteractiveToolWorkingData::OnToolApply(const FPCGInteractiveToolWorkingDataContext& Context)
-{
-	GeneratedResources.Finalize();
-}
-
-void FPCGInteractiveToolWorkingData::OnToolCancel(const FPCGInteractiveToolWorkingDataContext& Context)
-{
-	GeneratedResources.Cleanup();
-}
-
-void FPCGInteractiveToolWorkingDataGeneratedResources::Finalize()
-{
-	for(TWeakObjectPtr<UActorComponent> GeneratedComponent : GeneratedComponents)
-	{
-		if(GeneratedComponent.IsValid())
-		{
-			GeneratedComponent->ClearFlags(RF_Transient);
-		}
-	}
-
-	GeneratedComponents.Empty();
-}
-
-void FPCGInteractiveToolWorkingDataGeneratedResources::Cleanup()
-{
-	for(TWeakObjectPtr<UActorComponent> GeneratedComponent : GeneratedComponents)
-	{
-		if(GeneratedComponent.IsValid())
-		{
-			if(UPCGComponent* PCGComponent = Cast<UPCGComponent>(GeneratedComponent))
-			{
-				PCGComponent->Cleanup(true);	
-			}
-			
-			GeneratedComponent->DestroyComponent();
-		}
-	}
-
-	GeneratedComponents.Empty();
-}
-#endif
-
-#if WITH_EDITORONLY_DATA
-FPCGInteractiveToolWorkingDataContext::FPCGInteractiveToolWorkingDataContext(UPCGComponent* InOwningComponent, UInteractiveTool* InInteractiveTool,	UInteractiveToolPropertySet* InPropertySet, FName InDataInstanceIdentifier)
-{
-	OwningPCGComponent = InOwningComponent;
-	OwningActor = OwningPCGComponent->GetOwner();
-	DataInstanceIdentifier = InDataInstanceIdentifier;
-	InteractiveTool = InInteractiveTool;
-	PCGSettings = InPropertySet;
-}
-#endif
-
-bool FPCGInteractiveToolWorkingData::IsInitialized() const
-{
-	return bInitialized;
-}
-
-bool FPCGInteractiveToolWorkingData::IsValid() const
-{
-	return true;
-}
-
-void FPCGInteractiveToolWorkingData::InitializeRuntimeElementData(FPCGContext* InContext) const
-{
-}
-
-#if WITH_EDITORONLY_DATA
-void FPCGInteractiveToolWorkingData_PointArrayData::InitializeInternal(const FPCGInteractiveToolWorkingDataContext& Context)
-{
-	GeneratedPointData = NewObject<UPCGPointArrayData>(Context.OwningPCGComponent.Get());
-}
-
-void FPCGInteractiveToolWorkingData_PointArrayData::OnToolStart(const FPCGInteractiveToolWorkingDataContext& Context)
-{
-	OnToolStartPointArray = Cast<UPCGPointArrayData>(StaticDuplicateObject(GeneratedPointData, GetTransientPackage()));
-	
-	FPCGInteractiveToolWorkingData::OnToolStart(Context);
-}
-
-void FPCGInteractiveToolWorkingData_PointArrayData::OnToolCancel(const FPCGInteractiveToolWorkingDataContext& Context)
-{
-	if(OnToolStartPointArray != nullptr)
-	{
-		// We'd expect no relationship between the generated point data on tool cancel, but we'll make sure this is correct here.
-		GeneratedPointData = Cast<UPCGPointArrayData>(StaticDuplicateObject(OnToolStartPointArray, Context.OwningPCGComponent.Get()));
-	}
-	
-	FPCGInteractiveToolWorkingData::OnToolCancel(Context);
-}
-
-void FPCGInteractiveToolWorkingData_PointArrayData::OnResetToolDataRequested(const FPCGInteractiveToolWorkingDataContext& Context)
-{
-	if (GeneratedPointData)
-	{
-		GeneratedPointData->SetNumPoints(0);
-
-		FPropertyChangedEvent EmptyPropertyChangedEvent(nullptr);
-		FCoreUObjectDelegates::OnObjectPropertyChanged.Broadcast(GeneratedPointData, EmptyPropertyChangedEvent);
-	}
-}
-#endif
-
-void FPCGInteractiveToolWorkingData_PointArrayData::InitializeRuntimeElementData(FPCGContext* Context) const
-{
-	FPCGInteractiveToolWorkingData::InitializeRuntimeElementData(Context);
-
-	TArray<FPCGTaggedData>& Outputs = Context->OutputData.TaggedData;
-
-	FPCGTaggedData& Output = Outputs.Emplace_GetRef();
-
-#if WITH_EDITOR
-	// Implementation note - we duplicate here to make sure we are not preserving the parent, because it might still change, which would cause issues.
-	check(IsInGameThread());
-	Output.Data = Cast<UPCGData>(StaticDuplicateObject(GeneratedPointData, GetTransientPackage()));
-#else
-	Output.Data = GeneratedPointData->DuplicateData(Context);
-#endif
-
-#if WITH_EDITOR
-	FPCGSelectionKey Key = FPCGSelectionKey::CreateFromObjectPtr(GeneratedPointData);
-	FPCGDynamicTrackingHelper::AddSingleDynamicTrackingKey(Context, MoveTemp(Key), false);
-#endif
-}
-
-bool FPCGInteractiveToolWorkingData_PointArrayData::IsValid() const
-{
-	return GeneratedPointData != nullptr;
-}
-
-UPCGPointArrayData* FPCGInteractiveToolWorkingData_PointArrayData::GetPointArrayData() const
-{
-	return GeneratedPointData;
-}
-
-USplineComponent* FPCGInteractiveToolWorkingData_Spline::GetSplineComponent() const
-{
-	return SplineComponent.Get();
-}
-
-#if WITH_EDITOR
-void FPCGInteractiveToolWorkingData_Spline::InitializeInternal(const FPCGInteractiveToolWorkingDataContext& Context)
-{	
-	SplineComponent = FindOrGenerateSplineComponent(Context);
-}
-
-void FPCGInteractiveToolWorkingData_Spline::OnToolStart(const FPCGInteractiveToolWorkingDataContext& Context)
-{
-	FPCGInteractiveToolWorkingData::OnToolStart(Context);
-
-	if(SplineComponent.IsValid())
-	{
-		OnToolStartSplineComponent = Cast<USplineComponent>(StaticDuplicateObject(SplineComponent.Get(), GetTransientPackage()));
-	}
-}
-
-void FPCGInteractiveToolWorkingData_Spline::OnToolCancel(const FPCGInteractiveToolWorkingDataContext& Context)
-{
-	// If the Spline Component is valid and isn't a generated component, we revert its state.
-	if(SplineComponent.IsValid() && GeneratedResources.GeneratedComponents.Contains(SplineComponent) == false)
-	{
-		if(OnToolStartSplineComponent)
-		{
-			// We restore the spline component to its previous state
-			UEngine::CopyPropertiesForUnrelatedObjects(OnToolStartSplineComponent, SplineComponent.Get());
-			
-			SplineComponent->UpdateSpline();
-			
-			OnToolStartSplineComponent->MarkAsGarbage();
-			OnToolStartSplineComponent = nullptr;
-		}
-	}
-
-	// This will delete the spline component if it was generated
-	FPCGInteractiveToolWorkingData::OnToolCancel(Context);
-}
-
-void FPCGInteractiveToolWorkingData_Spline::OnResetToolDataRequested(const FPCGInteractiveToolWorkingDataContext& Context)
-{
-	if (SplineComponent.IsValid())
-	{
-		SplineComponent->ClearSplinePoints();
-
-		FPropertyChangedEvent EmptyPropertyChangedEvent(nullptr);
-		FCoreUObjectDelegates::OnObjectPropertyChanged.Broadcast(SplineComponent.Get(), EmptyPropertyChangedEvent);
-	}
-}
-
-USplineComponent* FPCGInteractiveToolWorkingData_Spline::FindOrGenerateSplineComponent(const FPCGInteractiveToolWorkingDataContext& Context)
-{
-	USplineComponent* MatchingSplineComponent = FindMatchingSplineComponent(Context);
-
-	// If we didn't find any spline component to match with, we create it now
-	if (MatchingSplineComponent == nullptr)
-	{
-		MatchingSplineComponent = GenerateMatchingSplineComponent(Context);
-		GeneratedResources.GeneratedComponents.Add(MatchingSplineComponent);
-	}	
-	
-	ensureMsgf(MatchingSplineComponent != nullptr, TEXT("Matching Spline Component should always be found after initialization as it was either found or created."));
-	return MatchingSplineComponent;
-}
-
-USplineComponent* FPCGInteractiveToolWorkingData_Spline::FindMatchingSplineComponent(const FPCGInteractiveToolWorkingDataContext& Context)
-{
-	TArray<USplineComponent*> AllSplineComponents;
-	Context.OwningActor->GetComponents<USplineComponent>(AllSplineComponents, false);
-	
-	if (!Context.DataInstanceIdentifier.IsNone())
-	{
-		// When the DataInstanceIdentifier is set, we only care about the splines with the given tag
-		AllSplineComponents.RemoveAll([DataInstanceIdentifier = Context.DataInstanceIdentifier](USplineComponent* Candidate)
-		{
-			return Candidate->ComponentHasTag(DataInstanceIdentifier) == false;
-		});
-	}
-
-	USplineComponent* MatchingSplineComponent = AllSplineComponents.IsEmpty() ? nullptr : AllSplineComponents[0];
-
-	if(AllSplineComponents.Num() > 1 && MatchingSplineComponent != nullptr)
-	{
-		UE_LOGFMT(LogPCG, Warning, "More than one spline component found. Choosing the first available: {0}", MatchingSplineComponent->GetName());
-	}
-	
-	return MatchingSplineComponent;
-}
-
-USplineComponent* FPCGInteractiveToolWorkingData_Spline::GenerateMatchingSplineComponent(const FPCGInteractiveToolWorkingDataContext& Context)
-{
-	USplineComponent* MatchingSplineComponent = NewObject<USplineComponent>(Context.OwningActor.Get());
-	MatchingSplineComponent->SetMobility(Context.OwningActor->GetRootComponent()->GetMobility());
-	MatchingSplineComponent->SetupAttachment(Context.OwningActor->GetRootComponent());
-	// Transient until Applied. 
-	MatchingSplineComponent->SetFlags(RF_Transient);
-
-	if (Context.DataInstanceIdentifier.IsNone() == false)
-	{
-		MatchingSplineComponent->ComponentTags.Add(Context.DataInstanceIdentifier);
-	}
-
-	MatchingSplineComponent->ClearSplinePoints();
-	MatchingSplineComponent->bSplineHasBeenEdited = true;
-
-	// We set it to open or closed using OnInitialized based on the user property
-	MatchingSplineComponent->SetClosedLoop(false);
-		
-	MatchingSplineComponent->RegisterComponent();
-	Context.OwningActor->AddInstanceComponent(MatchingSplineComponent);
-
-	return MatchingSplineComponent;
-}
-
-USplineComponent* FPCGInteractiveToolWorkingData_SplineSurface::FindMatchingSplineComponent(const FPCGInteractiveToolWorkingDataContext& Context)
-{
-	TArray<USplineComponent*> AllClosedSplineComponents;
-	Context.OwningActor->GetComponents<USplineComponent>(AllClosedSplineComponents, false);
-
-	// We remove all spline components that are open loops since this working data struct is only for spline surfaces
-	AllClosedSplineComponents.RemoveAll([](USplineComponent* Candidate)
-	{
-		return Candidate->IsClosedLoop() == false;
-	});
-	
-	USplineComponent* MatchingSplineComponent = nullptr;
-	if (Context.DataInstanceIdentifier.IsNone())
-	{
-		if(AllClosedSplineComponents.Num() > 0)
-		{
-			MatchingSplineComponent = AllClosedSplineComponents[0];
-		}
-	}
-	else
-	{
-		// When the DataInstanceIdentifier is set, we only care about the closed splines with the given tag
-		AllClosedSplineComponents.RemoveAll([DataInstanceIdentifier = Context.DataInstanceIdentifier](USplineComponent* Candidate)
-		{
-			return Candidate->ComponentHasTag(DataInstanceIdentifier) == false;
-		});
-		
-		// If we find a matching spline component with the specific tag, use i
-		if (AllClosedSplineComponents.Num() > 0)
-		{
-			MatchingSplineComponent = Cast<USplineComponent>(AllClosedSplineComponents[0]);
-		}
-	}
-
-	if(AllClosedSplineComponents.Num() > 1 && MatchingSplineComponent != nullptr)
-	{
-		UE_LOGFMT(LogPCG, Warning, "More than one spline component found. Choosing the first available: {0}", MatchingSplineComponent->GetName());
-	}
-
-	return MatchingSplineComponent;
-}
-
-USplineComponent* FPCGInteractiveToolWorkingData_SplineSurface::GenerateMatchingSplineComponent(const FPCGInteractiveToolWorkingDataContext& Context)
-{
-	USplineComponent* MatchingSplineComponent = NewObject<USplineComponent>(Context.OwningActor.Get());
-	MatchingSplineComponent->SetMobility(Context.OwningActor->GetRootComponent()->GetMobility());
-	MatchingSplineComponent->SetupAttachment(Context.OwningActor->GetRootComponent());
-	// Transient until Applied. 
-	MatchingSplineComponent->SetFlags(RF_Transient);
-
-	if (Context.DataInstanceIdentifier.IsNone() == false)
-	{
-		MatchingSplineComponent->ComponentTags.Add(Context.DataInstanceIdentifier);
-	}
-
-	MatchingSplineComponent->ClearSplinePoints();
-	MatchingSplineComponent->bSplineHasBeenEdited = true;
-	MatchingSplineComponent->SetClosedLoop(true);
-
-	MatchingSplineComponent->RegisterComponent();
-	Context.OwningActor->AddInstanceComponent(MatchingSplineComponent);
-
-	return MatchingSplineComponent;
-}
-#endif // WITH_EDITOR
-
-void FPCGInteractiveToolWorkingData_Spline::InitializeRuntimeElementData(FPCGContext* Context) const
-{
-	FPCGInteractiveToolWorkingData::InitializeRuntimeElementData(Context);
-
-	TArray<FPCGTaggedData>& Outputs = Context->OutputData.TaggedData;
-
-	if(SplineComponent.IsValid())
-	{
-		FPCGTaggedData& Output = Outputs.Emplace_GetRef();
-		// We duplicate it so that PCG can take 'ownership'.
-		// If not, it can't process references and caches properly
-		// Might be relying on the fact it shouldn't access the same object over multiple generations
-		// And if we recreate it every time it gets properly managed
-		// If it gets dropped from cache, it's expected to be invalid
-		UPCGSplineData* SplineData = FPCGContext::NewObject_AnyThread<UPCGSplineData>(Context);
-		SplineData->Initialize(SplineComponent.Get());
-		Output.Data = SplineData;
-#if WITH_EDITOR
-		FPCGSelectionKey Key = FPCGSelectionKey::CreateFromObjectPtr(SplineComponent.Get());
-		FPCGDynamicTrackingHelper::AddSingleDynamicTrackingKey(Context, MoveTemp(Key), false);
-#endif
-	}
-}
-
-bool FPCGInteractiveToolWorkingData_Spline::IsValid() const
-{
-	return SplineComponent.IsValid();
-}
-
-#if WITH_EDITOR
-void FPCGInteractiveToolWorkingData_Volume::InitializeInternal(const FPCGInteractiveToolWorkingDataContext& Context)
-{
-	Volume = Cast<AVolume>(Context.OwningActor.Get());
-	BrushComponent = FindOrCreateBrushComponent(Context);
-}
-
-void FPCGInteractiveToolWorkingData_Volume::OnToolStart(const FPCGInteractiveToolWorkingDataContext& Context)
-{
-	FPCGInteractiveToolWorkingData::OnToolStart(Context);
-
-	if (BrushComponent.IsValid())
-	{
-		OnToolStartBrushComponent = Cast<UBrushComponent>(StaticDuplicateObject(BrushComponent.Get(), GetTransientPackage()));
-	}
-
-	if (Volume.IsValid())
-	{
-		OnToolStartVolumeTransform = Volume->GetActorTransform();
-
-		if (Volume->BrushBuilder)
-		{
-			OnToolStartBrushBuilder = DuplicateObject<UBrushBuilder>(Volume->BrushBuilder, GetTransientPackage());
-		}
-	}
-}
-
-void FPCGInteractiveToolWorkingData_Volume::OnToolCancel(const FPCGInteractiveToolWorkingDataContext& Context)
-{
-	if (BrushComponent.IsValid() && GeneratedResources.GeneratedComponents.Contains(BrushComponent) == false && OnToolStartBrushComponent)
-	{
-		Volume->SetActorTransform(OnToolStartVolumeTransform);
-
-		if (OnToolStartBrushBuilder)
-		{
-			// Recreate the brush from the original actor
-			UActorFactory::CreateBrushForVolumeActor(Volume.Get(), OnToolStartBrushBuilder);
-			OnToolStartBrushBuilder->MarkAsGarbage();
-			OnToolStartBrushBuilder = nullptr;
-		}
-
-		// Restore the brush component to its previous state
-		UEngine::CopyPropertiesForUnrelatedObjects(OnToolStartBrushComponent, BrushComponent.Get());
-		OnToolStartBrushComponent->MarkAsGarbage();
-		OnToolStartBrushComponent = nullptr;
-	}
-
-	FPCGInteractiveToolWorkingData::OnToolCancel(Context);
-}
-
-UBrushComponent* FPCGInteractiveToolWorkingData_Volume::FindOrCreateBrushComponent(const FPCGInteractiveToolWorkingDataContext& Context)
-{
-	UBrushComponent* MatchingComponent = FindMatchingBrushComponent(Context);
-
-	if (!MatchingComponent)
-	{
-		MatchingComponent = GenerateMatchingBrushComponent(Context);
-		GeneratedResources.GeneratedComponents.Add(MatchingComponent);
-	}
-
-	ensureMsgf(MatchingComponent != nullptr, TEXT("Matching Brush Component should always be found after initialization as it was either found or created."));
-	return MatchingComponent;
-}
-
-UBrushComponent* FPCGInteractiveToolWorkingData_Volume::FindMatchingBrushComponent(const FPCGInteractiveToolWorkingDataContext& Context)
-{
-	TArray<UBrushComponent*> AllBrushComponents;
-	if (Context.OwningActor.IsValid())
-	{
-		Context.OwningActor->GetComponents<UBrushComponent>(AllBrushComponents, false);
-	}
-
-	UBrushComponent* MatchingComponent = nullptr;
-
-	if(!Context.DataInstanceIdentifier.IsNone())
-	{
-		AllBrushComponents.RemoveAll([DataInstanceIdentifier = Context.DataInstanceIdentifier](UBrushComponent* Candidate)
-		{
-			return !Candidate || !Candidate->ComponentHasTag(DataInstanceIdentifier);
-		});
-	}
-
-	MatchingComponent = AllBrushComponents.IsEmpty() ? nullptr : AllBrushComponents[0];
-
-	if (AllBrushComponents.Num() > 1 && MatchingComponent != nullptr)
-	{
-		UE_LOGFMT(LogPCG, Warning, "More than one brush component found. Choosing the first available: {0}", MatchingComponent->GetName());
-	}
-
-	return MatchingComponent;
-}
-
-UBrushComponent* FPCGInteractiveToolWorkingData_Volume::GenerateMatchingBrushComponent(const FPCGInteractiveToolWorkingDataContext& Context)
-{
-	if (!Context.OwningActor.IsValid())
-	{
-		return nullptr;
-	}
-
-	UBrushComponent* MatchingBrushComponent = NewObject<UBrushComponent>(Context.OwningActor.Get());
-	// ABrush setup
-	MatchingBrushComponent->Mobility = EComponentMobility::Static;
-	MatchingBrushComponent->SetCanEverAffectNavigation(false);
-	// AVolume setup
-	MatchingBrushComponent->AlwaysLoadOnClient = true;
-	MatchingBrushComponent->AlwaysLoadOnServer = true;
-	// APCGVolume setup
-	MatchingBrushComponent->SetCollisionObjectType(ECC_WorldStatic);
-	MatchingBrushComponent->SetCollisionResponseToAllChannels(ECR_Ignore);
-	MatchingBrushComponent->SetGenerateOverlapEvents(false);
-	// Transient until applied
-	MatchingBrushComponent->SetFlags(RF_Transient);
-
-	if (!Context.DataInstanceIdentifier.IsNone())
-	{
-		MatchingBrushComponent->ComponentTags.Add(Context.DataInstanceIdentifier);
-	}
-
-	MatchingBrushComponent->RegisterComponent();
-	Context.OwningActor->AddInstanceComponent(MatchingBrushComponent);
-
-	return MatchingBrushComponent;
-}
-#endif
+#include "GameFramework/Actor.h"
 
 FPCGGraphToolData::FPCGGraphToolData()
 	: InitialActorClassToSpawn(AActor::StaticClass())
diff --git a/Engine/Plugins/PCG/Source/PCG/Private/Data/Tool/PCGToolBaseData.cpp b/Engine/Plugins/PCG/Source/PCG/Private/Data/Tool/PCGToolBaseData.cpp
new file mode 100644
--- /dev/null
+++ b/Engine/Plugins/PCG/Source/PCG/Private/Data/Tool/PCGToolBaseData.cpp
@@ -0,0 +1,99 @@
+// Copyright Epic Games, Inc. All Rights Reserved.
+
+#include "Data/Tool/PCGToolBaseData.h"
+
+#include "PCGComponent.h"
+#include "PCGContext.h"
+#include "Helpers/PCGDynamicTrackingHelpers.h"
+
+#include "GameFramework/Actor.h"
+
+#if WITH_EDITOR
+#include "InteractiveTool.h"
+#endif
+
+#if WITH_EDITOR
+void FPCGInteractiveToolWorkingData::Initialize(const FPCGInteractiveToolWorkingDataContext& Context)
+{
+	ensureMsgf(Context.OwningPCGComponent.IsValid() && Context.OwningActor.IsValid(), TEXT("Owning PCG Component or Owning Actor invalid! This should never happen so an assumption was hurt."));
+
+	InitializeInternal(Context);
+
+	OnInitializedDelegate.ExecuteIfBound(this);
+
+	bInitialized = true;
+}
+
+void FPCGInteractiveToolWorkingData::InitializeInternal(const FPCGInteractiveToolWorkingDataContext& Context)
+{
+}
+
+void FPCGInteractiveToolWorkingData::OnToolStart(const FPCGInteractiveToolWorkingDataContext& Context)
+{
+}
+
+void FPCGInteractiveToolWorkingData::OnToolApply(const FPCGInteractiveToolWorkingDataContext& Context)
+{
+	GeneratedResources.Finalize();
+}
+
+void FPCGInteractiveToolWorkingData::OnToolCancel(const FPCGInteractiveToolWorkingDataContext& Context)
+{
+	GeneratedResources.Cleanup();
+}
+
+bool FPCGInteractiveToolWorkingData::IsInitialized() const
+{
+	return bInitialized;
+}
+
+bool FPCGInteractiveToolWorkingData::IsValid() const
+{
+	return true;
+}
+
+void FPCGInteractiveToolWorkingData::InitializeRuntimeElementData(FPCGContext* InContext) const
+{
+}
+
+void FPCGInteractiveToolWorkingDataGeneratedResources::Finalize()
+{
+	for (TWeakObjectPtr<UActorComponent> GeneratedComponent : GeneratedComponents)
+	{
+		if (GeneratedComponent.IsValid())
+		{
+			GeneratedComponent->ClearFlags(RF_Transient);
+		}
+	}
+
+	GeneratedComponents.Empty();
+}
+
+void FPCGInteractiveToolWorkingDataGeneratedResources::Cleanup()
+{
+	for (TWeakObjectPtr<UActorComponent> GeneratedComponent : GeneratedComponents)
+	{
+		if (GeneratedComponent.IsValid())
+		{
+			if (UPCGComponent* PCGComponent = Cast<UPCGComponent>(GeneratedComponent))
+			{
+				PCGComponent->Cleanup(true);
+			}
+
+			GeneratedComponent->DestroyComponent();
+		}
+	}
+
+	GeneratedComponents.Empty();
+}
+
+FPCGInteractiveToolWorkingDataContext::FPCGInteractiveToolWorkingDataContext(UPCGComponent* InOwningComponent, UInteractiveTool* InInteractiveTool, UInteractiveToolPropertySet* InPropertySet, FName InDataInstanceIdentifier)
+{
+	OwningPCGComponent = InOwningComponent;
+	OwningActor = OwningPCGComponent->GetOwner();
+	DataInstanceIdentifier = InDataInstanceIdentifier;
+	InteractiveTool = InInteractiveTool;
+	PCGSettings = InPropertySet;
+}
+#endif // WITH_EDITOR
+
diff --git a/Engine/Plugins/PCG/Source/PCG/Private/Data/Tool/PCGToolPointData.cpp b/Engine/Plugins/PCG/Source/PCG/Private/Data/Tool/PCGToolPointData.cpp
new file mode 100644
--- /dev/null
+++ b/Engine/Plugins/PCG/Source/PCG/Private/Data/Tool/PCGToolPointData.cpp
@@ -0,0 +1,80 @@
+// Copyright Epic Games, Inc. All Rights Reserved.
+
+#include "Data/Tool/PCGToolPointData.h"
+
+#include "PCGComponent.h"
+#include "PCGContext.h"
+#include "Data/PCGPointArrayData.h"
+#include "Helpers/PCGDynamicTrackingHelpers.h"
+
+#include "Logging/StructuredLog.h"
+#include "UObject/UObjectGlobals.h"
+#include "UObject/Package.h"
+
+#if WITH_EDITOR
+void FPCGInteractiveToolWorkingData_PointArrayData::InitializeInternal(const FPCGInteractiveToolWorkingDataContext& Context)
+{
+	GeneratedPointData = NewObject<UPCGPointArrayData>(Context.OwningPCGComponent.Get());
+}
+
+void FPCGInteractiveToolWorkingData_PointArrayData::OnToolStart(const FPCGInteractiveToolWorkingDataContext& Context)
+{
+	OnToolStartPointArray = Cast<UPCGPointArrayData>(StaticDuplicateObject(GeneratedPointData, GetTransientPackage()));
+
+	Super::OnToolStart(Context);
+}
+
+void FPCGInteractiveToolWorkingData_PointArrayData::OnToolCancel(const FPCGInteractiveToolWorkingDataContext& Context)
+{
+	if (OnToolStartPointArray != nullptr)
+	{
+		// We'd expect no relationship between the generated point data on tool cancel, but we'll make sure this is correct here.
+		GeneratedPointData = Cast<UPCGPointArrayData>(StaticDuplicateObject(OnToolStartPointArray, Context.OwningPCGComponent.Get()));
+	}
+
+	Super::OnToolCancel(Context);
+}
+
+void FPCGInteractiveToolWorkingData_PointArrayData::OnResetToolDataRequested(const FPCGInteractiveToolWorkingDataContext& Context)
+{
+	if (GeneratedPointData)
+	{
+		GeneratedPointData->SetNumPoints(0);
+
+		FPropertyChangedEvent EmptyPropertyChangedEvent(nullptr);
+		FCoreUObjectDelegates::OnObjectPropertyChanged.Broadcast(GeneratedPointData, EmptyPropertyChangedEvent);
+	}
+}
+#endif // WITH_EDITOR
+
+void FPCGInteractiveToolWorkingData_PointArrayData::InitializeRuntimeElementData(FPCGContext* Context) const
+{
+	Super::InitializeRuntimeElementData(Context);
+
+	TArray<FPCGTaggedData>& Outputs = Context->OutputData.TaggedData;
+
+	FPCGTaggedData& Output = Outputs.Emplace_GetRef();
+
+#if WITH_EDITOR
+	// Implementation note - we duplicate here to make sure we are not preserving the parent, because it might still change, which would cause issues.
+	check(IsInGameThread());
+	Output.Data = Cast<UPCGData>(StaticDuplicateObject(GeneratedPointData, GetTransientPackage()));
+#else
+	Output.Data = GeneratedPointData->DuplicateData(Context);
+#endif
+
+#if WITH_EDITOR
+	FPCGSelectionKey Key = FPCGSelectionKey::CreateFromObjectPtr(GeneratedPointData);
+	FPCGDynamicTrackingHelper::AddSingleDynamicTrackingKey(Context, MoveTemp(Key), false);
+#endif
+}
+
+bool FPCGInteractiveToolWorkingData_PointArrayData::IsValid() const
+{
+	return GeneratedPointData != nullptr;
+}
+
+UPCGPointArrayData* FPCGInteractiveToolWorkingData_PointArrayData::GetPointArrayData() const
+{
+	return GeneratedPointData;
+}
\ No newline at end of file
diff --git a/Engine/Plugins/PCG/Source/PCG/Private/Data/Tool/PCGToolSplineData.cpp b/Engine/Plugins/PCG/Source/PCG/Private/Data/Tool/PCGToolSplineData.cpp
new file mode 100644
--- /dev/null
+++ b/Engine/Plugins/PCG/Source/PCG/Private/Data/Tool/PCGToolSplineData.cpp
@@ -0,0 +1,227 @@
+// Copyright Epic Games, Inc. All Rights Reserved.
+
+#include "Data/Tool/PCGToolSplineData.h"
+
+#include "PCGComponent.h"
+#include "PCGContext.h"
+#include "Helpers/PCGDynamicTrackingHelpers.h"
+
+#include "Components/SplineComponent.h"
+#include "Engine/Engine.h"
+#include "Logging/StructuredLog.h"
+#include "UObject/UObjectGlobals.h"
+#include "UObject/Package.h"
+
+USplineComponent* FPCGInteractiveToolWorkingData_Spline::GetSplineComponent() const
+{
+	return SplineComponent.Get();
+}
+
+#if WITH_EDITOR
+void FPCGInteractiveToolWorkingData_Spline::InitializeInternal(const FPCGInteractiveToolWorkingDataContext& Context)
+{
+	SplineComponent = FindOrGenerateSplineComponent(Context);
+}
+
+void FPCGInteractiveToolWorkingData_Spline::OnToolStart(const FPCGInteractiveToolWorkingDataContext& Context)
+{
+	Super::OnToolStart(Context);
+
+	if (SplineComponent.IsValid())
+	{
+		OnToolStartSplineComponent = Cast<USplineComponent>(StaticDuplicateObject(SplineComponent.Get(), GetTransientPackage()));
+	}
+}
+
+void FPCGInteractiveToolWorkingData_Spline::OnToolCancel(const FPCGInteractiveToolWorkingDataContext& Context)
+{
+	// If the Spline Component is valid and isn't a generated component, we revert its state.
+	if (SplineComponent.IsValid() && GeneratedResources.GeneratedComponents.Contains(SplineComponent) == false)
+	{
+		if (OnToolStartSplineComponent)
+		{
+			// We restore the spline component to its previous state
+			UEngine::CopyPropertiesForUnrelatedObjects(OnToolStartSplineComponent, SplineComponent.Get());
+
+			SplineComponent->UpdateSpline();
+
+			OnToolStartSplineComponent->MarkAsGarbage();
+			OnToolStartSplineComponent = nullptr;
+		}
+	}
+
+	// This will delete the spline component if it was generated
+	Super::OnToolCancel(Context);
+}
+
+void FPCGInteractiveToolWorkingData_Spline::OnResetToolDataRequested(const FPCGInteractiveToolWorkingDataContext& Context)
+{
+	if (SplineComponent.IsValid())
+	{
+		SplineComponent->ClearSplinePoints();
+
+		FPropertyChangedEvent EmptyPropertyChangedEvent(nullptr);
+		FCoreUObjectDelegates::OnObjectPropertyChanged.Broadcast(SplineComponent.Get(), EmptyPropertyChangedEvent);
+	}
+}
+
+USplineComponent* FPCGInteractiveToolWorkingData_Spline::FindOrGenerateSplineComponent(const FPCGInteractiveToolWorkingDataContext& Context)
+{
+	USplineComponent* MatchingSplineComponent = FindMatchingSplineComponent(Context);
+
+	// If we didn't find any spline component to match with, we create it now
+	if (MatchingSplineComponent == nullptr)
+	{
+		MatchingSplineComponent = GenerateMatchingSplineComponent(Context);
+		GeneratedResources.GeneratedComponents.Add(MatchingSplineComponent);
+	}
+
+	ensureMsgf(MatchingSplineComponent != nullptr, TEXT("Matching Spline Component should always be found after initialization as it was either found or created."));
+	return MatchingSplineComponent;
+}
+
+USplineComponent* FPCGInteractiveToolWorkingData_Spline::FindMatchingSplineComponent(const FPCGInteractiveToolWorkingDataContext& Context)
+{
+	TArray<USplineComponent*> AllSplineComponents;
+	Context.OwningActor->GetComponents<USplineComponent>(AllSplineComponents, false);
+
+	if (!Context.DataInstanceIdentifier.IsNone())
+	{
+		// When the DataInstanceIdentifier is set, we only care about the splines with the given tag
+		AllSplineComponents.RemoveAll([DataInstanceIdentifier = Context.DataInstanceIdentifier](USplineComponent* Candidate)
+			{
+				return Candidate->ComponentHasTag(DataInstanceIdentifier) == false;
+			});
+	}
+
+	USplineComponent* MatchingSplineComponent = AllSplineComponents.IsEmpty() ? nullptr : AllSplineComponents[0];
+
+	if (AllSplineComponents.Num() > 1 && MatchingSplineComponent != nullptr)
+	{
+		UE_LOGFMT(LogPCG, Warning, "More than one spline component found. Choosing the first available: {0}", MatchingSplineComponent->GetName());
+	}
+
+	return MatchingSplineComponent;
+}
+
+USplineComponent* FPCGInteractiveToolWorkingData_Spline::GenerateMatchingSplineComponent(const FPCGInteractiveToolWorkingDataContext& Context)
+{
+	USplineComponent* MatchingSplineComponent = NewObject<USplineComponent>(Context.OwningActor.Get());
+	MatchingSplineComponent->SetMobility(Context.OwningActor->GetRootComponent()->GetMobility());
+	MatchingSplineComponent->SetupAttachment(Context.OwningActor->GetRootComponent());
+	// Transient until Applied. 
+	MatchingSplineComponent->SetFlags(RF_Transient);
+
+	if (Context.DataInstanceIdentifier.IsNone() == false)
+	{
+		MatchingSplineComponent->ComponentTags.Add(Context.DataInstanceIdentifier);
+	}
+
+	MatchingSplineComponent->ClearSplinePoints();
+	MatchingSplineComponent->bSplineHasBeenEdited = true;
+
+	// We set it to open or closed using OnInitialized based on the user property
+	MatchingSplineComponent->SetClosedLoop(false);
+
+	MatchingSplineComponent->RegisterComponent();
+	Context.OwningActor->AddInstanceComponent(MatchingSplineComponent);
+
+	return MatchingSplineComponent;
+}
+#endif // WITH_EDITOR
+
+void FPCGInteractiveToolWorkingData_Spline::InitializeRuntimeElementData(FPCGContext* Context) const
+{
+	Super::InitializeRuntimeElementData(Context);
+
+	TArray<FPCGTaggedData>& Outputs = Context->OutputData.TaggedData;
+
+	if (SplineComponent.IsValid())
+	{
+		FPCGTaggedData& Output = Outputs.Emplace_GetRef();
+		// We duplicate it so that PCG can take 'ownership'.
+		// If not, it can't process references and caches properly
+		// Might be relying on the fact it shouldn't access the same object over multiple generations
+		// And if we recreate it every time it gets properly managed
+		// If it gets dropped from cache, it's expected to be invalid
+		UPCGSplineData* SplineData = FPCGContext::NewObject_AnyThread<UPCGSplineData>(Context);
+		SplineData->Initialize(SplineComponent.Get());
+		Output.Data = SplineData;
+#if WITH_EDITOR
+		FPCGSelectionKey Key = FPCGSelectionKey::CreateFromObjectPtr(SplineComponent.Get());
+		FPCGDynamicTrackingHelper::AddSingleDynamicTrackingKey(Context, MoveTemp(Key), false);
+#endif
+	}
+}
+
+bool FPCGInteractiveToolWorkingData_Spline::IsValid() const
+{
+	return SplineComponent.IsValid();
+}
+
+#if WITH_EDITOR
+USplineComponent* FPCGInteractiveToolWorkingData_SplineSurface::FindMatchingSplineComponent(const FPCGInteractiveToolWorkingDataContext& Context)
+{
+	TArray<USplineComponent*> AllClosedSplineComponents;
+	Context.OwningActor->GetComponents<USplineComponent>(AllClosedSplineComponents, false);
+
+	// We remove all spline components that are open loops since this working data struct is only for spline surfaces
+	AllClosedSplineComponents.RemoveAll([](USplineComponent* Candidate)
+		{
+			return Candidate->IsClosedLoop() == false;
+		});
+
+	USplineComponent* MatchingSplineComponent = nullptr;
+	if (Context.DataInstanceIdentifier.IsNone())
+	{
+		if (AllClosedSplineComponents.Num() > 0)
+		{
+			MatchingSplineComponent = AllClosedSplineComponents[0];
+		}
+	}
+	else
+	{
+		// When the DataInstanceIdentifier is set, we only care about the closed splines with the given tag
+		AllClosedSplineComponents.RemoveAll([DataInstanceIdentifier = Context.DataInstanceIdentifier](USplineComponent* Candidate)
+			{
+				return Candidate->ComponentHasTag(DataInstanceIdentifier) == false;
+			});
+
+		// If we find a matching spline component with the specific tag, use i
+		if (AllClosedSplineComponents.Num() > 0)
+		{
+			MatchingSplineComponent = Cast<USplineComponent>(AllClosedSplineComponents[0]);
+		}
+	}
+
+	if (AllClosedSplineComponents.Num() > 1 && MatchingSplineComponent != nullptr)
+	{
+		UE_LOGFMT(LogPCG, Warning, "More than one spline component found. Choosing the first available: {0}", MatchingSplineComponent->GetName());
+	}
+
+	return MatchingSplineComponent;
+}
+
+USplineComponent* FPCGInteractiveToolWorkingData_SplineSurface::GenerateMatchingSplineComponent(const FPCGInteractiveToolWorkingDataContext& Context)
+{
+	USplineComponent* MatchingSplineComponent = NewObject<USplineComponent>(Context.OwningActor.Get());
+	MatchingSplineComponent->SetMobility(Context.OwningActor->GetRootComponent()->GetMobility());
+	MatchingSplineComponent->SetupAttachment(Context.OwningActor->GetRootComponent());
+	// Transient until Applied. 
+	MatchingSplineComponent->SetFlags(RF_Transient);
+
+	if (Context.DataInstanceIdentifier.IsNone() == false)
+	{
+		MatchingSplineComponent->ComponentTags.Add(Context.DataInstanceIdentifier);
+	}
+
+	MatchingSplineComponent->ClearSplinePoints();
+	MatchingSplineComponent->bSplineHasBeenEdited = true;
+	MatchingSplineComponent->SetClosedLoop(true);
+
+	MatchingSplineComponent->RegisterComponent();
+	Context.OwningActor->AddInstanceComponent(MatchingSplineComponent);
+
+	return MatchingSplineComponent;
+}
+#endif // WITH_EDITOR
\ No newline at end of file
diff --git a/Engine/Plugins/PCG/Source/PCG/Private/Data/Tool/PCGToolVolumeData.cpp b/Engine/Plugins/PCG/Source/PCG/Private/Data/Tool/PCGToolVolumeData.cpp
new file mode 100644
--- /dev/null
+++ b/Engine/Plugins/PCG/Source/PCG/Private/Data/Tool/PCGToolVolumeData.cpp
@@ -0,0 +1,141 @@
+// Copyright Epic Games, Inc. All Rights Reserved.
+
+#include "Data/Tool/PCGToolVolumeData.h"
+
+#include "PCGComponent.h"
+#include "PCGContext.h"
+
+#include "Components/BrushComponent.h"
+#include "Engine/Engine.h"
+#include "Engine/BrushBuilder.h"
+#include "Logging/StructuredLog.h"
+#include "GameFramework/Volume.h"
+
+#if WITH_EDITOR
+#include "ActorFactories/ActorFactory.h"
+#endif
+
+#if WITH_EDITOR
+void FPCGInteractiveToolWorkingData_Volume::InitializeInternal(const FPCGInteractiveToolWorkingDataContext& Context)
+{
+	Volume = Cast<AVolume>(Context.OwningActor.Get());
+	BrushComponent = FindOrCreateBrushComponent(Context);
+}
+
+void FPCGInteractiveToolWorkingData_Volume::OnToolStart(const FPCGInteractiveToolWorkingDataContext& Context)
+{
+	Super::OnToolStart(Context);
+
+	if (BrushComponent.IsValid())
+	{
+		OnToolStartBrushComponent = Cast<UBrushComponent>(StaticDuplicateObject(BrushComponent.Get(), GetTransientPackage()));
+	}
+
+	if (Volume.IsValid())
+	{
+		OnToolStartVolumeTransform = Volume->GetActorTransform();
+
+		if (Volume->BrushBuilder)
+		{
+			OnToolStartBrushBuilder = DuplicateObject<UBrushBuilder>(Volume->BrushBuilder, GetTransientPackage());
+		}
+	}
+}
+
+void FPCGInteractiveToolWorkingData_Volume::OnToolCancel(const FPCGInteractiveToolWorkingDataContext& Context)
+{
+	if (BrushComponent.IsValid() && GeneratedResources.GeneratedComponents.Contains(BrushComponent) == false && OnToolStartBrushComponent)
+	{
+		Volume->SetActorTransform(OnToolStartVolumeTransform);
+
+		if (OnToolStartBrushBuilder)
+		{
+			// Recreate the brush from the original actor
+			UActorFactory::CreateBrushForVolumeActor(Volume.Get(), OnToolStartBrushBuilder);
+			OnToolStartBrushBuilder->MarkAsGarbage();
+			OnToolStartBrushBuilder = nullptr;
+		}
+
+		// Restore the brush component to its previous state
+		UEngine::CopyPropertiesForUnrelatedObjects(OnToolStartBrushComponent, BrushComponent.Get());
+		OnToolStartBrushComponent->MarkAsGarbage();
+		OnToolStartBrushComponent = nullptr;
+	}
+
+	Super::OnToolCancel(Context);
+}
+
+UBrushComponent* FPCGInteractiveToolWorkingData_Volume::FindOrCreateBrushComponent(const FPCGInteractiveToolWorkingDataContext& Context)
+{
+	UBrushComponent* MatchingComponent = FindMatchingBrushComponent(Context);
+
+	if (!MatchingComponent)
+	{
+		MatchingComponent = GenerateMatchingBrushComponent(Context);
+		GeneratedResources.GeneratedComponents.Add(MatchingComponent);
+	}
+
+	ensureMsgf(MatchingComponent != nullptr, TEXT("Matching Brush Component should always be found after initialization as it was either found or created."));
+	return MatchingComponent;
+}
+
+UBrushComponent* FPCGInteractiveToolWorkingData_Volume::FindMatchingBrushComponent(const FPCGInteractiveToolWorkingDataContext& Context)
+{
+	TArray<UBrushComponent*> AllBrushComponents;
+	if (Context.OwningActor.IsValid())
+	{
+		Context.OwningActor->GetComponents<UBrushComponent>(AllBrushComponents, false);
+	}
+
+	UBrushComponent* MatchingComponent = nullptr;
+
+	if (!Context.DataInstanceIdentifier.IsNone())
+	{
+		AllBrushComponents.RemoveAll([DataInstanceIdentifier = Context.DataInstanceIdentifier](UBrushComponent* Candidate)
+			{
+				return !Candidate || !Candidate->ComponentHasTag(DataInstanceIdentifier);
+			});
+	}
+
+	MatchingComponent = AllBrushComponents.IsEmpty() ? nullptr : AllBrushComponents[0];
+
+	if (AllBrushComponents.Num() > 1 && MatchingComponent != nullptr)
+	{
+		UE_LOGFMT(LogPCG, Warning, "More than one brush component found. Choosing the first available: {0}", MatchingComponent->GetName());
+	}
+
+	return MatchingComponent;
+}
+
+UBrushComponent* FPCGInteractiveToolWorkingData_Volume::GenerateMatchingBrushComponent(const FPCGInteractiveToolWorkingDataContext& Context)
+{
+	if (!Context.OwningActor.IsValid())
+	{
+		return nullptr;
+	}
+
+	UBrushComponent* MatchingBrushComponent = NewObject<UBrushComponent>(Context.OwningActor.Get());
+	// ABrush setup
+	MatchingBrushComponent->Mobility = EComponentMobility::Static;
+	MatchingBrushComponent->SetCanEverAffectNavigation(false);
+	// AVolume setup
+	MatchingBrushComponent->AlwaysLoadOnClient = true;
+	MatchingBrushComponent->AlwaysLoadOnServer = true;
+	// APCGVolume setup
+	MatchingBrushComponent->SetCollisionObjectType(ECC_WorldStatic);
+	MatchingBrushComponent->SetCollisionResponseToAllChannels(ECR_Ignore);
+	MatchingBrushComponent->SetGenerateOverlapEvents(false);
+	// Transient until applied
+	MatchingBrushComponent->SetFlags(RF_Transient);
+
+	if (!Context.DataInstanceIdentifier.IsNone())
+	{
+		MatchingBrushComponent->ComponentTags.Add(Context.DataInstanceIdentifier);
+	}
+
+	MatchingBrushComponent->RegisterComponent();
+	Context.OwningActor->AddInstanceComponent(MatchingBrushComponent);
+
+	return MatchingBrushComponent;
+}
+#endif // WITH_EDITOR
\ No newline at end of file
diff --git a/Engine/Plugins/PCG/Source/PCG/Private/Elements/PCGDataFromTool.cpp b/Engine/Plugins/PCG/Source/PCG/Private/Elements/PCGDataFromTool.cpp
--- a/Engine/Plugins/PCG/Source/PCG/Private/Elements/PCGDataFromTool.cpp
+++ b/Engine/Plugins/PCG/Source/PCG/Private/Elements/PCGDataFromTool.cpp
@@ -6,6 +6,7 @@
 #include "PCGCustomVersion.h"
 #include "PCGSubsystem.h"
 #include "Data/PCGBasePointData.h"
+#include "Data/Tool/PCGToolBaseData.h"
 #include "Elements/PCGAddComponent.h"
 #include "Helpers/PCGDynamicTrackingHelpers.h"
 #include "Helpers/PCGHelpers.h"
diff --git a/Engine/Plugins/PCG/Source/PCG/Public/Data/PCGToolData.h b/Engine/Plugins/PCG/Source/PCG/Public/Data/PCGToolData.h
--- a/Engine/Plugins/PCG/Source/PCG/Public/Data/PCGToolData.h
+++ b/Engine/Plugins/PCG/Source/PCG/Public/Data/PCGToolData.h
@@ -2,256 +2,25 @@
 
 #pragma once
 
-#include "Data/PCGSplineData.h"
+#include "PCGModule.h"
+
+#include "StructUtils/InstancedStruct.h"
 #include "UObject/WeakObjectPtrTemplates.h"
 
 #include "PCGToolData.generated.h"
 
 class AActor;
-class AVolume;
 class USplineComponent;
 class UPCGPointArrayData;
 class UPCGComponent;
-class UBrushComponent;
-class UInteractiveTool;
-class UInteractiveToolPropertySet;
-class UBrushBuilder;
-
-/** A transient struct that contains resources generated during tool usage for a specific tool data struct. Used to finalize or clean up when a tool is applied or cancelled. */
-USTRUCT()
-struct FPCGInteractiveToolWorkingDataGeneratedResources
-{
-	GENERATED_BODY()
-
-	FPCGInteractiveToolWorkingDataGeneratedResources() = default;
-	virtual ~FPCGInteractiveToolWorkingDataGeneratedResources() = default;
-
-#if WITH_EDITORONLY_DATA
-	/** Finalizes generated resources (e.g. remove Transient flag). */
-	PCG_API void Finalize();
-
-	/** Remove and destroy generated resources. */
-	PCG_API void Cleanup();
-	
-	UPROPERTY(Transient)
-	TArray<TWeakObjectPtr<UActorComponent>> GeneratedComponents;
-#endif
-};
 
-/** A transient context struct that is passed in so the tool data can manage its data properly during tool usage. Used at editor-time only. */
 USTRUCT()
-struct FPCGInteractiveToolWorkingDataContext
-{
-	GENERATED_BODY()
-
-	FPCGInteractiveToolWorkingDataContext() = default;
-	virtual ~FPCGInteractiveToolWorkingDataContext() = default;
-
-#if WITH_EDITORONLY_DATA
-	PCG_API FPCGInteractiveToolWorkingDataContext(UPCGComponent* OwningComponent, UInteractiveTool* InteractiveTool, UInteractiveToolPropertySet* PropertySet, FName DataInstanceIdentifier);
-
-	/** The owning PCG Component. Should always be valid. Can be used as an Outer for new objects. */
-	UPROPERTY(Transient)
-	TWeakObjectPtr<UPCGComponent> OwningPCGComponent;
-
-	/** An optional DataInstanceIdentifier. Can be used to find or create a specific object, such as a Spline Component with a certain tag. */
-	UPROPERTY(Transient)
-	FName DataInstanceIdentifier;
-
-	/** The owning actor. The owning PCG component's outer. Should always be valid. Convenience accessor. */
-	UPROPERTY(Transient)
-	TWeakObjectPtr<AActor> OwningActor;
-
-	/** The interactive tool currently editing the tool data. */
-	UPROPERTY(Transient)
-	TWeakObjectPtr<UInteractiveTool> InteractiveTool;
-
-	/** The PCG Setting properties of the Interactive Tool. */
-	UPROPERTY(Transient)
-	TWeakObjectPtr<UInteractiveToolPropertySet> PCGSettings;
-#endif
-};
-
-/** The base struct for working data for a tool. This data is created by the tool, stored on the PCG Component,
- * then read back by the tool on subsequent tool invocations. */
-USTRUCT()
-struct FPCGInteractiveToolWorkingData
-{
-	GENERATED_BODY()
-
-	FPCGInteractiveToolWorkingData() = default;
-	virtual ~FPCGInteractiveToolWorkingData() = default;
-
-	PCG_API bool IsInitialized() const;
-
-	PCG_API virtual bool IsValid() const;
-
-	/** Construct and add pcg data to the runtime context.
-	 * This will translate data from this struct into UPCGData to be processed by the graph.
-	 * Add dynamic tracking keys in here if you want to refresh when a change to your source data happens. */
-	PCG_API virtual void InitializeRuntimeElementData(FPCGContext* InContext) const;
-	
-#if WITH_EDITORONLY_DATA
-public:
-	DECLARE_DELEGATE_OneParam(FOnInitialized, FPCGInteractiveToolWorkingData*)
-	
-	/** Initializes data used by the tool working data. Initialize is called once when the working data is created. */
-	PCG_API void Initialize(const FPCGInteractiveToolWorkingDataContext& Context);
-	
-	/** OnToolStart is called every time the tool is started up, or when the data is created.
-	 * Can be used to save the initial state of persistent data to restore to on Cancel, or to Modify existing data. */
-	PCG_API virtual void OnToolStart(const FPCGInteractiveToolWorkingDataContext& Context);
-	
-	/** Called when the tool is applied. Finalize changes here. */
-	PCG_API virtual void OnToolApply(const FPCGInteractiveToolWorkingDataContext& Context);
-	
-	/** Called when the tool is canceled. Revert changes you have made here. */
-	PCG_API virtual void OnToolCancel(const FPCGInteractiveToolWorkingDataContext& Context);
-
-	/** Called when reset data was specifically requested, for example by the user. */
-	virtual void OnResetToolDataRequested(const FPCGInteractiveToolWorkingDataContext& Context) {}
-	
-	PCG_API FOnInitialized& OnInitialized() { return OnInitializedDelegate; }
-private:
-	/** Implement your initialization logic here. Called by Initialize. */
-	PCG_API virtual void InitializeInternal(const FPCGInteractiveToolWorkingDataContext& Context);
-protected:
-	/** Each working data struct keeps track of its generated resources. They have to be reverted during OnToolCancel. */
-	UPROPERTY(Transient)
-	FPCGInteractiveToolWorkingDataGeneratedResources GeneratedResources;
-
-	/** OnInitialized delegate to allow the settings class to do more work. */
-	FOnInitialized OnInitializedDelegate;
-#endif
-	
-private:
-	UPROPERTY()
-	bool bInitialized = false;
-};
-
-/** Tool Data that stores points. */
-USTRUCT(BlueprintType, DisplayName="Point Array Tool Data")
-struct FPCGInteractiveToolWorkingData_PointArrayData : public FPCGInteractiveToolWorkingData
+struct FPCGInteractiveToolWorkingBaseData
 {
 	GENERATED_BODY()
-	
-public:
-	FPCGInteractiveToolWorkingData_PointArrayData() = default;
-
-	PCG_API virtual void InitializeRuntimeElementData(FPCGContext* InContext) const override;
-	
-	PCG_API virtual bool IsValid() const override;
-	
-	PCG_API UPCGPointArrayData* GetPointArrayData() const;
-
-#if WITH_EDITORONLY_DATA
-	PCG_API virtual void InitializeInternal(const FPCGInteractiveToolWorkingDataContext& Context) override;
-	
-	PCG_API virtual void OnToolStart(const FPCGInteractiveToolWorkingDataContext& Context) override;
-	
-	PCG_API virtual void OnToolCancel(const FPCGInteractiveToolWorkingDataContext& Context) override;
-
-	PCG_API virtual void OnResetToolDataRequested(const FPCGInteractiveToolWorkingDataContext& Context) override;
-#endif
-	
-protected:
-	UPROPERTY(VisibleAnywhere, BlueprintReadOnly, Category="Working Data")
-	TObjectPtr<UPCGPointArrayData> GeneratedPointData;
-
-	/** A transient copy of the point array data, created on tool start/data access. We revert back to these points on Cancel. */
-	UPROPERTY(Transient)
-	TObjectPtr<UPCGPointArrayData> OnToolStartPointArray;
-};
-
-/** The working data for splines. */
-USTRUCT(BlueprintType, DisplayName="Spline Tool Data")
-struct FPCGInteractiveToolWorkingData_Spline : public FPCGInteractiveToolWorkingData
-{
-	GENERATED_BODY()
-public:
-	FPCGInteractiveToolWorkingData_Spline() = default;
-	virtual ~FPCGInteractiveToolWorkingData_Spline() override = default;
-	
-	PCG_API virtual void InitializeRuntimeElementData(FPCGContext* InContext) const override;
-	
-	PCG_API virtual bool IsValid() const override;
-	
-	PCG_API USplineComponent* GetSplineComponent() const;
-
-#if WITH_EDITOR
-public:
-	virtual void OnToolStart(const FPCGInteractiveToolWorkingDataContext& Context) override;
-	virtual void OnToolCancel(const FPCGInteractiveToolWorkingDataContext& Context) override;
-	virtual void OnResetToolDataRequested(const FPCGInteractiveToolWorkingDataContext& Context) override;
-private:
-	virtual void InitializeInternal(const FPCGInteractiveToolWorkingDataContext& Context) override;
-	
-private:
-	USplineComponent* FindOrGenerateSplineComponent(const FPCGInteractiveToolWorkingDataContext& Context);
-	
-	virtual USplineComponent* FindMatchingSplineComponent(const FPCGInteractiveToolWorkingDataContext& Context);
-	virtual USplineComponent* GenerateMatchingSplineComponent(const FPCGInteractiveToolWorkingDataContext& Context);
-#endif
-	
-protected:
-	UPROPERTY(VisibleAnywhere, BlueprintReadOnly, Category="Working Data")
-	TWeakObjectPtr<USplineComponent> SplineComponent;
-
-	UPROPERTY(Transient)
-	TObjectPtr<USplineComponent> OnToolStartSplineComponent;
-};
-
-USTRUCT(BlueprintType, DisplayName="Spline Surface Tool Data")
-struct FPCGInteractiveToolWorkingData_SplineSurface : public FPCGInteractiveToolWorkingData_Spline
-{
-	GENERATED_BODY()
-	
-	FPCGInteractiveToolWorkingData_SplineSurface() = default;
-	virtual ~FPCGInteractiveToolWorkingData_SplineSurface() override = default;
-
-#if WITH_EDITOR
-private:
-	/** We look for and generate a closed spline component. */
-	virtual USplineComponent* FindMatchingSplineComponent(const FPCGInteractiveToolWorkingDataContext& Context) override;
-	virtual USplineComponent* GenerateMatchingSplineComponent(const FPCGInteractiveToolWorkingDataContext& Context) override;
-#endif
-};
-
-USTRUCT(BlueprintType, DisplayName="Volume Tool Data")
-struct FPCGInteractiveToolWorkingData_Volume : public FPCGInteractiveToolWorkingData
-{
-	GENERATED_BODY()
-public:
-	UPROPERTY(VisibleAnywhere, BlueprintReadOnly, Category = "Working Data")
-	TWeakObjectPtr<AVolume> Volume = nullptr;
-
-	UPROPERTY(VisibleAnywhere, BlueprintReadOnly, Category = "Working Data")
-	TWeakObjectPtr<UBrushComponent> BrushComponent = nullptr;
-
-#if WITH_EDITOR
-public:
-	virtual void OnToolStart(const FPCGInteractiveToolWorkingDataContext& Context) override;
-	virtual void OnToolCancel(const FPCGInteractiveToolWorkingDataContext& Context) override;
-
-private:
-	virtual void InitializeInternal(const FPCGInteractiveToolWorkingDataContext& Context) override;
-
-	UBrushComponent* FindOrCreateBrushComponent(const FPCGInteractiveToolWorkingDataContext& Context);
-	virtual UBrushComponent* FindMatchingBrushComponent(const FPCGInteractiveToolWorkingDataContext& Context);
-	virtual UBrushComponent* GenerateMatchingBrushComponent(const FPCGInteractiveToolWorkingDataContext& Context);
-#endif
-
-protected:
-#if WITH_EDITORONLY_DATA
-	UPROPERTY(Transient)
-	TObjectPtr<UBrushComponent> OnToolStartBrushComponent = nullptr;
-
-	UPROPERTY(Transient)
-	TObjectPtr<UBrushBuilder> OnToolStartBrushBuilder = nullptr;
 
-	UPROPERTY(Transient)
-	FTransform OnToolStartVolumeTransform;
-#endif
+	FPCGInteractiveToolWorkingBaseData() = default;
+	virtual ~FPCGInteractiveToolWorkingBaseData() = default;
 };
 
 /** The tool data stored on the graph itself. Informs the UI. */
@@ -282,5 +51,5 @@ struct FPCGInteractiveToolWorkingDataContainer
 	virtual ~FPCGInteractiveToolWorkingDataContainer() = default;
 	
 	UPROPERTY(BlueprintReadOnly, VisibleAnywhere, Category = "Tool Data")
-	TMap<FName, TInstancedStruct<FPCGInteractiveToolWorkingData>> ToolData;
+	TMap<FName, TInstancedStruct<FPCGInteractiveToolWorkingBaseData>> ToolData;
 };
\ No newline at end of file
diff --git a/Engine/Plugins/PCG/Source/PCGEditor/Private/Details/PCGInteractiveToolSettingsDetails.cpp b/Engine/Plugins/PCG/Source/PCGEditor/Private/Details/PCGInteractiveToolSettingsDetails.cpp
--- a/Engine/Plugins/PCG/Source/PCGEditor/Private/Details/PCGInteractiveToolSettingsDetails.cpp
+++ b/Engine/Plugins/PCG/Source/PCGEditor/Private/Details/PCGInteractiveToolSettingsDetails.cpp
@@ -4,10 +4,12 @@
 
 #include "PCGEditorModule.h"
 #include "PCGGraph.h"
+#include "Data/Tool/PCGToolBaseData.h"
+#include "EditorMode/Tools/PCGInteractiveToolSettings.h"
+
 #include "DetailCategoryBuilder.h"
 #include "DetailLayoutBuilder.h"
 #include "DetailWidgetRow.h"
-#include "EditorMode/Tools/PCGInteractiveToolSettings.h"
 #include "Widgets/Layout/SWidgetSwitcher.h"
 #include "Framework/Application/SlateApplication.h"
 
diff --git a/Engine/Plugins/PCG/Source/PCGEditor/Private/EditorMode/PCGEdMode.h b/Engine/Plugins/PCG/Source/PCGEditor/Private/EditorMode/PCGEdMode.h
new file mode 100644
--- /dev/null
+++ b/Engine/Plugins/PCG/Source/PCGEditor/Private/EditorMode/PCGEdMode.h
@@ -0,0 +1,70 @@
+// Copyright Epic Games, Inc. All Rights Reserved.
+
+#pragma once
+
+#include "EdMode.h"
+#include "InteractiveToolManager.h"
+
+#include "PCGEdMode.generated.h"
+
+class UPCGInteractiveToolSettings;
+
+/** The unique state for the editor to be in, specific to PCG's toolkits. */
+UCLASS(Transient)
+class UPCGEditorMode : public UEdMode
+{
+	GENERATED_BODY()
+
+public:
+	/** Unique name for PCG Editor Mode */
+	PCGEDITOR_API const static FEditorModeID EM_PCGEditorModeId;
+
+	UPCGEditorMode();
+
+	static void RegisterEditorMode();
+	static void UnregisterEditorMode();
+
+	//~Begin UEdMode interface
+	virtual bool ProcessEditDelete() override;
+	virtual bool ProcessEditCut() override;
+	virtual bool CanAutoSave() const override;
+	/** Override the viewport pivot, if needed. */
+	virtual bool GetPivotForOrbit(FVector& OutPivot) const override;
+	/** Focus event when the Editor enters PCG Editor Mode. */
+	virtual void Enter() override;
+	/** Focus event when the Editor leaves PCG Editor Mode. */
+	virtual void Exit() override;
+	/** Used to restrict starting a tool, if a certain state is required beforehand. */
+	virtual bool ShouldToolStartBeAllowed(const FString& ToolIdentifier) const override;
+	virtual bool HasCustomViewportFocus() const override;
+	virtual FBox ComputeCustomViewportFocus() const override;
+	virtual bool IsSelectionAllowed(AActor* InActor, bool bInSelection) const override;
+	//~End UEdMode interface
+
+protected:
+	virtual void CreateToolkit() override;
+
+	virtual void OnToolStarted(UInteractiveToolManager* Manager, UInteractiveTool* Tool) override;
+	virtual void OnToolEnded(UInteractiveToolManager* Manager, UInteractiveTool* Tool) override;
+	virtual void OnToolPostBuild(UInteractiveToolManager* InToolManager, EToolSide InSide, UInteractiveTool* InBuiltTool, UInteractiveToolBuilder* InToolBuilder, const FToolBuilderState& ToolState);
+
+	void OnEditorClosed() const;
+	void AcceptActiveToolActionOrTool() const;
+	void CancelActiveToolActionOrTool() const;
+
+	void ConfigureRealTimeViewportsOverride(bool bEnable);
+
+private:
+	// @todo_pcg: Activate when custom tools are supported.
+	// void RegisterCustomTool(
+	// 	TSharedPtr<FUICommandInfo> UICommand,
+	// 	FString ToolIdentifier,
+	// 	UInteractiveToolBuilder* Builder,
+	// 	const TFunction<bool(UInteractiveToolManager*, EToolSide)>& ExecuteAction,
+	// 	const TFunction<bool(UInteractiveToolManager*, EToolSide)>& CanExecuteAction,
+	// 	const TFunction<bool(UInteractiveToolManager*, EToolSide)>& IsActionChecked) const;
+
+	const UPCGInteractiveToolSettings* GetCurrentToolSettings() const;
+
+	bool bIsToolActive = false;
+};
diff --git a/Engine/Plugins/PCG/Source/PCGEditor/Private/EditorMode/PCGEdModeCommands.h b/Engine/Plugins/PCG/Source/PCGEditor/Private/EditorMode/PCGEdModeCommands.h
new file mode 100644
--- /dev/null
+++ b/Engine/Plugins/PCG/Source/PCGEditor/Private/EditorMode/PCGEdModeCommands.h
@@ -0,0 +1,96 @@
+// Copyright Epic Games, Inc. All Rights Reserved.
+
+#pragma once
+
+#include "Framework/Commands/Commands.h"
+#include "Tools/InteractiveToolsCommands.h"
+
+/** General actions related to the PCG Editor Mode. */
+class PCGEDITOR_API FPCGEditorModeCommands : public TCommands<FPCGEditorModeCommands>
+{
+public:
+	FPCGEditorModeCommands();
+
+	virtual void RegisterCommands() override;
+
+	// @todo_pcg: Any file menu, hotkeys, etc related to the editor mode in general can go here.
+};
+
+/** Actions related to the PCG Editor Mode Toolkit Palette. */
+class PCGEDITOR_API FPCGEditorModePaletteCommands : public TCommands<FPCGEditorModePaletteCommands>
+{
+public:
+	FPCGEditorModePaletteCommands();
+
+	/*** PCG EDITOR CONTEXTS **/
+	TSharedPtr<FUICommandInfo> LoadLineContextPalette;
+	TSharedPtr<FUICommandInfo> LoadPaintContextPalette;
+	TSharedPtr<FUICommandInfo> LoadVolumeContextPalette;
+
+	// @todo_pcg: Selection/Filtering actions.
+
+	virtual void RegisterCommands() override;
+
+	// @todo_pcg: Allow extension for custom tools. See Modeling Tools as an example.
+};
+
+/** Tool specific actions. */
+class FPCGEditorModeToolCommands : public TInteractiveToolCommands<FPCGEditorModeToolCommands>
+{
+public:
+	FPCGEditorModeToolCommands();
+
+	virtual void RegisterCommands() override;
+
+	/** Register all Tool command sets. Call this in module startup */
+	static void RegisterAllToolCommands();
+
+	/** Unregister all Tool command sets. Call this from module shutdown. */
+	static void UnregisterAllToolCommands();
+
+	/** Needed to retrieve the CDOs to call GetActionSet directly on the tool. */
+	virtual void GetToolDefaultObjectList(TArray<UInteractiveTool*>& ToolCDOs) override;
+
+	// @todo_pcg: Implement if needed for custom tool commands.
+	/**
+	 * Add the tool-specific commands from an interactive tool to a specified UICommandList 
+	 * Call this when the active tool changes (eg on ToolManager.OnToolStarted / OnToolEnded)
+	 */
+	// static void AddToolCommandBinding(UInteractiveTool* Tool, TSharedPtr<FUICommandList> UICommandList);
+	// static void RemoveToolCommandBinding(UInteractiveTool* Tool, TSharedPtr<FUICommandList> UICommandList);
+
+	/*** PCG EDITOR MODE TOOLS - Line Context **/
+	TSharedPtr<FUICommandInfo> EnableDrawSplineTool;
+	TSharedPtr<FUICommandInfo> EnableDrawSurfaceTool;
+	TArray<TSharedPtr<FUICommandInfo>> GetLineContextToolCommands();
+
+	/*** PCG EDITOR MODE TOOLS - Paint Context **/
+	TSharedPtr<FUICommandInfo> EnablePaintTool;
+	TArray<TSharedPtr<FUICommandInfo>> GetPaintContextToolCommands();
+
+	/*** PCG EDITOR MODE TOOLS - Volume Context **/
+	TSharedPtr<FUICommandInfo> EnableVolumeTool;
+	TArray<TSharedPtr<FUICommandInfo>> GetVolumeContextToolCommands();
+};
+
+/**
+ * Macro for the class declaration for each tool as a TInteractiveToolCommands class
+ * @todo_pcg: This is basically done elsewhere for other toolkits, but can probably be modernized using C++ Concepts 
+ * @param ToolCommandClassName The tool's class name
+ */
+#define DECLARE_TOOL_COMMANDS(ToolCommandClassName) \
+class ToolCommandClassName : public TInteractiveToolCommands<ToolCommandClassName> \
+{\
+public:\
+ToolCommandClassName();\
+\
+virtual void GetToolDefaultObjectList(TArray<UInteractiveTool*>& ToolCDOs) override;\
+};\
+
+// Define all tools commands here
+DECLARE_TOOL_COMMANDS(FPCGDrawSplineCommands);
+DECLARE_TOOL_COMMANDS(FPCGDrawSurfaceCommands);
+DECLARE_TOOL_COMMANDS(FPCGPaintCommands);
+DECLARE_TOOL_COMMANDS(FPCGVolumeCommands);
+
+#undef DECLARE_TOOL_COMMANDS
diff --git a/Engine/Plugins/PCG/Source/PCGEditor/Private/EditorMode/PCGEdModeHelpers.h b/Engine/Plugins/PCG/Source/PCGEditor/Private/EditorMode/PCGEdModeHelpers.h
new file mode 100644
--- /dev/null
+++ b/Engine/Plugins/PCG/Source/PCGEditor/Private/EditorMode/PCGEdModeHelpers.h
@@ -0,0 +1,20 @@
+// Copyright Epic Games, Inc. All Rights Reserved.
+
+#pragma once
+
+#include "Framework/Notifications/NotificationManager.h"
+#include "Widgets/Notifications/SNotificationList.h"
+
+namespace PCG::EditorMode
+{
+	static constexpr float DefaultNotificationDuration = 2.5f;
+
+	inline void DispatchEditorNotification(const FText& Text, const FText& SubText = FText{}, const float Duration = DefaultNotificationDuration)
+	{
+		FNotificationInfo NotificationInfo(std::move(Text));
+		NotificationInfo.ExpireDuration = Duration;
+		NotificationInfo.bFireAndForget = true;
+		NotificationInfo.SubText = std::move(SubText);
+		FSlateNotificationManager::Get().AddNotification(NotificationInfo);
+	}
+}
diff --git a/Engine/Plugins/PCG/Source/PCGEditor/Private/EditorMode/PCGEdModeSettings.h b/Engine/Plugins/PCG/Source/PCGEditor/Private/EditorMode/PCGEdModeSettings.h
new file mode 100644
--- /dev/null
+++ b/Engine/Plugins/PCG/Source/PCGEditor/Private/EditorMode/PCGEdModeSettings.h
@@ -0,0 +1,76 @@
+// Copyright Epic Games, Inc. All Rights Reserved.
+
+#pragma once
+
+#include "Engine/DeveloperSettings.h"
+
+#include "PCGEdModeSettings.generated.h"
+
+class UPCGInteractiveToolSettings;
+class UPCGGraph;
+
+USTRUCT()
+struct FPCGPerInteractiveToolSettingSettings
+{
+	GENERATED_BODY()
+	
+	UPROPERTY(EditAnywhere, Config, Category = "Default")
+	TSoftClassPtr<UPCGInteractiveToolSettings> SettingsClass;
+
+	// @todo_pcg GetAssetFilter only works on UClasses due to need for UFUNCTION. Will reenable this via details customization instead.
+	//UPROPERTY(EditAnywhere, Config, Category = "Default", meta=(GetAssetFilter = "FilterBySettingsClass"))
+	UPROPERTY(EditAnywhere, Config, Category = "Default")
+	TSoftObjectPtr<UPCGGraph> DefaultGraph;
+};
+
+/** User preferences related to the PCG Editor Mode. */
+UCLASS(Config = EditorPerProjectUserSettings, DisplayName = "PCG Editor Mode Settings")
+class UPCGEditorModeSettings : public UDeveloperSettings
+{
+	GENERATED_BODY()
+
+public:
+	UPCGEditorModeSettings(const FObjectInitializer& ObjectInitializer)
+		: UDeveloperSettings(ObjectInitializer)
+	{}
+
+	virtual FName GetContainerName() const override { return FName("Editor"); }
+	virtual FName GetCategoryName() const override { return FName("Plugins"); }
+	virtual FName GetSectionName() const override { return FName("PCGEditorMode"); }
+
+	/** Editor Mode user configurations. */
+	UPROPERTY(EditAnywhere, Config, Category = "General", meta = (DisplayName = "Enable the PCG Editor Mode Toolkit (requires restart)", ConfigRestartRequired = true))
+	bool bEnableEditorMode = true;
+
+	/** The rate at which PCG Graphs will be refreshed while actively editing with a tool. */
+	UPROPERTY(EditAnywhere, Category = "Component", meta = (
+		ForceUnits = "Seconds",
+		ConsoleVariable = "pcg.editormode.GraphRefreshRate"))
+	float GraphRefreshRate = 0.5f;
+
+	/** Hide the tool buttons from the UI when a tool is actively being used. */
+	UPROPERTY(EditAnywhere, Config, Category = "Tools", meta = (
+		EditCondition = "bEnableEditorMode",
+		EditConditionHides,
+		ConsoleVariable = "pcg.editormode.HideToolButtonsDuringActiveTool"))
+	bool bHideToolButtonsDuringActiveTool = false;
+
+	/** Enable an editor toast issued when an Editor Mode tool has an error. */
+	UPROPERTY(EditAnywhere, Config, Category = "Tools", meta = (
+		EditCondition = "bEnableEditorMode",
+		EditConditionHides,
+		ConsoleVariable = "pcg.editormode.ShowEditorToastOnToolErrors"))
+	bool bShowEditorToastOnToolErrors = true;
+
+	UPROPERTY(EditAnywhere, Category = "Tools|Output", meta = (EditCondition = "bEnableEditorMode", EditConditionHides))
+	FName DefaultNewActorName = "PCG Tool Actor";
+
+	UPROPERTY(EditAnywhere, Category = "Tools|Output", meta = (EditCondition = "bEnableEditorMode", EditConditionHides))
+	FName DefaultNewPCGComponentName = "PCG Tool Component";
+
+	UPROPERTY(EditAnywhere, Category = "Tools|Output|Spline", meta = (EditCondition = "bEnableEditorMode", EditConditionHides))
+	FName DefaultNewSplineComponentName = "PCG Tool Spline Component";
+
+	UPROPERTY(EditAnywhere, Config, Category = "Tools")
+	TArray<FPCGPerInteractiveToolSettingSettings> InteractiveToolSettings;
+};
diff --git a/Engine/Plugins/PCG/Source/PCGEditor/Private/EditorMode/PCGEdModeStyle.h b/Engine/Plugins/PCG/Source/PCGEditor/Private/EditorMode/PCGEdModeStyle.h
new file mode 100644
--- /dev/null
+++ b/Engine/Plugins/PCG/Source/PCGEditor/Private/EditorMode/PCGEdModeStyle.h
@@ -0,0 +1,20 @@
+// Copyright Epic Games, Inc. All Rights Reserved.
+
+#pragma once
+
+#include "Styling/SlateStyle.h"
+
+struct FSlateIcon;
+
+/** Style details specific to the PCG Editor Mode. */
+class FPCGEditorModeStyle : public FSlateStyleSet
+{
+public:
+	static void Register();
+	static void Unregister();
+
+	static const FPCGEditorModeStyle& Get();
+
+private:
+	FPCGEditorModeStyle();
+};
diff --git a/Engine/Plugins/PCG/Source/PCGEditor/Private/EditorMode/PCGEdModeToolkit.h b/Engine/Plugins/PCG/Source/PCGEditor/Private/EditorMode/PCGEdModeToolkit.h
new file mode 100644
--- /dev/null
+++ b/Engine/Plugins/PCG/Source/PCGEditor/Private/EditorMode/PCGEdModeToolkit.h
@@ -0,0 +1,85 @@
+// Copyright Epic Games, Inc. All Rights Reserved.
+
+#pragma once
+
+#include "InteractiveTool.h"
+#include "StatusBarSubsystem.h"
+#include "Toolkits/BaseToolkit.h"
+
+#include "Widgets/SPCGEdModeAcceptToolOverlay.h"
+
+class UPCGInteractiveTool;
+
+/** The basic toolkit needed to show the docked editor mode panel and set up the tool shelves and details views. */
+class FPCGEditorModeToolkit : public FModeToolkit
+{
+public:
+	// ~Begin FModeToolkit interface
+	virtual void Init(const TSharedPtr<IToolkitHost>& InitToolkitHost, TWeakObjectPtr<UEdMode> OwningMode) override;
+	virtual FName GetToolkitFName() const override;
+	virtual FText GetBaseToolkitName() const override;
+
+	/** Returns the PCG Editor Mode specific tabs in the mode toolbar */
+	virtual void GetToolPaletteNames(TArray<FName>& PaletteNames) const override;
+	virtual FText GetToolPaletteDisplayName(FName PaletteName) const override;
+	virtual void OnToolPaletteChanged(FName PaletteName) override;
+	virtual bool HasIntegratedToolPalettes() const override { return false; }
+	virtual void BuildToolPalette(FName Palette, class FToolBarBuilder& ToolbarBuilder) override;
+
+	virtual FText GetActiveToolDisplayName() const override;
+	virtual FText GetActiveToolMessage() const override;
+
+	virtual void OnToolStarted(UInteractiveToolManager* Manager, UInteractiveTool* Tool) override;
+	virtual void OnToolEnded(UInteractiveToolManager* Manager, UInteractiveTool* Tool) override;
+
+	// Keep track of the widget for customization.
+	virtual TSharedPtr<SWidget> GetInlineContent() const override { return ToolkitWidget; }
+	// ~End FModeToolkit interface
+
+	virtual void PostToolNotification(const FText& Message);
+	virtual void ClearToolNotification();
+
+	virtual void PostToolWarning(const FText& Message);
+	virtual void ClearToolWarning();
+
+	virtual void ShowModeWarnings();
+
+	/*** UI Events ***/
+	void OnActiveViewportChanged(TSharedPtr<IAssetViewport>, TSharedPtr<IAssetViewport>) const;
+
+	/*** Accept/Cancel Window ***/
+	const FSlateBrush* GetActiveToolIcon() const;
+	bool IsAcceptButtonEnabled() const;
+	bool IsCancelButtonEnabled() const;
+	EVisibility GetAcceptButtonVisibility() const;
+	EVisibility GetCancelButtonVisibility() const;
+	FReply OnAcceptButtonClicked() const;
+	FReply OnCancelButtonClicked() const;
+
+private:
+	void RegisterToolkitCallbacks();
+	void UnregisterToolkitCallbacks();
+
+	void RegisterToolCallbacks(UInteractiveTool* CurrentTool);
+	void UnregisterToolCallbacks(UInteractiveTool* CurrentTool) const;
+
+	void OnEditorModeSettingsChanged(UObject* InObject, FPropertyChangedEvent& InChangedEvent);
+
+	void UpdateActiveTool();
+	void InvalidateCachedDetailPanelState(UObject* ChangedObject) const;
+
+	/** A utility function to register the tool palettes with the ToolkitBuilder. */
+	void RegisterPalettes();
+	/** A utility function to create additional widgets for the Toolkit. */
+	void CreateWidgets();
+
+	/** UI overlay in the viewport to accept or cancel the current tool. */
+	TSharedPtr<SPCGEdModeAcceptToolOverlay> AcceptToolWidget;
+
+	/** Currently only used to cancel the current tool. */
+	FDelegateHandle ActivePaletteChangedHandle;
+
+	FStatusBarMessageHandle ActiveToolMessageHandle;
+	FText ActiveWarning{};
+	FInteractiveToolInfo CachedActiveToolInfo;
+};
diff --git a/Engine/Plugins/PCG/Source/PCGEditor/Private/EditorMode/Tools/Gizmos/PCGGizmos.h b/Engine/Plugins/PCG/Source/PCGEditor/Private/EditorMode/Tools/Gizmos/PCGGizmos.h
new file mode 100644
--- /dev/null
+++ b/Engine/Plugins/PCG/Source/PCGEditor/Private/EditorMode/Tools/Gizmos/PCGGizmos.h
@@ -0,0 +1,70 @@
+﻿// Copyright Epic Games, Inc. All Rights Reserved.
+
+#pragma once
+
+#include "BaseGizmos/BrushStampIndicator.h"
+#include "InteractiveTool.h"
+#include "InteractiveGizmoBuilder.h"
+#include "PCGGizmos.generated.h"
+
+class UPreviewMesh;
+class UMaterialInstanceDynamic;
+
+namespace UE::PCG::EditorMode::Gizmos
+{
+	namespace Constants
+	{
+		const FString BrushIdentifier_VolumetricBox = "VolumetricBoxBrushGizmoType";
+	}
+	
+	UMaterialInstanceDynamic* GetDefaultVolumeBrushMaterial(UInteractiveToolManager* ToolManager);
+}
+
+UENUM()
+enum class EPCGBrushMode
+{
+	Sphere,
+	Box
+};
+
+UCLASS(MinimalAPI)
+class UVolumetricBoxBrushStampIndicatorBuilder : public UInteractiveGizmoBuilder
+{
+	GENERATED_BODY()
+public:
+	virtual UInteractiveGizmo* BuildGizmo(const FToolBuilderState& SceneState) const override;
+};
+
+UCLASS(Transient, MinimalAPI)
+class UPCGVolumetricBrushStampIndicator : public UBrushStampIndicator
+{
+	GENERATED_BODY()
+
+public:
+	virtual void Setup() override;
+	virtual void Render(IToolsContextRenderAPI* RenderAPI) override;
+	virtual void Shutdown() override;
+
+	void SetVolumetricBrushVisible(bool bInVisible);
+
+	bool IsVolumetricBrushVisible() const;
+
+	void MakeBrushIndicatorMesh(UInteractiveTool* ParentTool, UWorld* World);
+
+	UPreviewMesh* GetPreviewMesh() const;
+
+	void AddYawRotation(float AddYawRotation);
+
+	void SetBrushMode(EPCGBrushMode InBrushMode);
+protected:
+	void GenerateMesh(EPCGBrushMode InBrushMode);
+private:
+	UPROPERTY()
+	TObjectPtr<UPreviewMesh> PreviewMesh;
+
+	UPROPERTY()
+	bool bDrawExtents = true;
+
+	UPROPERTY()
+	EPCGBrushMode BrushMode = EPCGBrushMode::Sphere;
+};
\ No newline at end of file
diff --git a/Engine/Plugins/PCG/Source/PCGEditor/Private/EditorMode/Tools/Helpers/PCGEdModeActorHelpers.h b/Engine/Plugins/PCG/Source/PCGEditor/Private/EditorMode/Tools/Helpers/PCGEdModeActorHelpers.h
new file mode 100644
--- /dev/null
+++ b/Engine/Plugins/PCG/Source/PCGEditor/Private/EditorMode/Tools/Helpers/PCGEdModeActorHelpers.h
@@ -0,0 +1,33 @@
+// Copyright Epic Games, Inc. All Rights Reserved.
+
+#pragma once
+
+#include "Helpers/PCGActorHelpers.h"
+
+namespace UE::PCG::EditorMode
+{
+	namespace Tags
+	{
+		const FLazyName ToolGeneratedTag = "PCG Tool";
+	}
+
+	namespace Actor
+	{
+		/** Spawn an actor for the toolkit with spawn parameters. */
+		AActor* Spawn(const UPCGActorHelpers::FSpawnDefaultActorParams& SpawnParams);
+		/** Spawn an actor for the toolkit with explicit parameters. */
+		AActor* Spawn(UWorld* World, FName ActorName, const FTransform& Transform = FTransform::Identity);
+		/** Helper to spawn a temporary working actor for a tool. */
+		AActor* SpawnWorking(UWorld* World, const FTransform& Transform = FTransform::Identity, TSubclassOf<AActor> ActorClass = AActor::StaticClass());
+	}
+
+	namespace PCGComponent
+	{
+		/** Find a PCG Component on an actor if it exists. */
+		UPCGComponent* Find(const AActor* Owner, const FName ToolTag);
+		/** Spawn a PCG Component on an actor. */
+		UPCGComponent* Create(AActor* Owner, FName ComponentName);
+		/** Helper to spawn a temporary working PCG component for a tool. */
+		UPCGComponent* CreateWorking(AActor* Owner);
+	}
+}
diff --git a/Engine/Plugins/PCG/Source/PCGEditor/Private/EditorMode/Tools/Helpers/PCGEdModeEditorUtilities.h b/Engine/Plugins/PCG/Source/PCGEditor/Private/EditorMode/Tools/Helpers/PCGEdModeEditorUtilities.h
new file mode 100644
--- /dev/null
+++ b/Engine/Plugins/PCG/Source/PCGEditor/Private/EditorMode/Tools/Helpers/PCGEdModeEditorUtilities.h
@@ -0,0 +1,24 @@
+// Copyright Epic Games, Inc. All Rights Reserved.
+
+#pragma once
+
+#include "Data/PCGToolData.h"
+
+#include "AssetRegistry/AssetData.h"
+#include "Templates/SharedPointer.h"
+
+class FUICommandInfo;
+class FBindingContext;
+ 
+namespace UE::PCG::EditorMode::Utility
+{
+	struct FPCGGraphToolEditorData
+	{
+		FAssetData AssetData;
+		FPCGGraphToolData GraphToolData;
+	};
+	
+	TArray<FPCGGraphToolEditorData> GetGraphToolsWithToolTag(FName ToolTag);
+	bool DoesPCGGraphHaveToolTag(FName ToolTag, const FAssetData& AssetData);
+	TArray<TSharedPtr<FUICommandInfo>> GetOrCreateUICommandsWithToolTag(TSharedRef<FBindingContext> BindingContext, FName ToolTag);
+}
diff --git a/Engine/Plugins/PCG/Source/PCGEditor/Private/EditorMode/Tools/Helpers/PCGEdModeSceneQueryHelpers.h b/Engine/Plugins/PCG/Source/PCGEditor/Private/EditorMode/Tools/Helpers/PCGEdModeSceneQueryHelpers.h
new file mode 100644
--- /dev/null
+++ b/Engine/Plugins/PCG/Source/PCGEditor/Private/EditorMode/Tools/Helpers/PCGEdModeSceneQueryHelpers.h
@@ -0,0 +1,145 @@
+// Copyright Epic Games, Inc. All Rights Reserved.
+
+#pragma once
+
+#include "InputState.h"
+#include "Math/Vector.h"
+#include "Misc/Optional.h"
+#include "StructUtils/InstancedStruct.h"
+#include "Templates/SubclassOf.h"
+#include "GameFramework/Actor.h"
+#include "PCGEdModeSceneQueryHelpers.generated.h"
+
+struct FViewCameraState;
+struct FHitResult;
+class ALandscape;
+
+USTRUCT()
+struct FPCGRaycastFilterRule
+{
+	GENERATED_BODY()
+
+	PCGEDITOR_API FPCGRaycastFilterRule();
+	PCGEDITOR_API virtual ~FPCGRaycastFilterRule();
+	
+	/** If this returns false, the raycast hit is filtered out.
+	 *  If it returns true, this rule has determined is valid under this given rule.
+	 *  If it returns an empty TOptional, this rule doesn't apply to the hit.
+	 */
+	virtual TOptional<bool> IsRaycastHitValid(const FHitResult& HitResult) const { return {}; }
+
+	/** Whether this rule is enabled or not. Not exposed as editing is done via custom UI. */
+	UPROPERTY()
+	bool bEnabled = true;
+};
+
+USTRUCT(DisplayName="Landscape")
+struct FPCGRaycastFilterRule_Landscape : public FPCGRaycastFilterRule
+{
+	GENERATED_BODY()
+
+	virtual TOptional<bool> IsRaycastHitValid(const FHitResult& HitResult) const override;
+};
+
+USTRUCT(DisplayName="Meshes")
+struct FPCGRaycastFilterRule_Meshes : public FPCGRaycastFilterRule
+{
+	GENERATED_BODY()
+
+	virtual TOptional<bool> IsRaycastHitValid(const FHitResult& HitResult) const override;
+};
+
+USTRUCT(DisplayName="Ignore PCG Components")
+struct FPCGRaycastFilterRule_IgnorePCGGeneratedComponents : public FPCGRaycastFilterRule
+{
+	GENERATED_BODY()
+
+	virtual TOptional<bool> IsRaycastHitValid(const FHitResult& HitResult) const override;
+};
+
+USTRUCT(DisplayName="Constrain to Actor")
+struct FPCGRaycastFilterRule_ConstrainToActor : public FPCGRaycastFilterRule
+{
+	GENERATED_BODY()
+
+	virtual TOptional<bool> IsRaycastHitValid(const FHitResult& HitResult) const override;
+
+	UPROPERTY(EditAnywhere, Category= "Raycast")
+	TObjectPtr<AActor> SelectedActor;
+};
+
+/** The Ray will only be valid when the hit actor belongs to any of the specified classes. */
+USTRUCT(DisplayName="Allowed Classes")
+struct FPCGRaycastFilterRule_AllowedClasses : public FPCGRaycastFilterRule
+{
+	GENERATED_BODY()
+
+	virtual TOptional<bool> IsRaycastHitValid(const FHitResult& HitResult) const override;
+
+	UPROPERTY(EditAnywhere, Category= "Raycast")
+	TSet<TSubclassOf<AActor>> AllowedClasses;
+};
+
+/** A collection struct used for details customization purposes. This allows us to add the collection to any PCGInteractiveToolSettings class where required. */
+USTRUCT()
+struct FPCGRaycastFilterRuleCollection
+{
+	GENERATED_BODY()
+
+	TOptional<bool> IsRaycastHitValid(const FHitResult& RaycastHit) const;
+
+	/** The raycast rules for a given tool. */
+	UPROPERTY(EditAnywhere, Category = "Raycast", meta=(ExcludeBaseStruct))
+	TArray<TInstancedStruct<FPCGRaycastFilterRule>> RaycastRules;
+};
+
+namespace UE::PCG::EditorMode::Scene
+{
+	/** Custom parameters for UI ray casts from the current view into the scene. */
+	struct FViewRayParams
+	{
+		const UWorld* World = nullptr;
+		double Distance = UE_DOUBLE_BIG_NUMBER;
+		const TArray<const UPrimitiveComponent*>* ComponentsToIgnore = nullptr;
+		const TArray<const UPrimitiveComponent*>* AllowedInvisibleComponents = nullptr;
+		const FPCGRaycastFilterRuleCollection* FilterRuleCollection = nullptr;
+	};
+
+	/** The resulting data from a UI ray cast hit. */
+	struct FViewHitResult
+	{
+		FVector ImpactPosition;
+		FVector ImpactNormal;
+		double Depth;
+
+		FInputRayHit ToInputRayHit() const
+		{
+			FInputRayHit OutRayHit;
+			OutRayHit.bHit = true;
+			OutRayHit.bHasHitNormal = true;
+			OutRayHit.HitNormal = ImpactNormal;
+			OutRayHit.HitDepth = Depth;
+			OutRayHit.HitObject = nullptr;
+			OutRayHit.HitOwner = nullptr;
+
+			return OutRayHit;
+		}
+	};
+
+	/** The hit actor or primitive component is visible. */
+	bool IsHitTargetVisible(const FHitResult& HitResult);
+
+	/** Cast a ray and if a hit is recorded, return the hit result nearest object. */
+	TOptional<FHitResult> TraceToNearestObject(const FRay& WorldRay, const FViewRayParams& Params);
+	TOptional<FHitResult> TraceToNearestObject(const FInputDeviceRay& Ray, const FViewRayParams& Params);
+
+	/** Cast a ray from the view and if a hit is recorded, return a condensed view-related hit result. */
+	TOptional<FViewHitResult> ViewportRay(const FRay& WorldRay, const FViewRayParams& Params);
+	TOptional<FViewHitResult> ViewportRay(const FInputDeviceRay& Ray, const FViewRayParams& Params);
+
+	/** Cast a ray from the view to a plane in space. If a hit is recorded, return a condensed view-related hit result. */
+	TOptional<FViewHitResult> ViewportRayAgainstPlane(const FInputDeviceRay& Ray, const FViewRayParams& Params, const FPlane& Plane);
+
+	/** Cast a ray from the view to a plane parallel to the user view. If a hit is recorded, return a condensed view-related hit result. */
+	TOptional<FViewHitResult> ViewportRayAgainstCameraPlane(const FInputDeviceRay& Ray, const FViewRayParams& Params, const FViewCameraState& CameraState);
+}
diff --git a/Engine/Plugins/PCG/Source/PCGEditor/Private/EditorMode/Tools/Helpers/PCGEdModeSceneVisualizeHelpers.h b/Engine/Plugins/PCG/Source/PCGEditor/Private/EditorMode/Tools/Helpers/PCGEdModeSceneVisualizeHelpers.h
new file mode 100644
--- /dev/null
+++ b/Engine/Plugins/PCG/Source/PCGEditor/Private/EditorMode/Tools/Helpers/PCGEdModeSceneVisualizeHelpers.h
@@ -0,0 +1,17 @@
+// Copyright Epic Games, Inc. All Rights Reserved.
+
+#pragma once
+
+#include "CoreMinimal.h"
+
+class IToolsContextRenderAPI;
+class USplineComponent;
+
+namespace PCG::EditorMode::Scene::Visualize
+{
+	// Visualize a spline within the editor viewport.
+	void Spline(const USplineComponent& Spline, IToolsContextRenderAPI& RenderAPI, int32 LatestIndex = INDEX_NONE, bool bShouldVisualizeScale = false, bool bShouldVisualizeTangents = true);
+
+	// Visualize a spline tangent within the editor viewport.
+	void SplineTangent(const USplineComponent& Spline, IToolsContextRenderAPI& RenderAPI, const int32 SplinePointIndex);
+}
diff --git a/Engine/Plugins/PCG/Source/PCGEditor/Private/EditorMode/Tools/Helpers/PCGEdModeSplineHelpers.h b/Engine/Plugins/PCG/Source/PCGEditor/Private/EditorMode/Tools/Helpers/PCGEdModeSplineHelpers.h
new file mode 100644
--- /dev/null
+++ b/Engine/Plugins/PCG/Source/PCGEditor/Private/EditorMode/Tools/Helpers/PCGEdModeSplineHelpers.h
@@ -0,0 +1,23 @@
+// Copyright Epic Games, Inc. All Rights Reserved.
+
+#pragma once
+
+#include "CoreMinimal.h"
+
+class USplineComponent;
+class AActor;
+
+namespace UE::PCG::EditorMode::Spline
+{	
+	USplineComponent* CreateComponent(
+		AActor* Actor,
+		FName ComponentName,
+		const USplineComponent* TemplateComponent = nullptr,
+		const bool bTransactional = false);
+
+	USplineComponent* CreateWorkingComponent(AActor* Actor);
+
+	USplineComponent* Find(const AActor* Actor, const int32 TargetIndex = INDEX_NONE);
+
+	void CopySplinePoints(const USplineComponent& Source, USplineComponent& Destination, const bool bTransactional = false);
+}
diff --git a/Engine/Plugins/PCG/Source/PCGEditor/Private/EditorMode/Tools/Line/PCGDrawSplineTool.cpp b/Engine/Plugins/PCG/Source/PCGEditor/Private/EditorMode/Tools/Line/PCGDrawSplineTool.cpp
--- a/Engine/Plugins/PCG/Source/PCGEditor/Private/EditorMode/Tools/Line/PCGDrawSplineTool.cpp
+++ b/Engine/Plugins/PCG/Source/PCGEditor/Private/EditorMode/Tools/Line/PCGDrawSplineTool.cpp
@@ -2,12 +2,13 @@
 
 #include "EditorMode/Tools/Line/PCGDrawSplineTool.h"
 
-#include "InteractiveToolManager.h"
 #include "PCGComponent.h"
+#include "Data/Tool/PCGToolSplineData.h"
 #include "EditorMode/PCGEdModeSettings.h"
 #include "EditorMode/PCGEdModeStyle.h"
 #include "EditorMode/Tools/Helpers/PCGEdModeSceneVisualizeHelpers.h"
 
+#include "InteractiveToolManager.h"
 #include "BaseBehaviors/SingleClickOrDragBehavior.h"
 #include "Components/SplineComponent.h"
 #include "Settings/LevelEditorViewportSettings.h"
diff --git a/Engine/Plugins/PCG/Source/PCGEditor/Private/EditorMode/Tools/Line/PCGDrawSplineTool.h b/Engine/Plugins/PCG/Source/PCGEditor/Private/EditorMode/Tools/Line/PCGDrawSplineTool.h
new file mode 100644
--- /dev/null
+++ b/Engine/Plugins/PCG/Source/PCGEditor/Private/EditorMode/Tools/Line/PCGDrawSplineTool.h
@@ -0,0 +1,176 @@
+// Copyright Epic Games, Inc. All Rights Reserved.
+
+#pragma once
+
+#include "EditorMode/Tools/PCGInteractiveToolBuilder.h"
+#include "EditorMode/Tools/PCGInteractiveToolCommon.h"
+#include "EditorMode/Tools/PCGInteractiveToolSettings.h"
+#include "EditorMode/Tools/Helpers/PCGEdModeSceneQueryHelpers.h"
+
+#include "BaseBehaviors/BehaviorTargetInterfaces.h"
+#include "Components/SplineComponent.h"
+
+#include "PCGDrawSplineTool.generated.h"
+
+class USplineComponent;
+class USingleClickOrDragInputBehavior;
+
+UENUM(BlueprintType)
+enum class EPCGDrawSplineDrawMode : uint8
+{
+	TangentDrag UMETA(ToolTip = "Click to place a point and then drag to set its tangent. Clicking without dragging will create sharp corners."),
+	ClickAutoTangent UMETA(ToolTip = "Click and drag new points, with the tangent set automatically."),
+	FreeDraw UMETA(ToolTip = "Drag to place multiple points, with spacing controlled by Min Point Spacing.")
+};
+
+UCLASS(Abstract, BlueprintType)
+class PCGEDITOR_API UPCGInteractiveToolSettings_SplineBase : public UPCGInteractiveToolSettings
+{
+	GENERATED_BODY()
+
+public:
+	UPCGInteractiveToolSettings_SplineBase();
+	
+	USplineComponent* GetWorkingSplineComponent(FName DataInstanceIdentifier) const;
+	TArray<USplineComponent*> GetWorkingSplineComponents() const;
+	
+	UPROPERTY(EditAnywhere, Category = "Spline")
+	EPCGDrawSplineDrawMode DrawMode = EPCGDrawSplineDrawMode::FreeDraw;
+	
+	/** Minimum world distance to cover while free drawing before a point will be generated. */
+	UPROPERTY(EditAnywhere, Category = "Spline", meta = (
+			ClampMin = 0,
+			UIMax = 1000,
+			EditCondition = "DrawMode == EPCGDrawSplineDrawMode::FreeDraw",
+			EditConditionHides))
+	double MinPointSpacing = 200.0;
+
+	UPROPERTY(EditAnywhere, Category = "Spline", meta = (TransientToolProperty))
+	FName SplineName;
+
+	/** The raycast rules that determine whether a raycast is valid or not.
+	 *  Used to determine different raycast targets such as landscapes, meshes, selected actor etc.	 */
+	UPROPERTY(EditAnywhere, Category= "Raycast")
+	FPCGRaycastFilterRuleCollection RaycastRuleCollection;
+
+	/** The settings to manipulate raycast results, such as normals or offsets of a spline point. */
+	UPROPERTY(EditAnywhere, Category = "Raycast", meta = (ShowOnlyInnerProperties))
+	FPCGToolRaycastSettings Settings{};
+};
+
+UCLASS(BlueprintType)
+class PCGEDITOR_API UPCGInteractiveToolSettings_Spline : public UPCGInteractiveToolSettings_SplineBase
+{
+	GENERATED_BODY()
+
+protected:
+	virtual const UScriptStruct* GetWorkingDataType() const override;
+	virtual FName GetToolTag() const override;
+	virtual void OnWorkingDataInitialized(FPCGInteractiveToolWorkingData* WorkingData) const override;
+	virtual void OnPropertyModified(UObject* Object, FProperty* Property) override;
+public:
+	UPROPERTY(EditAnywhere, Category = "Spline")
+	bool bClosedSpline = false;
+
+	static FName StaticGetToolTag();
+};
+
+
+UCLASS(BlueprintType)
+class PCGEDITOR_API UPCGInteractiveToolSettings_SplineSurface : public UPCGInteractiveToolSettings_SplineBase
+{
+	GENERATED_BODY()
+
+protected:
+	virtual const UScriptStruct* GetWorkingDataType() const override;
+	virtual FName GetToolTag() const override;
+
+public:
+	static FName StaticGetToolTag();
+};
+
+UCLASS(Abstract, Experimental)
+class PCGEDITOR_API UPCGDrawSplineToolBase : public UInteractiveTool, public IClickBehaviorTarget, public IClickDragBehaviorTarget
+{
+	GENERATED_BODY()
+
+public:
+	// ~Begin UInteractiveTool interface
+	virtual void Setup() override;
+	virtual void Shutdown(EToolShutdownType ShutdownType) override;
+	virtual void OnTick(float DeltaTime) override;
+
+	virtual bool HasCancel() const override { return true; }
+	virtual bool HasAccept() const override { return true; }
+	virtual bool CanAccept() const override;
+	virtual void Render(IToolsContextRenderAPI* RenderAPI) override;
+	// ~End UInteractiveTool interface
+
+	// ~Begin IClickBehaviorTarget interface
+	virtual FInputRayHit IsHitByClick(const FInputDeviceRay& ClickPosition) override;
+	virtual void OnClicked(const FInputDeviceRay& ClickPosition) override;
+	// ~End IClickBehaviorTarget interface
+
+	// ~Begin IClickDragBehaviorTarget interface
+	virtual FInputRayHit CanBeginClickDragSequence(const FInputDeviceRay& DragPosition) override;
+	virtual void OnClickPress(const FInputDeviceRay& ClickPosition) override;
+	virtual void OnClickDrag(const FInputDeviceRay& DragPosition) override;
+	virtual void OnClickRelease(const FInputDeviceRay& ReleasePosition) override;
+	virtual void OnTerminateDragSequence() override;
+	// ~End IClickDragBehaviorTarget interface
+	
+	UPCGInteractiveToolSettings_SplineBase* GetSplineSettings() const;
+private:
+	void AddSplinePoint(const FVector& HitLocation, const FVector& HitNormal) const;
+	FVector GetNewPointNormal(const FVector& HitNormal) const;
+
+	UPROPERTY(Instanced)
+	TObjectPtr<USingleClickOrDragInputBehavior> ClickOrDragBehavior = nullptr;
+	
+	UE::PCG::EditorMode::Scene::FViewRayParams CachedRayParams{};
+	bool bDrawTangentForPreviousPoint = false;
+	bool bFreeDrawPlacedPreviewPoint = false;
+	int32 FreeDrawStrokeStartIndex = 0;
+
+protected:
+	double CachedMinPointSpacingSquared = 0;
+};
+
+UCLASS(Experimental)
+class PCGEDITOR_API UPCGDrawSplineTool : public UPCGDrawSplineToolBase
+{
+	GENERATED_BODY()
+
+public:
+	virtual void Setup() override;
+public:
+	UPROPERTY(Transient)
+	TObjectPtr<UPCGInteractiveToolSettings_Spline> ToolSettings;
+};
+
+UCLASS(Experimental)
+class PCGEDITOR_API UPCGDrawSplineSurfaceTool : public UPCGDrawSplineToolBase
+{
+	GENERATED_BODY()
+
+public:
+	virtual void Setup() override;
+public:
+	UPROPERTY(Transient)
+	TObjectPtr<UPCGInteractiveToolSettings_SplineSurface> ToolSettings;
+};
+
+UCLASS(Transient)
+class PCGEDITOR_API UPCGDrawSplineToolBuilder : public UPCGInteractiveToolBuilder
+{
+	GENERATED_BODY()
+
+public:
+	virtual bool CanBuildTool(const FToolBuilderState& SceneState) const override;
+	virtual UInteractiveTool* BuildTool(const FToolBuilderState& SceneState) const override;
+
+	void SetSplineToolClass(TSubclassOf<UPCGDrawSplineToolBase> InClass);
+private:
+	UPROPERTY()
+	TSubclassOf<UPCGDrawSplineToolBase> SplineToolClass;
+};
diff --git a/Engine/Plugins/PCG/Source/PCGEditor/Private/EditorMode/Tools/PCGInteractiveToolBuilder.h b/Engine/Plugins/PCG/Source/PCGEditor/Private/EditorMode/Tools/PCGInteractiveToolBuilder.h
new file mode 100644
--- /dev/null
+++ b/Engine/Plugins/PCG/Source/PCGEditor/Private/EditorMode/Tools/PCGInteractiveToolBuilder.h
@@ -0,0 +1,49 @@
+// Copyright Epic Games, Inc. All Rights Reserved.
+
+#pragma once
+
+#include "InteractiveToolBuilder.h"
+#include "PCGGraph.h"
+
+#include "PCGInteractiveToolBuilder.generated.h"
+
+namespace UE::PCG::EditorMode::Tool
+{
+	bool BuildTool(UInteractiveTool* InTool);
+}
+
+UCLASS()
+class PCGEDITOR_API UPCGInteractiveToolBuilder : public UInteractiveToolBuilder
+{
+	GENERATED_BODY()
+
+public:
+	// ~Begin UInteractiveToolBuilder interface
+	virtual bool CanBuildTool(const FToolBuilderState& SceneState) const override;
+	virtual UInteractiveTool* BuildTool(const FToolBuilderState& SceneState) const override;
+	virtual void PostSetupTool(UInteractiveTool* Tool, const FToolBuilderState& SceneState) const override;
+	// ~End UInteractiveToolBuilder interface
+
+	void SetToolGraph(UPCGGraph* InGraph) { ToolGraph = InGraph; }
+	UPCGGraph* GetToolGraph() const { return ToolGraph; }
+private:
+	UPROPERTY()
+	TObjectPtr<UPCGGraph> ToolGraph;
+};
+
+UCLASS()
+class PCGEDITOR_API UPCGInteractiveToolWithToolTargetsBuilder : public UInteractiveToolWithToolTargetsBuilder
+{
+	GENERATED_BODY()
+
+public:
+	// ~Begin UInteractiveToolBuilder interface
+	virtual void PostSetupTool(UInteractiveTool* Tool, const FToolBuilderState& SceneState) const override;
+	// ~End UInteractiveToolBuilder interface
+
+	void SetToolGraph(UPCGGraph* InGraph) { ToolGraph = InGraph; }
+	UPCGGraph* GetToolGraph() const { return ToolGraph; }
+private:
+	UPROPERTY()
+	TObjectPtr<UPCGGraph> ToolGraph;
+};
\ No newline at end of file
diff --git a/Engine/Plugins/PCG/Source/PCGEditor/Private/EditorMode/Tools/PCGInteractiveToolCommon.h b/Engine/Plugins/PCG/Source/PCGEditor/Private/EditorMode/Tools/PCGInteractiveToolCommon.h
new file mode 100644
--- /dev/null
+++ b/Engine/Plugins/PCG/Source/PCGEditor/Private/EditorMode/Tools/PCGInteractiveToolCommon.h
@@ -0,0 +1,57 @@
+// Copyright Epic Games, Inc. All Rights Reserved.
+
+#pragma once
+
+#include "Helpers/PCGEdModeSceneQueryHelpers.h"
+
+#include "PCGInteractiveToolCommon.generated.h"
+
+UENUM()
+enum class EPCGToolDrawTargetOffset : uint8
+{
+	// Drawn points will be at the hit location.
+	HitLocation,
+	// Drawn points will be at the hit location with a distance offset.
+	DistanceOffset,
+	// Drawn points will be at the hit location with a custom vector offset.
+	ExplicitOffset
+};
+
+UENUM()
+enum class EPCGToolDrawTargetNormal : uint8
+{
+	// Drawn points will adopt the impact normal.
+	HitNormal,
+	// Drawn points will adopt the world up normal.
+	WorldUp,
+	// Align to the previous spline point's up vector.
+	AlignToPrevious,
+	// Drawn points will adopt an explicit normal.
+	Explicit
+};
+
+USTRUCT()
+struct PCGEDITOR_API FPCGToolRaycastSettings
+{
+	GENERATED_BODY()
+
+	/** How to choose the direction to offset points from the drawn surface. */
+	UPROPERTY(EditAnywhere, Category = "Draw Target")
+	EPCGToolDrawTargetOffset OffsetMode = EPCGToolDrawTargetOffset::HitLocation;
+
+	/** How far to offset drawn points from the clicked surface, along the normal. */
+	UPROPERTY(EditAnywhere, Category = "Draw Target", DisplayName = "Offset", meta = (UIMin = 0, UIMax = 100, EditCondition = "OffsetMode == EPCGToolDrawTargetOffset::DistanceOffset", EditConditionHides))
+	double OffsetDistance = 100.0;
+
+	/** Manual vector to offset drawn points from the clicked surface. */
+	UPROPERTY(EditAnywhere, Category = "Draw Target", DisplayName = "Offset", meta = (UIMin = 0, UIMax = 100, EditCondition = "OffsetMode == EPCGToolDrawTargetOffset::ExplicitOffset", EditConditionHides))
+	FVector OffsetVector = FVector(0, 0, 100);
+
+	/** How to choose the direction to offset points from the drawn surface. */
+	UPROPERTY(EditAnywhere, Category = "Draw Target")
+	EPCGToolDrawTargetNormal NormalMode = EPCGToolDrawTargetNormal::HitNormal;
+
+	/** Manual vector to offset drawn points from the clicked surface. */
+	UPROPERTY(EditAnywhere, Category = "Draw Target", meta = (UIMin = 0, UIMax = 100, EditCondition = "NormalMode == EPCGToolDrawTargetNormal::Explicit", EditConditionHides))
+	FVector Normal = FVector::ZAxisVector;
+};
diff --git a/Engine/Plugins/PCG/Source/PCGEditor/Private/EditorMode/Tools/PCGInteractiveToolSettings.cpp b/Engine/Plugins/PCG/Source/PCGEditor/Private/EditorMode/Tools/PCGInteractiveToolSettings.cpp
--- a/Engine/Plugins/PCG/Source/PCGEditor/Private/EditorMode/Tools/PCGInteractiveToolSettings.cpp
+++ b/Engine/Plugins/PCG/Source/PCGEditor/Private/EditorMode/Tools/PCGInteractiveToolSettings.cpp
@@ -4,6 +4,7 @@
 #include "PCGComponent.h"
 #include "PCGEditorModule.h"
 #include "PCGEditorSettings.h"
+#include "Data/Tool/PCGToolBaseData.h"
 #include "EditorMode/PCGEdModeHelpers.h"
 #include "EditorMode/PCGEdModeSettings.h"
 #include "EditorMode/Tools/Helpers/PCGEdModeEditorUtilities.h"
diff --git a/Engine/Plugins/PCG/Source/PCGEditor/Private/EditorMode/Tools/PCGInteractiveToolSettings.h b/Engine/Plugins/PCG/Source/PCGEditor/Private/EditorMode/Tools/PCGInteractiveToolSettings.h
new file mode 100644
--- /dev/null
+++ b/Engine/Plugins/PCG/Source/PCGEditor/Private/EditorMode/Tools/PCGInteractiveToolSettings.h
@@ -0,0 +1,251 @@
+// Copyright Epic Games, Inc. All Rights Reserved.
+
+#pragma once
+
+#include "PCGGraph.h"
+#include "PCGInteractiveToolCommon.h"
+#include "Data/Tool/PCGToolBaseData.h"
+#include "Helpers/PCGEdModeActorHelpers.h"
+
+#include "InteractiveTool.h"
+#include "NativeGameplayTags.h"
+#include "PCGComponent.h"
+#include "TickableEditorObject.h"
+#include "BaseBehaviors/BehaviorTargetInterfaces.h"
+#include "PCGInteractiveToolSettings.generated.h"
+
+namespace UE::PCG::EditorMode::Tool
+{
+	PCGEDITOR_API void Shutdown(UInteractiveTool* InTool, EToolShutdownType ShutdownType);
+	
+	PCGEDITOR_API extern TAutoConsoleVariable<bool> CVarCommitPropertyChangedEventOnTimer;
+}
+
+UCLASS(Abstract, Blueprintable)
+class PCGEDITOR_API UPCGInteractiveToolSettings : public UInteractiveToolPropertySet, public FTickableEditorObject, public IKeyInputBehaviorTarget
+{
+	GENERATED_BODY()
+
+public:
+	UPCGInteractiveToolSettings();
+	virtual ~UPCGInteractiveToolSettings() override {}
+
+	virtual TValueOrError<bool, FText> Initialize(UInteractiveTool* InTool);
+	virtual void Apply(UInteractiveTool* OwningTool);
+	virtual void Cancel(UInteractiveTool* OwningTool);
+	virtual void Shutdown();
+
+	//~Begin interface UInteractiveToolPropertySet
+	virtual void RestoreProperties(UInteractiveTool* RestoreToTool, const FString& CacheIdentifier = TEXT("")) override;
+	//~End interface UInteractiveToolPropertySet
+
+	/** By default, when a tool is active, we can only select the Edit Actor in the Outliner. */
+	virtual bool IsSelectionAllowed(AActor* InActor, bool bInSelection) const;
+
+	/** Whether tool data can be reset. Controls visibility of the Reset button in the UI. */
+	virtual bool CanResetToolData(FName DataInstanceIdentifier);
+	
+	/** Register property watchers and callbacks in here. */
+	virtual void RegisterPropertyWatchers();
+	
+	/** Initializes or finds the working data associated with the data instance identifier.
+	 * If NAME_None, the ToolTag is used as identifier. */
+	void InitializeWorkingData(FName DataInstanceIdentifier = NAME_None);
+	
+	UPCGGraphInterface* GetToolGraph() const { return ToolGraph; }
+	void SetToolGraph(UPCGGraphInterface* InGraph);
+
+	void SetDataInstance(FName DataInstanceIdentifier);
+
+	virtual void SetActorClassToSpawn(TSubclassOf<AActor> Class);
+	
+	/** A soft object path to the graph asset to load by default for this tool. */
+	virtual FSoftObjectPath GetDefaultGraph() const;
+	
+	/** Tool tag for use with filtering compatible graphs. */
+	virtual FName GetToolTag() const PURE_VIRTUAL(UPCGInteractiveToolSettings::GetToolTag, return NAME_None;);
+
+	/** The working data type for this settings class. */
+	virtual const UScriptStruct* GetWorkingDataType() const PURE_VIRTUAL(UPCGInteractiveToolSettings::GetWorkingDataType, return FPCGInteractiveToolWorkingData::StaticStruct(););
+	
+	/** The actor to instantiate for edit-purposes when we aren't editing an existing actor. 
+	 * If your target actor has any special requirements, create an actor for that purpose and return it in your settings. */
+	virtual TSubclassOf<AActor> GetWorkingActorClass() const;
+
+	/** Filter details view graph selection. */
+	UFUNCTION()
+	virtual bool GraphAssetFilter(const FAssetData& AssetData) const;
+	
+	UPCGComponent* GetWorkingPCGComponent() const;
+	UPCGGraphInstance* GetGraphInstance() const;
+
+	UFUNCTION()
+	TArray<FName> GetDataInstanceNamesForGraph() const;
+
+	UFUNCTION()
+	TArray<FName> GetValidDataInstanceNamesForGraph() const;
+	FName GetDefaultDataInstanceName() const;
+	
+	template<typename T = AActor>
+	T* GetTypedWorkingActor() const
+	{
+		static_assert(TIsDerivedFrom<T, AActor>::IsDerived, "T needs to be derived from AActor.");
+		
+		return Cast<T>(EditActor.Get());
+	}
+
+	FName GetWorkingDataIdentifier(FName DataInstanceIdentifier = NAME_None) const;
+	
+	template<typename T = FPCGInteractiveToolWorkingData>
+	T* GetMutableTypedWorkingData(FName DataInstanceIdentifier = NAME_None)
+	{
+		static_assert(TIsDerivedFrom<T, FPCGInteractiveToolWorkingData>::IsDerived, "T needs to be derived from FPCGInteractiveToolWorkingData.");
+
+		FName DataIdentifier = GetWorkingDataIdentifier(DataInstanceIdentifier);
+		
+		UPCGComponent* PCGComponent = GetWorkingPCGComponent();
+		if (PCGComponent && PCGComponent->ToolDataContainer.ToolData.Contains(DataIdentifier))
+		{
+			return PCGComponent->ToolDataContainer.ToolData[DataIdentifier].GetMutablePtr<T>(); 
+		}
+
+		return nullptr;
+	}
+	
+	template<typename T = FPCGInteractiveToolWorkingData>
+	const T* GetTypedWorkingData(FName DataInstanceIdentifier = NAME_None) const
+	{
+		static_assert(TIsDerivedFrom<T, FPCGInteractiveToolWorkingData>::IsDerived, "T needs to be derived from FPCGInteractiveToolWorkingData.");
+
+		FName WorkingDataIdentifier = GetWorkingDataIdentifier(DataInstanceIdentifier);
+
+		UPCGComponent* PCGComponent = GetWorkingPCGComponent();
+		if(PCGComponent && PCGComponent->ToolDataContainer.ToolData.Contains(WorkingDataIdentifier))
+		{
+			return PCGComponent->ToolDataContainer.ToolData[WorkingDataIdentifier].GetPtr<T>(); 
+		}
+		
+		return nullptr;
+	}
+
+	template<typename T = FPCGInteractiveToolWorkingData>
+	TMap<FName, T*> GetMutableTypedWorkingDataMap() const
+	{
+		static_assert(TIsDerivedFrom<T, FPCGInteractiveToolWorkingData>::IsDerived, "T needs to be derived from FPCGInteractiveToolWorkingData.");
+
+		TMap<FName, T*> Result;
+
+		if (UPCGComponent* PCGComponent = GetWorkingPCGComponent())
+		{
+			for(const FName& DataInstanceName : GetDataInstanceNamesForGraph())
+			{
+				FName WorkingDataIdentifier = GetWorkingDataIdentifier(DataInstanceName);
+			
+				if(PCGComponent->ToolDataContainer.ToolData.Contains(WorkingDataIdentifier))
+				{
+					Result.Add(WorkingDataIdentifier, PCGComponent->ToolDataContainer.ToolData[WorkingDataIdentifier].GetMutablePtr<T>()); 
+				}
+			}
+		}
+
+		return Result;
+	}
+
+	template<typename T = FPCGInteractiveToolWorkingData>
+	TMap<FName, const T*> GetTypedWorkingDataMap() const
+	{
+		static_assert(TIsDerivedFrom<T, FPCGInteractiveToolWorkingData>::IsDerived, "T needs to be derived from FPCGInteractiveToolWorkingData.");
+
+		TMap<FName, const T*> Result;
+
+		if (UPCGComponent* PCGComponent = GetWorkingPCGComponent())
+		{
+			for(const FName& DataInstanceName : GetDataInstanceNamesForGraph())
+			{
+				FName WorkingDataIdentifier = GetWorkingDataIdentifier(DataInstanceName);
+			
+				if(PCGComponent->ToolDataContainer.ToolData.Contains(WorkingDataIdentifier))
+				{
+					Result.Add(WorkingDataIdentifier, PCGComponent->ToolDataContainer.ToolData[WorkingDataIdentifier].GetPtr<T>()); 
+				}
+			}
+		}
+
+		return Result;
+	}
+protected:
+	/** Try to load the defaults (graph, actor class, ...) into the tool settings (ToolGraph, etc.). */
+	virtual TValueOrError<bool, FText> TryLoadDefaults();
+	
+	virtual void Tick(float DeltaTime) override;
+	virtual bool IsTickable() const override;
+	virtual TStatId GetStatId() const override;
+
+	void CommitPropertyChangedEvents();
+
+	/** IKeyInputBehaviorTarget */
+	virtual void OnKeyPressed(const FKey& KeyID) override;
+
+	/** Allows to modify the just initialized working data from the editor. This can be used if the initialization of working data requires editor data access. */
+	virtual void OnWorkingDataInitialized(FPCGInteractiveToolWorkingData* WorkingData) const;
+	virtual void OnPropertyModified(UObject* Object, FProperty* Property);
+	virtual void OnRefreshInternal();
+	
+private:
+	UPCGComponent* FindOrGeneratePCGComponent(AActor& OwningActor) const;
+	
+	void DeleteGeneratedResources() const;
+	
+	void OnRefresh();
+	
+public:
+	UPROPERTY(EditAnywhere, BlueprintReadOnly, Category = "PCG", meta = (NoResetToDefault, GetAssetFilter = "GraphAssetFilter"))
+	TObjectPtr<UPCGGraphInterface> ToolGraph = nullptr;
+
+	UPROPERTY(EditAnywhere, BlueprintReadOnly, Category = "PCG", meta = (NoResetToDefault, EditCondition="ToolGraph != nullptr", EditConditionHides, GetOptions="GetValidDataInstanceNamesForGraph"))
+	FName DataInstance = NAME_None;
+	
+	// @todo_pcg: Convert to a target object, rather than actor, to ease transition into SceneGraph
+	/** The actor to modify with the tool. Either a selected actor, or a newly spawned one. */
+	UPROPERTY(meta = (TransientToolProperty))
+	TWeakObjectPtr<AActor> EditActor = nullptr;
+
+	/** The name to be issued to a newly created actor when the tool is applied. */
+	UPROPERTY(EditAnywhere, DisplayName="Actor Label", Category = "PCG", meta = (TransientToolProperty, EditCondition="HasSpawnedActor()", EditConditionHides))
+	FName NewActorName;
+
+	/** The name to be issued to a newly created PCG component when the tool is applied. */
+	UPROPERTY(EditAnywhere, DisplayName="Component Name", Category = "PCG", meta = (TransientToolProperty, EditCondition="HasGeneratedPCGComponent()", EditConditionHides))
+	FName NewPCGComponentName;
+
+	UPROPERTY(EditAnywhere, Category = "PCG", meta=(TransientToolProperty, EditConditionHides))
+	TSubclassOf<AActor> ActorClassToSpawn;
+
+	/** The current time left before a graph refresh (in seconds). */
+	UPROPERTY(meta = (TransientToolProperty))
+	float GraphRefreshTimer = 0.f;
+	
+	UPROPERTY(meta = (TransientToolProperty))
+	mutable TWeakObjectPtr<UPCGComponent> WorkingComponentCache;
+	
+	TMap<TWeakObjectPtr<UObject>, FPropertyChangedEvent> PropertyChangedEventQueue;
+
+	struct FGeneratedResources
+	{
+		TWeakObjectPtr<AActor> SpawnedActor;
+		TWeakObjectPtr<UPCGComponent> GeneratedPCGComponent;
+		/** A list of working data identifiers added this tool instantiation. Used to delete data on Cancel. */
+		TArray<FName> AddedWorkingDataIdentifiers;
+	};
+	
+	mutable FGeneratedResources GeneratedResources;
+
+	bool bTryToGenerate = false;
+
+private:
+	UFUNCTION()
+	bool HasSpawnedActor() const;
+
+	UFUNCTION()
+	bool HasGeneratedPCGComponent() const;
+};
diff --git a/Engine/Plugins/PCG/Source/PCGEditor/Private/EditorMode/Tools/Paint/PCGEdModeBrushBase.h b/Engine/Plugins/PCG/Source/PCGEditor/Private/EditorMode/Tools/Paint/PCGEdModeBrushBase.h
new file mode 100644
--- /dev/null
+++ b/Engine/Plugins/PCG/Source/PCGEditor/Private/EditorMode/Tools/Paint/PCGEdModeBrushBase.h
@@ -0,0 +1,16 @@
+// Copyright Epic Games, Inc. All Rights Reserved.
+
+#pragma once
+
+#include "BaseTools/BaseBrushTool.h"
+
+#include "PCGEdModeBrushBase.generated.h"
+
+UCLASS()
+class PCGEDITOR_API UPCGEdModeBrushBase : public UBrushBaseProperties
+{
+	GENERATED_BODY()
+
+public:
+	UPCGEdModeBrushBase() = default;
+};
diff --git a/Engine/Plugins/PCG/Source/PCGEditor/Private/EditorMode/Tools/Paint/PCGPaintTool.cpp b/Engine/Plugins/PCG/Source/PCGEditor/Private/EditorMode/Tools/Paint/PCGPaintTool.cpp
--- a/Engine/Plugins/PCG/Source/PCGEditor/Private/EditorMode/Tools/Paint/PCGPaintTool.cpp
+++ b/Engine/Plugins/PCG/Source/PCGEditor/Private/EditorMode/Tools/Paint/PCGPaintTool.cpp
@@ -2,6 +2,7 @@
 
 #include "EditorMode/Tools/Paint/PCGPaintTool.h"
 #include "Data/PCGPointArrayData.h"
+#include "Data/Tool/PCGToolPointData.h"
 #include "EditorMode/PCGEdModeStyle.h"
 #include "EditorMode/Tools/Gizmos/PCGGizmos.h"
 #include "Helpers/PCGHelpers.h"
diff --git a/Engine/Plugins/PCG/Source/PCGEditor/Private/EditorMode/Tools/Paint/PCGPaintTool.h b/Engine/Plugins/PCG/Source/PCGEditor/Private/EditorMode/Tools/Paint/PCGPaintTool.h
new file mode 100644
--- /dev/null
+++ b/Engine/Plugins/PCG/Source/PCGEditor/Private/EditorMode/Tools/Paint/PCGPaintTool.h
@@ -0,0 +1,110 @@
+// Copyright Epic Games, Inc. All Rights Reserved.
+
+#pragma once
+
+#include "PCGEdModeBrushBase.h"
+#include "BaseBehaviors/KeyAsModifierInputBehavior.h"
+#include "EditorMode/Tools/PCGInteractiveToolBuilder.h"
+#include "EditorMode/Tools/PCGInteractiveToolSettings.h"
+#include "EditorMode/Tools/Gizmos/PCGGizmos.h"
+#include "ToolTargets/PrimitiveComponentToolTarget.h"
+#include "PCGPaintTool.generated.h"
+
+class UMouseWheelInputBehavior;
+
+UCLASS()
+class PCGEDITOR_API UPCGInteractiveToolSettings_PaintTool : public UPCGInteractiveToolSettings
+{
+	GENERATED_BODY()
+
+public:
+	UPCGInteractiveToolSettings_PaintTool();
+
+	static FName StaticGetToolTag();
+
+	// ~Begin UPCGInteractiveToolSettings interface
+	virtual FName GetToolTag() const override;
+	virtual const UScriptStruct* GetWorkingDataType() const override;
+	// ~End UPCGInteractiveToolSettings interface
+
+	/** The brush shape to use. */
+	UPROPERTY(EditAnywhere, BlueprintReadOnly, Category = "Paint")
+	EPCGBrushMode BrushMode = EPCGBrushMode::Sphere;
+	
+	/** If false, will erase points on all paint data instances. */
+	UPROPERTY(EditAnywhere, BlueprintReadOnly, Category = "Paint")
+	bool bErasePointsOfSelectedDataInstance = true;
+
+	/** The step size for rotation of the brush. Moving the mouse wheel up or down will add or subtract from the brush rotation. */
+	UPROPERTY(EditAnywhere, BlueprintReadOnly, Category = "Paint")
+	float RotateYawStepSize = 10.f;
+
+	UPROPERTY(EditAnywhere, Category= "Raycast")
+	FPCGRaycastFilterRuleCollection RaycastRuleCollection;
+};
+
+UCLASS(Transient)
+class UPCGPaintTool : public UBaseBrushTool, public IMouseWheelBehaviorTarget
+{
+	GENERATED_BODY()
+
+public:
+	// ~Begin UPCGInteractiveToolSettings interface
+	virtual void Setup() override;
+	virtual void Shutdown(EToolShutdownType ShutdownType) override;
+	virtual bool HasAccept() const override;
+	virtual bool CanAccept() const override;
+	virtual bool HasCancel() const override;
+	virtual void OnTick(float DeltaTime) override;
+	// ~End UPCGInteractiveToolSettings interface
+
+	virtual bool HitTest(const FRay& Ray, FHitResult& OutHit) override;
+	
+protected:
+	virtual void OnClickPress(const FInputDeviceRay& PressPos) override;
+	virtual void OnUpdateDrag(const FRay& Ray) override;
+
+	/** IMouseWheelBehaviorTarget */
+	virtual FInputRayHit ShouldRespondToMouseWheel(const FInputDeviceRay& CurrentPos) override;
+	virtual void OnMouseWheelScrollUp(const FInputDeviceRay& CurrentPos) override;
+	virtual void OnMouseWheelScrollDown(const FInputDeviceRay& CurrentPos) override;
+	
+	virtual void SetupBrushStampIndicator() override;
+	virtual void UpdateBrushStampIndicator() override;
+	virtual void ShutdownBrushStampIndicator() override;
+
+	void RestartBrushStampIndicator();
+	
+	/** This will be used as the brush base, if not specifying a radius directly */
+	virtual double EstimateMaximumTargetDimension() override;
+	
+private:
+	void AddSinglePointFromBrush();
+	void TryDeletePointsFromBrush();
+	
+private:
+	UPROPERTY(Transient)
+	TObjectPtr<UPCGInteractiveToolSettings_PaintTool> ToolSettings;
+
+	UPROPERTY(Transient)
+	TObjectPtr<UMouseWheelInputBehavior> MouseWheelInputBehavior;
+	
+	UPROPERTY(Transient)
+	bool bPaintOnTick = false;
+
+	UPROPERTY(Transient)
+	bool bShiftPressedLastTick = false;
+
+	UE::PCG::EditorMode::Scene::FViewRayParams CachedRayParams{};
+};
+
+UCLASS(Transient)
+class UPCGPaintToolBuilder : public UPCGInteractiveToolWithToolTargetsBuilder
+{
+	GENERATED_BODY()
+
+public:
+	virtual bool CanBuildTool(const FToolBuilderState& SceneState) const override;
+	virtual UInteractiveTool* BuildTool(const FToolBuilderState& SceneState) const override;
+	virtual const FToolTargetTypeRequirements& GetTargetRequirements() const override;
+};
diff --git a/Engine/Plugins/PCG/Source/PCGEditor/Private/EditorMode/Tools/Volume/PCGVolumeTool.cpp b/Engine/Plugins/PCG/Source/PCGEditor/Private/EditorMode/Tools/Volume/PCGVolumeTool.cpp
--- a/Engine/Plugins/PCG/Source/PCGEditor/Private/EditorMode/Tools/Volume/PCGVolumeTool.cpp
+++ b/Engine/Plugins/PCG/Source/PCGEditor/Private/EditorMode/Tools/Volume/PCGVolumeTool.cpp
@@ -4,6 +4,7 @@
 
 #include "PCGEditorModule.h"
 #include "PCGVolume.h"
+#include "Data/Tool/PCGToolVolumeData.h"
 #include "EditorMode/PCGEdModeStyle.h"
 
 #include "InteractiveToolManager.h"
diff --git a/Engine/Plugins/PCG/Source/PCGEditor/Private/EditorMode/Tools/Volume/PCGVolumeTool.h b/Engine/Plugins/PCG/Source/PCGEditor/Private/EditorMode/Tools/Volume/PCGVolumeTool.h
new file mode 100644
--- /dev/null
+++ b/Engine/Plugins/PCG/Source/PCGEditor/Private/EditorMode/Tools/Volume/PCGVolumeTool.h
@@ -0,0 +1,118 @@
+// Copyright Epic Games, Inc. All Rights Reserved.
+
+#pragma once
+
+#include "EditorMode/Tools/PCGInteractiveToolBuilder.h"
+#include "EditorMode/Tools/PCGInteractiveToolSettings.h"
+#include "EditorMode/Tools/Helpers/PCGEdModeSceneQueryHelpers.h"
+#include "BaseBehaviors/BehaviorTargetInterfaces.h"
+
+#include "PCGVolumeTool.generated.h"
+
+class AVolume;
+class USingleClickOrDragInputBehavior;
+class UMouseHoverBehavior;
+
+UCLASS()
+class PCGEDITOR_API UPCGInteractiveToolSettings_Volume : public UPCGInteractiveToolSettings
+{
+	GENERATED_BODY()
+
+public:
+	UPCGInteractiveToolSettings_Volume();
+
+	static FName StaticGetToolTag();
+
+	// ~Begin UPCGInteractiveToolSettings interface
+	virtual FName GetToolTag() const override;
+	virtual const UScriptStruct* GetWorkingDataType() const override;
+	virtual void SetActorClassToSpawn(TSubclassOf<AActor> Class) override;
+	virtual bool CanResetToolData(FName DataInstanceIdentifier) override { return false; }
+	// ~End UPCGInteractiveToolSettings interface
+
+	AVolume* GetWorkingVolume(FName DataInstanceIdentifier) const;
+	// @todo_pcg: Origin modes: Center, Corner
+	// @todo_pcg: Snap to grid
+
+	UPROPERTY(EditAnywhere, Category = Raycast)
+	FPCGRaycastFilterRuleCollection RaycastRuleCollection;
+};
+
+/** 
+ * Tool class to create 3d volumes. Current implementation only creates cubes.
+ *	Click + drag -> creates the base.
+ *	Second click -> creates the height.
+ */
+UCLASS(Transient, MinimalAPI)
+class UPCGVolumeTool : public UInteractiveTool, public IClickBehaviorTarget, public IClickDragBehaviorTarget, public IHoverBehaviorTarget
+{
+	GENERATED_BODY()
+
+public:
+	// ~Begin UInteractiveTool interface
+	virtual bool HasCancel() const override { return true; }
+	virtual bool HasAccept() const override { return true; }
+
+	virtual void Setup() override;
+	virtual void Shutdown(EToolShutdownType ShutdownType) override;
+	virtual void OnTick(float DeltaTime) override;
+	virtual bool CanAccept() const override;
+	// ~End UInteractiveTool interface
+
+	// ~Begin IClickBehaviorTarget interface
+	virtual FInputRayHit IsHitByClick(const FInputDeviceRay& ClickPosition) override;
+	virtual void OnClicked(const FInputDeviceRay& ClickPosition) override;
+	// ~End IClickBehaviorTarget interface
+
+	// ~Begin IClickDragBehaviorTarget interface
+	virtual FInputRayHit CanBeginClickDragSequence(const FInputDeviceRay& DragPosition) override;
+	virtual void OnClickPress(const FInputDeviceRay& ClickPosition) override;
+	virtual void OnClickDrag(const FInputDeviceRay& DragPosition) override;
+	virtual void OnClickRelease(const FInputDeviceRay& ReleasePosition) override;
+	virtual void OnTerminateDragSequence() override {}
+	// ~End IClickDragBehaviorTarget interface
+
+	// ~Begin IHoverBehaviorTarget interface
+	virtual FInputRayHit BeginHoverSequenceHitTest(const FInputDeviceRay& PressPos) override;
+	virtual void OnBeginHover(const FInputDeviceRay& DevicePos) override {}
+	virtual bool OnUpdateHover(const FInputDeviceRay& DevicePos) override;
+	virtual void OnEndHover() override {}
+	// ~End IHoverBehaviorTarget interface
+
+private:
+	FBox GetCurrentBox() const;
+	void UpdateVolumeFromBox();
+	FVector GetBaseDragPosition(const FInputDeviceRay& DragPosition) const;
+
+	UE::PCG::EditorMode::Scene::FViewRayParams CachedRayParams;
+
+	UPROPERTY(Instanced)
+	TObjectPtr<USingleClickOrDragInputBehavior> ClickOrDragBehavior = nullptr;
+
+	UPROPERTY(Instanced)
+	TObjectPtr<UMouseHoverBehavior> HoverBehavior = nullptr;
+
+	UPROPERTY(Transient)
+	TObjectPtr<UPCGInteractiveToolSettings_Volume> ToolSettings;
+
+	bool bIsDraggingBase = false;
+	bool bIsDraggingHeight = false;
+	bool bBaseSelectionDone = false;
+	FVector FirstPoint = FVector::ZeroVector;
+	FVector SecondPoint = FVector::ZeroVector;
+	FVector InteractiveSecondPoint = FVector::ZeroVector;
+	double DefaultHalfHeight = 0;
+
+	FRay LastRay;
+	FBox LastBox = FBox(EForceInit::ForceInit);
+};
+
+UCLASS(Transient, MinimalAPI)
+class UPCGVolumeToolBuilder : public UPCGInteractiveToolWithToolTargetsBuilder
+{
+	GENERATED_BODY()
+
+public:
+	virtual bool CanBuildTool(const FToolBuilderState& SceneState) const override;
+	virtual UInteractiveTool* BuildTool(const FToolBuilderState& SceneState) const override;
+};
diff --git a/Engine/Plugins/PCG/Source/PCGEditor/Private/EditorMode/Widgets/SPCGEdModeAcceptToolOverlay.h b/Engine/Plugins/PCG/Source/PCGEditor/Private/EditorMode/Widgets/SPCGEdModeAcceptToolOverlay.h
new file mode 100644
--- /dev/null
+++ b/Engine/Plugins/PCG/Source/PCGEditor/Private/EditorMode/Widgets/SPCGEdModeAcceptToolOverlay.h
@@ -0,0 +1,41 @@
+// Copyright Epic Games, Inc. All Rights Reserved.
+
+#pragma once
+
+#include "Framework/SlateDelegates.h"
+
+#include "Widgets/SCompoundWidget.h"
+
+/** A simple Accept/Cancel widget for the PCG Editor Mode Toolkit. */
+class SPCGEdModeAcceptToolOverlay : public SCompoundWidget
+{
+	SLATE_BEGIN_ARGS(SPCGEdModeAcceptToolOverlay)
+		: _ActiveToolDisplayName(FText::GetEmpty())
+		, _ActiveToolIcon(nullptr)
+		, _IsAcceptButtonEnabled(false)
+		, _IsCancelButtonEnabled(false)
+		, _GetAcceptButtonVisibility(EVisibility::Collapsed)
+		, _GetCancelButtonVisibility(EVisibility::Collapsed) {}
+
+		SLATE_ATTRIBUTE(FText, ActiveToolDisplayName)
+		SLATE_ATTRIBUTE(const FSlateBrush*, ActiveToolIcon)
+
+		SLATE_ATTRIBUTE(FText, AcceptButtonLabel)
+		SLATE_ATTRIBUTE(FText, CancelButtonLabel)
+		SLATE_ATTRIBUTE(FText, AcceptButtonTooltip)
+		SLATE_ATTRIBUTE(FText, CancelButtonTooltip)
+		SLATE_ATTRIBUTE(bool, IsAcceptButtonEnabled)
+		SLATE_ATTRIBUTE(bool, IsCancelButtonEnabled)
+		SLATE_ATTRIBUTE(EVisibility, GetAcceptButtonVisibility)
+		SLATE_ATTRIBUTE(EVisibility, GetCancelButtonVisibility)
+
+		SLATE_EVENT(FOnClicked, OnAcceptButtonClicked)
+		SLATE_EVENT(FOnClicked, OnCancelButtonClicked)
+
+	SLATE_END_ARGS()
+
+	void Construct(const FArguments& InArgs);
+
+private:
+	FArguments Args;
+};
diff --git a/Engine/Plugins/PCG/Source/PCGEditor/Public/EditorMode/PCGEdMode.h b/Engine/Plugins/PCG/Source/PCGEditor/Public/EditorMode/PCGEdMode.h
deleted file mode 100644
diff --git a/Engine/Plugins/PCG/Source/PCGEditor/Public/EditorMode/PCGEdModeCommands.h b/Engine/Plugins/PCG/Source/PCGEditor/Public/EditorMode/PCGEdModeCommands.h
deleted file mode 100644
diff --git a/Engine/Plugins/PCG/Source/PCGEditor/Public/EditorMode/PCGEdModeHelpers.h b/Engine/Plugins/PCG/Source/PCGEditor/Public/EditorMode/PCGEdModeHelpers.h
deleted file mode 100644
diff --git a/Engine/Plugins/PCG/Source/PCGEditor/Public/EditorMode/PCGEdModeSettings.h b/Engine/Plugins/PCG/Source/PCGEditor/Public/EditorMode/PCGEdModeSettings.h
deleted file mode 100644
diff --git a/Engine/Plugins/PCG/Source/PCGEditor/Public/EditorMode/PCGEdModeStyle.h b/Engine/Plugins/PCG/Source/PCGEditor/Public/EditorMode/PCGEdModeStyle.h
deleted file mode 100644
diff --git a/Engine/Plugins/PCG/Source/PCGEditor/Public/EditorMode/PCGEdModeToolkit.h b/Engine/Plugins/PCG/Source/PCGEditor/Public/EditorMode/PCGEdModeToolkit.h
deleted file mode 100644
diff --git a/Engine/Plugins/PCG/Source/PCGEditor/Public/EditorMode/Tools/Gizmos/PCGGizmos.h b/Engine/Plugins/PCG/Source/PCGEditor/Public/EditorMode/Tools/Gizmos/PCGGizmos.h
deleted file mode 100644
diff --git a/Engine/Plugins/PCG/Source/PCGEditor/Public/EditorMode/Tools/Helpers/PCGEdModeActorHelpers.h b/Engine/Plugins/PCG/Source/PCGEditor/Public/EditorMode/Tools/Helpers/PCGEdModeActorHelpers.h
deleted file mode 100644
diff --git a/Engine/Plugins/PCG/Source/PCGEditor/Public/EditorMode/Tools/Helpers/PCGEdModeEditorUtilities.h b/Engine/Plugins/PCG/Source/PCGEditor/Public/EditorMode/Tools/Helpers/PCGEdModeEditorUtilities.h
deleted file mode 100644
diff --git a/Engine/Plugins/PCG/Source/PCGEditor/Public/EditorMode/Tools/Helpers/PCGEdModeSceneQueryHelpers.h b/Engine/Plugins/PCG/Source/PCGEditor/Public/EditorMode/Tools/Helpers/PCGEdModeSceneQueryHelpers.h
deleted file mode 100644
diff --git a/Engine/Plugins/PCG/Source/PCGEditor/Public/EditorMode/Tools/Helpers/PCGEdModeSceneVisualizeHelpers.h b/Engine/Plugins/PCG/Source/PCGEditor/Public/EditorMode/Tools/Helpers/PCGEdModeSceneVisualizeHelpers.h
deleted file mode 100644
diff --git a/Engine/Plugins/PCG/Source/PCGEditor/Public/EditorMode/Tools/Helpers/PCGEdModeSplineHelpers.h b/Engine/Plugins/PCG/Source/PCGEditor/Public/EditorMode/Tools/Helpers/PCGEdModeSplineHelpers.h
deleted file mode 100644
diff --git a/Engine/Plugins/PCG/Source/PCGEditor/Public/EditorMode/Tools/Line/PCGDrawSplineTool.h b/Engine/Plugins/PCG/Source/PCGEditor/Public/EditorMode/Tools/Line/PCGDrawSplineTool.h
deleted file mode 100644
diff --git a/Engine/Plugins/PCG/Source/PCGEditor/Public/EditorMode/Tools/PCGInteractiveToolBuilder.h b/Engine/Plugins/PCG/Source/PCGEditor/Public/EditorMode/Tools/PCGInteractiveToolBuilder.h
deleted file mode 100644
diff --git a/Engine/Plugins/PCG/Source/PCGEditor/Public/EditorMode/Tools/PCGInteractiveToolCommon.h b/Engine/Plugins/PCG/Source/PCGEditor/Public/EditorMode/Tools/PCGInteractiveToolCommon.h
deleted file mode 100644
diff --git a/Engine/Plugins/PCG/Source/PCGEditor/Public/EditorMode/Tools/PCGInteractiveToolSettings.h b/Engine/Plugins/PCG/Source/PCGEditor/Public/EditorMode/Tools/PCGInteractiveToolSettings.h
deleted file mode 100644
diff --git a/Engine/Plugins/PCG/Source/PCGEditor/Public/EditorMode/Tools/Paint/PCGEdModeBrushBase.h b/Engine/Plugins/PCG/Source/PCGEditor/Public/EditorMode/Tools/Paint/PCGEdModeBrushBase.h
deleted file mode 100644
diff --git a/Engine/Plugins/PCG/Source/PCGEditor/Public/EditorMode/Tools/Paint/PCGPaintTool.h b/Engine/Plugins/PCG/Source/PCGEditor/Public/EditorMode/Tools/Paint/PCGPaintTool.h
deleted file mode 100644
diff --git a/Engine/Plugins/PCG/Source/PCGEditor/Public/EditorMode/Tools/Volume/PCGVolumeTool.h b/Engine/Plugins/PCG/Source/PCGEditor/Public/EditorMode/Tools/Volume/PCGVolumeTool.h
deleted file mode 100644
diff --git a/Engine/Plugins/PCG/Source/PCGEditor/Public/EditorMode/Widgets/SPCGEdModeAcceptToolOverlay.h b/Engine/Plugins/PCG/Source/PCGEditor/Public/EditorMode/Widgets/SPCGEdModeAcceptToolOverlay.h
deleted file mode 100644
