Safety check for UI callback

diff --git a/Engine/Plugins/PCG/Source/PCGEditor/Private/Details/PCGInteractiveToolSettingsDetails.cpp b/Engine/Plugins/PCG/Source/PCGEditor/Private/Details/PCGInteractiveToolSettingsDetails.cpp
--- a/Engine/Plugins/PCG/Source/PCGEditor/Private/Details/PCGInteractiveToolSettingsDetails.cpp
+++ b/Engine/Plugins/PCG/Source/PCGEditor/Private/Details/PCGInteractiveToolSettingsDetails.cpp
@@ -186,21 +186,26 @@ FReply FPCGInteractiveToolSettingsDetails::ResetToolData() const
 
 EVisibility FPCGInteractiveToolSettingsDetails::IsResetToolDataButtonVisible() const
 {
-	if (FSlateApplication::Get().GetModifierKeys().IsControlDown())
+	if (ToolSettings.IsValid())
 	{
-		bool bAny = false;
-		for (const auto& WorkingDataTuple : ToolSettings->GetMutableTypedWorkingDataMap())
-		{			
-			if(ToolSettings->CanResetToolData(WorkingDataTuple.Key))
-			{
-				bAny = true;
+		if (FSlateApplication::Get().GetModifierKeys().IsControlDown())
+		{
+			bool bAny = false;
+			for (const auto& WorkingDataTuple : ToolSettings->GetMutableTypedWorkingDataMap())
+			{			
+				if(ToolSettings->CanResetToolData(WorkingDataTuple.Key))
+				{
+					bAny = true;
+				}
 			}
-		}
 
-		return bAny ? EVisibility::Visible : EVisibility::Collapsed;
+			return bAny ? EVisibility::Visible : EVisibility::Collapsed;
+		}
+		
+		return ToolSettings->CanResetToolData(ToolSettings->DataInstance) ? EVisibility::Visible : EVisibility::Collapsed;
 	}
-	
-	return ToolSettings->CanResetToolData(ToolSettings->DataInstance) ? EVisibility::Visible : EVisibility::Collapsed;
+
+	return EVisibility::Collapsed;
 }
 
 #undef LOCTEXT_NAMESPACE
