Added Min Point Spacing during Paint Tool Drag

diff --git a/Engine/Plugins/PCG/Source/PCGEditor/Private/EditorMode/Tools/Paint/PCGPaintTool.cpp b/Engine/Plugins/PCG/Source/PCGEditor/Private/EditorMode/Tools/Paint/PCGPaintTool.cpp
--- a/Engine/Plugins/PCG/Source/PCGEditor/Private/EditorMode/Tools/Paint/PCGPaintTool.cpp
+++ b/Engine/Plugins/PCG/Source/PCGEditor/Private/EditorMode/Tools/Paint/PCGPaintTool.cpp
@@ -56,6 +56,16 @@ const UScriptStruct* UPCGInteractiveToolSettings_PaintTool::GetWorkingDataType()
 	return FPCGInteractiveToolWorkingData_PointArrayData::StaticStruct();
 }
 
+void UPCGInteractiveToolSettings_PaintTool::RegisterPropertyWatchers()
+{
+	Super::RegisterPropertyWatchers();
+
+	WatchProperty(MinPointSpacing, [this](double InValue)
+	{
+		CachedMinPointSpacingSquared = InValue * InValue;		
+	});
+}
+
 void UPCGPaintTool::Setup()
 {
 	PropertyClass = UPCGEdModeBrushBase::StaticClass();
@@ -117,23 +127,8 @@ void UPCGPaintTool::OnTick(float DeltaTime)
 	Super::OnTick(DeltaTime);
 
 	bool bShiftPressedThisTick = GetShiftToggle();
-	bool bShiftPressedChanged = bShiftPressedThisTick != bShiftPressedLastTick;
-	
-	if(bPaintOnTick)
-	{
-		if(bShiftPressedThisTick)
-		{
-			TryDeletePointsFromBrush();	
-		}
-		else
-		{
-			AddSinglePointFromBrush();
-		}
 
-		bPaintOnTick = false;
-	}
-
-	if(bShiftPressedChanged)
+	if(bShiftPressedThisTick != bShiftPressedLastTick)
 	{
 		if(UPCGVolumetricBrushStampIndicator* VolumetricStampIndicator = Cast<UPCGVolumetricBrushStampIndicator>(BrushStampIndicator))
 		{
@@ -172,14 +167,42 @@ void UPCGPaintTool::OnClickPress(const FInputDeviceRay& PressPos)
 {
 	Super::OnClickPress(PressPos);
 
-	bPaintOnTick = true;
+	// Reset the last point transform on a new click
+	ToolSettings->LastPointTransform.Reset();
+
+	if(GetShiftToggle())
+	{
+		TryDeletePointsFromBrush();
+	}
+	else
+	{
+		AddSinglePointFromBrush();
+	}
 }
 
 void UPCGPaintTool::OnUpdateDrag(const FRay& Ray)
 {
 	Super::OnUpdateDrag(Ray);
 
-	bPaintOnTick = true;
+	if (GetShiftToggle())
+	{
+		TryDeletePointsFromBrush();
+	}
+	else
+	{
+		// If we added a point in one continuous operation, respect the min distance for the next point 
+		if (ToolSettings->LastPointTransform.IsSet())
+		{
+			if (FVector::DistSquared(LastBrushStamp.WorldPosition, ToolSettings->LastPointTransform.GetValue().GetLocation()) >= ToolSettings->CachedMinPointSpacingSquared)
+			{
+				AddSinglePointFromBrush();
+			}		
+		}
+		else
+		{
+			AddSinglePointFromBrush();
+		}
+	}
 }
 
 FInputRayHit UPCGPaintTool::ShouldRespondToMouseWheel(const FInputDeviceRay& CurrentPos)
@@ -282,9 +305,6 @@ void UPCGPaintTool::AddSinglePointFromBrush()
 	{
 		return;
 	}
-	
-	const int32 PointIndex = PointData->GetNumPoints();
-	PointData->SetNumPoints(PointIndex + 1);
 
 	TOptional<FRotator> BrushRotator;
 	if (UPCGVolumetricBrushStampIndicator* TypedBrushStampIndicator = Cast<UPCGVolumetricBrushStampIndicator>(BrushStampIndicator))
@@ -292,9 +312,13 @@ void UPCGPaintTool::AddSinglePointFromBrush()
 		BrushRotator = TypedBrushStampIndicator->GetPreviewMesh()->GetActor()->GetActorRotation();
 	}
 	
+	const FTransform PointTransform = FTransform(BrushRotator.Get(FRotator::ZeroRotator), LastBrushStamp.WorldPosition);
+	
+	const int32 PointIndex = PointData->GetNumPoints();
+	PointData->SetNumPoints(PointIndex + 1);
+	
 	// Values to assign
 	const float PointSteepness = 1.0f;
-	const FTransform PointTransform = FTransform(BrushRotator.Get(FRotator::ZeroRotator), LastBrushStamp.WorldPosition);
 	const FVector& Position = PointTransform.GetLocation();
 	const int32 PointSeed = PCGHelpers::ComputeSeed((int)Position.X, (int)Position.Y, (int)Position.Z);
 	const FBox LocalBounds = FBox(FVector(-CurrentBrushRadius), FVector(CurrentBrushRadius));
@@ -381,6 +405,8 @@ void UPCGPaintTool::AddSinglePointFromBrush()
 	{
 		FCoreUObjectDelegates::OnObjectPropertyChanged.Broadcast(PointData, EmptyPropertyChangedEvent);
 	}
+
+	ToolSettings->LastPointTransform = PointTransform;
 }
 
 void UPCGPaintTool::TryDeletePointsFromBrush()
diff --git a/Engine/Plugins/PCG/Source/PCGEditor/Private/EditorMode/Tools/Paint/PCGPaintTool.h b/Engine/Plugins/PCG/Source/PCGEditor/Private/EditorMode/Tools/Paint/PCGPaintTool.h
--- a/Engine/Plugins/PCG/Source/PCGEditor/Private/EditorMode/Tools/Paint/PCGPaintTool.h
+++ b/Engine/Plugins/PCG/Source/PCGEditor/Private/EditorMode/Tools/Paint/PCGPaintTool.h
@@ -25,6 +25,7 @@ public:
 	// ~Begin UPCGInteractiveToolSettings interface
 	virtual FName GetToolTag() const override;
 	virtual const UScriptStruct* GetWorkingDataType() const override;
+	virtual void RegisterPropertyWatchers() override;
 	// ~End UPCGInteractiveToolSettings interface
 
 	/** The brush shape to use. */
@@ -39,8 +40,18 @@ public:
 	UPROPERTY(EditAnywhere, BlueprintReadOnly, Category = "Paint")
 	float RotateYawStepSize = 10.f;
 
+	/** The minimum amount of distance to pass before drawing another point while dragging the mouse. */
+	UPROPERTY(EditAnywhere, Category = "Paint", meta = (ClampMin = 0, UIMax = 1000))
+	double MinPointSpacing = 50.0;
+
 	UPROPERTY(EditAnywhere, Category= "Raycast")
 	FPCGRaycastFilterRuleCollection RaycastRuleCollection;
+
+	/** Helper property so we can avoid sqrt. */
+	double CachedMinPointSpacingSquared = 0;
+
+	/** Tracks the last point that was added. Used for min-spacing during drag. */
+	TOptional<FTransform> LastPointTransform;
 };
 
 UCLASS(Transient)
@@ -88,9 +99,6 @@ private:
 
 	UPROPERTY(Transient)
 	TObjectPtr<UMouseWheelInputBehavior> MouseWheelInputBehavior;
-	
-	UPROPERTY(Transient)
-	bool bPaintOnTick = false;
 
 	UPROPERTY(Transient)
 	bool bShiftPressedLastTick = false;
