[PCG] Fix crash when reloading runtime graph assets

#rnx

diff --git a/Engine/Plugins/PCG/Source/PCG/Private/PCGActorAndComponentMapping.cpp b/Engine/Plugins/PCG/Source/PCG/Private/PCGActorAndComponentMapping.cpp
--- a/Engine/Plugins/PCG/Source/PCG/Private/PCGActorAndComponentMapping.cpp
+++ b/Engine/Plugins/PCG/Source/PCG/Private/PCGActorAndComponentMapping.cpp
@@ -1443,20 +1443,30 @@ void FPCGActorAndComponentMapping::OnPackageReloaded(const EPackageReloadPhase I
 	}
 	else if (InPackageReloadPhase == EPackageReloadPhase::PostPackageFixup)
 	{
+		TSet<UPCGGraphInstance*> GraphInstancesToRefresh;
+
 		for (const auto& RepointedObjectPair : InPackageReloadedEvent->GetRepointedObjects())
 		{
 			if (RepointedObjectPair.Value && RepointedObjectPair.Value->IsAsset() && RepointedObjectPair.Value->IsA<UPCGGraphInterface>())
 			{
-				ForEachObjectOfClass(UPCGGraphInstance::StaticClass(), [NewGraphInterface = RepointedObjectPair.Value](UObject* InObj)
+				ForEachObjectOfClass(UPCGGraphInstance::StaticClass(), [&GraphInstancesToRefresh, NewGraphInterface = RepointedObjectPair.Value](UObject* InObj)
 				{
 					if (UPCGGraphInstance* GraphInstance = Cast<UPCGGraphInstance>(InObj); GraphInstance && GraphInstance->Graph == NewGraphInterface)
 					{
-						GraphInstance->SetupCallbacks();
-						GraphInstance->OnGraphParametersChanged(EPCGGraphParameterEvent::GraphChanged, NAME_None);
+						GraphInstancesToRefresh.Add(GraphInstance);
 					}
 				});
 			}
 		}
+
+		for (UPCGGraphInstance* GraphInstance : GraphInstancesToRefresh)
+		{
+			if (GraphInstance)
+			{
+				GraphInstance->SetupCallbacks();
+				GraphInstance->OnGraphParametersChanged(EPCGGraphParameterEvent::GraphChanged, NAME_None);
+			}
+		}
 	}
 }
 
