[PCG] Expand PCGDataForGPU into PCGDataPacking and PCGDataDescription

#rnx

diff --git a/Engine/Plugins/PCG/Source/PCG/Private/Compute/Data/PCGProxyForGPUData.cpp b/Engine/Plugins/PCG/Source/PCG/Private/Compute/Data/PCGProxyForGPUData.cpp
--- a/Engine/Plugins/PCG/Source/PCG/Private/Compute/Data/PCGProxyForGPUData.cpp
+++ b/Engine/Plugins/PCG/Source/PCG/Private/Compute/Data/PCGProxyForGPUData.cpp
@@ -3,11 +3,25 @@
 #include "Compute/Data/PCGProxyForGPUData.h"
 
 #include "PCGModule.h"
+#include "Compute/Packing/PCGDataCollectionPacking.h"
+#include "Helpers/PCGAsync.h"
+
+#include "RHIGPUReadback.h"
+#include "RenderGraphResources.h"
 
 #include UE_INLINE_GENERATED_CPP_BY_NAME(PCGProxyForGPUData)
 
+#define LOCTEXT_NAMESPACE "PCGProxyForGPUData"
+
 PCG_DEFINE_TYPE_INFO(FPCGDataTypeInfoProxyForGPU, UPCGProxyForGPUData)
 
+#if WITH_EDITOR
+static TAutoConsoleVariable<bool> CVarWarnOnGPUReadbacks(
+	TEXT("pcg.Graph.GPU.WarnOnGPUReadbacks"),
+	false,
+	TEXT("Emits warnings on nodes which trigger a readback of data from GPU to CPU."));
+#endif
+
 void UPCGProxyForGPUData::Initialize(TSharedPtr<FPCGProxyForGPUDataCollection> InDataCollection, int InDataIndexInCollection)
 {
 	DataCollectionOnGPU = InDataCollection;
@@ -54,6 +68,11 @@ TSharedPtr<const FPCGProxyForGPUDataCollection> UPCGProxyForGPUData::GetInputDat
 	return DataCollectionOnGPU;
 }
 
+int UPCGProxyForGPUData::GetDataIndexInCollection() const
+{
+	return DataIndexInCollection;
+}
+
 FIntVector4 UPCGProxyForGPUData::GetElementCount() const
 {
 	if (TSharedPtr<const FPCGProxyForGPUDataCollection> DataCollectionGPU = GetGPUInfo())
@@ -120,3 +139,134 @@ TSharedPtr<FPCGProxyForGPUDataCollection> UPCGProxyForGPUData::GetGPUInfoMutable
 
 	return DataCollectionOnGPU;
 }
+
+FPCGProxyForGPUDataCollection::FPCGProxyForGPUDataCollection(TRefCountPtr<FRDGPooledBuffer> InBuffer, uint32 InBufferSizeBytes, const TSharedPtr<const FPCGDataCollectionDesc> InDescription, const TArray<FString>& InStringTable)
+	: Buffer(InBuffer)
+	, BufferSizeBytes(InBufferSizeBytes)
+	, Description(InDescription)
+	, StringTable(InStringTable)
+{
+}
+
+bool FPCGProxyForGPUDataCollection::GetCPUData(FPCGContext* InContext, int32 InDataIndex, FPCGTaggedData& OutData)
+{
+	UE::TScopeLock Lock(ReadbackLock);
+
+	if (bReadbackDataProcessed)
+	{
+		OutData = ReadbackData.IsValidIndex(InDataIndex) ? ReadbackData[InDataIndex] : FPCGTaggedData();
+		return true;
+	}
+
+	if (!ReadbackRequest && !bReadbackDataArrived)
+	{
+#if WITH_EDITOR
+		if (CVarWarnOnGPUReadbacks.GetValueOnAnyThread())
+		{
+			PCGLog::LogWarningOnGraph(LOCTEXT("PerformingReadback", "Performing readback"), InContext);
+		}
+#endif
+
+		ReadbackRequest = MakeShared<FRHIGPUBufferReadback>(TEXT("FPCGProxyForGPUDataCollectionReadback"));
+
+		ENQUEUE_RENDER_COMMAND(ReadbackDataCollection)([WeakThis=AsWeak(), ReadbackRequest=ReadbackRequest, Buffer=Buffer, BufferSizeBytes=BufferSizeBytes](FRHICommandListImmediate& RHICmdList)
+		{
+			LLM_SCOPE_BYTAG(PCG);
+
+			FRDGPooledBuffer* RDGBuffer = Buffer.GetReference();
+
+			ReadbackRequest->EnqueueCopy(RHICmdList, Buffer->GetRHI());
+
+			auto ExecuteAsync = [](auto&& RunnerFunc) -> void
+			{
+				if (IsInActualRenderingThread())
+				{
+					AsyncTask(ENamedThreads::ActualRenderingThread, [RunnerFunc]()
+					{
+						RunnerFunc(RunnerFunc);
+					});
+				}
+				else
+				{
+					// In specific cases (Server, -onethread, etc) the RenderingThread is actually the same as the GameThread.
+					// When this happens we want to avoid calling AsyncTask which could put us in a infinite task execution loop. 
+					// The reason is that if we are running this callback through the task graph we might stay in an executing loop until it has no tasks to execute,
+					// since we are pushing a new task as long as our data isn't ready and we are not advancing the GameThread as we are already on the GameThread this causes a infinite task execution.
+					// Instead delay to GameThread with ExecuteOnGameThread
+					ExecuteOnGameThread(UE_SOURCE_LOCATION, [RunnerFunc]()
+					{
+						RunnerFunc(RunnerFunc);
+					});
+				}
+			};
+
+			auto RunnerFunc = [WeakThis, ReadbackRequest, BufferSizeBytes, ExecuteAsync](auto&& RunnerFunc) -> void
+			{
+				LLM_SCOPE_BYTAG(PCG);
+
+				if (TSharedPtr<FPCGProxyForGPUDataCollection> This = WeakThis.Pin())
+				{
+					if (ReadbackRequest->IsReady())
+					{
+						if (void* RawData = ReadbackRequest->Lock(BufferSizeBytes))
+						{
+							uint32* DataAsUint = static_cast<uint32*>(RawData);
+
+							This->RawReadbackData.SetNumUninitialized(BufferSizeBytes);
+							FMemory::Memcpy(This->RawReadbackData.GetData(), RawData, BufferSizeBytes);
+						}
+
+						ReadbackRequest->Unlock();
+
+						This->bReadbackDataArrived = true;
+					}
+					else
+					{
+						ExecuteAsync(RunnerFunc);
+					}
+				}
+			};
+
+			ExecuteAsync(RunnerFunc);
+
+		});
+	}
+
+	if (!bReadbackDataArrived)
+	{
+		OutData = FPCGTaggedData();
+		return false;
+	}
+
+	ReadbackRequest.Reset();
+
+	bReadbackDataProcessed = true;
+
+	if (!RawReadbackData.IsEmpty())
+	{
+		FPCGDataCollection DataFromGPU;
+		const EPCGUnpackDataCollectionResult Result = PCGDataCollectionPackingHelpers::UnpackDataCollection(InContext, Description, RawReadbackData, /*OutputPinLabelAlias*/NAME_None, StringTable, DataFromGPU);
+
+		RawReadbackData.Empty();
+
+		if (Result != EPCGUnpackDataCollectionResult::Success)
+		{
+			OutData = FPCGTaggedData();
+			return true;
+		}
+
+		ReadbackData.Reset();
+		ReadbackDataRefs.Reset();
+
+		for (const FPCGTaggedData& TaggedData : DataFromGPU.TaggedData)
+		{
+			ReadbackData.Add(TaggedData);
+			ReadbackDataRefs.Add(TStrongObjectPtr<const UPCGData>(TaggedData.Data));
+		}
+	}
+
+	OutData = ReadbackData.IsValidIndex(InDataIndex) ? ReadbackData[InDataIndex] : FPCGTaggedData();
+	return true;
+}
+
+#undef LOCTEXT_NAMESPACE
diff --git a/Engine/Plugins/PCG/Source/PCG/Private/Compute/DataInterfaces/BuiltInKernels/PCGCountUniqueAttributeValuesDataInterface.cpp b/Engine/Plugins/PCG/Source/PCG/Private/Compute/DataInterfaces/BuiltInKernels/PCGCountUniqueAttributeValuesDataInterface.cpp
--- a/Engine/Plugins/PCG/Source/PCG/Private/Compute/DataInterfaces/BuiltInKernels/PCGCountUniqueAttributeValuesDataInterface.cpp
+++ b/Engine/Plugins/PCG/Source/PCG/Private/Compute/DataInterfaces/BuiltInKernels/PCGCountUniqueAttributeValuesDataInterface.cpp
@@ -4,6 +4,7 @@
 
 #include "Compute/PCGComputeCommon.h"
 #include "Compute/PCGDataBinding.h"
+#include "Compute/PCGDataDescription.h"
 #include "Compute/BuiltInKernels/PCGCountUniqueAttributeValuesKernel.h"
 #include "Elements/Metadata/PCGMetadataPartition.h"
 #include "Elements/Metadata/PCGMetadataPartitionKernel.h"
diff --git a/Engine/Plugins/PCG/Source/PCG/Private/Compute/DataInterfaces/Elements/PCGCopyPointsAnalysisDataInterface.cpp b/Engine/Plugins/PCG/Source/PCG/Private/Compute/DataInterfaces/Elements/PCGCopyPointsAnalysisDataInterface.cpp
--- a/Engine/Plugins/PCG/Source/PCG/Private/Compute/DataInterfaces/Elements/PCGCopyPointsAnalysisDataInterface.cpp
+++ b/Engine/Plugins/PCG/Source/PCG/Private/Compute/DataInterfaces/Elements/PCGCopyPointsAnalysisDataInterface.cpp
@@ -5,6 +5,7 @@
 #include "PCGContext.h"
 #include "Compute/PCGComputeKernel.h"
 #include "Compute/PCGDataBinding.h"
+#include "Compute/PCGDataDescription.h"
 #include "Elements/PCGCopyPoints.h"
 
 #include "GlobalRenderResources.h"
diff --git a/Engine/Plugins/PCG/Source/PCG/Private/Compute/DataInterfaces/Elements/PCGSkinnedMeshSpawnerDataInterface.cpp b/Engine/Plugins/PCG/Source/PCG/Private/Compute/DataInterfaces/Elements/PCGSkinnedMeshSpawnerDataInterface.cpp
--- a/Engine/Plugins/PCG/Source/PCG/Private/Compute/DataInterfaces/Elements/PCGSkinnedMeshSpawnerDataInterface.cpp
+++ b/Engine/Plugins/PCG/Source/PCG/Private/Compute/DataInterfaces/Elements/PCGSkinnedMeshSpawnerDataInterface.cpp
@@ -11,6 +11,7 @@
 #include "Compute/PCGDataBinding.h"
 #include "Compute/BuiltInKernels/PCGCountUniqueAttributeValuesKernel.h"
 #include "Compute/Data/PCGProxyForGPUData.h"
+#include "Compute/Packing/PCGMeshSpawnerPacking.h"
 #include "Compute/PrimitiveFactories/PCGPrimitiveFactoryISKMC.h"
 #include "Elements/PCGSkinnedMeshSpawner.h"
 #include "Elements/PCGSkinnedMeshSpawnerKernel.h"
@@ -520,7 +521,7 @@ void UPCGSkinnedMeshSpawnerDataProvider::CreatePrimitiveDescriptors(FPCGContext*
 			}
 		}
 
-		PCGDataForGPUHelpers::ComputeCustomFloatPacking(InContext, Settings, AttributeSelectors.GetValue(), InBinding, InputDataDesc, CustomFloatCount, AttributeIdOffsetStrides);
+		PCGMeshSpawnerPackingHelpers::ComputeCustomFloatPacking(InContext, Settings, AttributeSelectors.GetValue(), InputDataDesc, CustomFloatCount, AttributeIdOffsetStrides);
 	}
 
 	if (Settings->MeshSelectorParameters)
diff --git a/Engine/Plugins/PCG/Source/PCG/Private/Compute/DataInterfaces/Elements/PCGStaticMeshSpawnerDataInterface.cpp b/Engine/Plugins/PCG/Source/PCG/Private/Compute/DataInterfaces/Elements/PCGStaticMeshSpawnerDataInterface.cpp
--- a/Engine/Plugins/PCG/Source/PCG/Private/Compute/DataInterfaces/Elements/PCGStaticMeshSpawnerDataInterface.cpp
+++ b/Engine/Plugins/PCG/Source/PCG/Private/Compute/DataInterfaces/Elements/PCGStaticMeshSpawnerDataInterface.cpp
@@ -11,6 +11,7 @@
 #include "Compute/BuiltInKernels/PCGCountUniqueAttributeValuesKernel.h"
 #include "Compute/Data/PCGProxyForGPUData.h"
 #include "Compute/Data/PCGRawBufferData.h"
+#include "Compute/Packing/PCGMeshSpawnerPacking.h"
 #include "Compute/PrimitiveFactories/PCGPrimitiveFactoryPISMC.h"
 #include "Elements/PCGStaticMeshSpawner.h"
 #include "Elements/PCGStaticMeshSpawnerKernel.h"
@@ -566,7 +567,7 @@ void UPCGStaticMeshSpawnerDataProvider::CreatePrimitiveDescriptors(FPCGContext*
 			}
 		}
 
-		PCGDataForGPUHelpers::ComputeCustomFloatPacking(InContext, Settings, AttributeSelectors.GetValue(), InBinding, InputDataDesc, CustomFloatCount, AttributeIdOffsetStrides);
+		PCGMeshSpawnerPackingHelpers::ComputeCustomFloatPacking(InContext, Settings, AttributeSelectors.GetValue(), InputDataDesc, CustomFloatCount, AttributeIdOffsetStrides);
 	}
 
 	if (const UPCGMeshSelectorByAttribute* SelectorByAttribute = Cast<UPCGMeshSelectorByAttribute>(Settings->MeshSelectorParameters))
diff --git a/Engine/Plugins/PCG/Source/PCG/Private/Compute/DataInterfaces/PCGDataCollectionDataInterface.cpp b/Engine/Plugins/PCG/Source/PCG/Private/Compute/DataInterfaces/PCGDataCollectionDataInterface.cpp
--- a/Engine/Plugins/PCG/Source/PCG/Private/Compute/DataInterfaces/PCGDataCollectionDataInterface.cpp
+++ b/Engine/Plugins/PCG/Source/PCG/Private/Compute/DataInterfaces/PCGDataCollectionDataInterface.cpp
@@ -11,9 +11,11 @@
 #include "Compute/PCGComputeCommon.h"
 #include "Compute/PCGComputeKernel.h"
 #include "Compute/PCGDataBinding.h"
+#include "Compute/PCGDataDescription.h"
 #include "Compute/PCGDataUploadAdaptor.h"
 #include "Compute/Data/PCGProxyForGPUData.h"
 #include "Compute/Elements/PCGComputeGraphElement.h"
+#include "Compute/Packing/PCGDataCollectionPacking.h"
 #include "Helpers/PCGHelpers.h"
 
 #include "RenderGraphBuilder.h"
@@ -508,28 +510,28 @@ void UPCGDataCollectionDataInterface::GetDefines(FComputeKernelDefinitionSet& Ou
 	Super::GetDefines(OutDefinitionSet);
 
 	// Point property IDs
-	OutDefinitionSet.Defines.Add(FComputeKernelDefinition(TEXT("PCG_POINT_POSITION_ID"),   FString::FromInt(PCGComputeConstants::POINT_POSITION_ATTRIBUTE_ID)));
-	OutDefinitionSet.Defines.Add(FComputeKernelDefinition(TEXT("PCG_POINT_ROTATION_ID"),   FString::FromInt(PCGComputeConstants::POINT_ROTATION_ATTRIBUTE_ID)));
-	OutDefinitionSet.Defines.Add(FComputeKernelDefinition(TEXT("PCG_POINT_SCALE_ID"),      FString::FromInt(PCGComputeConstants::POINT_SCALE_ATTRIBUTE_ID)));
-	OutDefinitionSet.Defines.Add(FComputeKernelDefinition(TEXT("PCG_POINT_BOUNDS_MIN_ID"), FString::FromInt(PCGComputeConstants::POINT_BOUNDS_MIN_ATTRIBUTE_ID)));
-	OutDefinitionSet.Defines.Add(FComputeKernelDefinition(TEXT("PCG_POINT_BOUNDS_MAX_ID"), FString::FromInt(PCGComputeConstants::POINT_BOUNDS_MAX_ATTRIBUTE_ID)));
-	OutDefinitionSet.Defines.Add(FComputeKernelDefinition(TEXT("PCG_POINT_COLOR_ID"),      FString::FromInt(PCGComputeConstants::POINT_COLOR_ATTRIBUTE_ID)));
-	OutDefinitionSet.Defines.Add(FComputeKernelDefinition(TEXT("PCG_POINT_DENSITY_ID"),    FString::FromInt(PCGComputeConstants::POINT_DENSITY_ATTRIBUTE_ID)));
-	OutDefinitionSet.Defines.Add(FComputeKernelDefinition(TEXT("PCG_POINT_SEED_ID"),       FString::FromInt(PCGComputeConstants::POINT_SEED_ATTRIBUTE_ID)));
-	OutDefinitionSet.Defines.Add(FComputeKernelDefinition(TEXT("PCG_POINT_STEEPNESS_ID"),  FString::FromInt(PCGComputeConstants::POINT_STEEPNESS_ATTRIBUTE_ID)));
+	OutDefinitionSet.Defines.Add(FComputeKernelDefinition(TEXT("PCG_POINT_POSITION_ID"),   FString::FromInt(PCGDataCollectionPackingConstants::POINT_POSITION_ATTRIBUTE_ID)));
+	OutDefinitionSet.Defines.Add(FComputeKernelDefinition(TEXT("PCG_POINT_ROTATION_ID"),   FString::FromInt(PCGDataCollectionPackingConstants::POINT_ROTATION_ATTRIBUTE_ID)));
+	OutDefinitionSet.Defines.Add(FComputeKernelDefinition(TEXT("PCG_POINT_SCALE_ID"),      FString::FromInt(PCGDataCollectionPackingConstants::POINT_SCALE_ATTRIBUTE_ID)));
+	OutDefinitionSet.Defines.Add(FComputeKernelDefinition(TEXT("PCG_POINT_BOUNDS_MIN_ID"), FString::FromInt(PCGDataCollectionPackingConstants::POINT_BOUNDS_MIN_ATTRIBUTE_ID)));
+	OutDefinitionSet.Defines.Add(FComputeKernelDefinition(TEXT("PCG_POINT_BOUNDS_MAX_ID"), FString::FromInt(PCGDataCollectionPackingConstants::POINT_BOUNDS_MAX_ATTRIBUTE_ID)));
+	OutDefinitionSet.Defines.Add(FComputeKernelDefinition(TEXT("PCG_POINT_COLOR_ID"),      FString::FromInt(PCGDataCollectionPackingConstants::POINT_COLOR_ATTRIBUTE_ID)));
+	OutDefinitionSet.Defines.Add(FComputeKernelDefinition(TEXT("PCG_POINT_DENSITY_ID"),    FString::FromInt(PCGDataCollectionPackingConstants::POINT_DENSITY_ATTRIBUTE_ID)));
+	OutDefinitionSet.Defines.Add(FComputeKernelDefinition(TEXT("PCG_POINT_SEED_ID"),       FString::FromInt(PCGDataCollectionPackingConstants::POINT_SEED_ATTRIBUTE_ID)));
+	OutDefinitionSet.Defines.Add(FComputeKernelDefinition(TEXT("PCG_POINT_STEEPNESS_ID"),  FString::FromInt(PCGDataCollectionPackingConstants::POINT_STEEPNESS_ATTRIBUTE_ID)));
 
 	// Header sizes
-	OutDefinitionSet.Defines.Add(FComputeKernelDefinition(TEXT("PCG_DATA_COLLECTION_HEADER_SIZE_BYTES"), FString::FromInt(PCGComputeConstants::DATA_COLLECTION_HEADER_SIZE_BYTES)));
-	OutDefinitionSet.Defines.Add(FComputeKernelDefinition(TEXT("PCG_DATA_HEADER_SIZE_BYTES"),            FString::FromInt(PCGComputeConstants::DATA_HEADER_SIZE_BYTES)));
-	OutDefinitionSet.Defines.Add(FComputeKernelDefinition(TEXT("PCG_ATTRIBUTE_HEADER_SIZE_BYTES"),       FString::FromInt(PCGComputeConstants::ATTRIBUTE_HEADER_SIZE_BYTES)));
+	OutDefinitionSet.Defines.Add(FComputeKernelDefinition(TEXT("PCG_DATA_COLLECTION_HEADER_SIZE_BYTES"), FString::FromInt(PCGDataCollectionPackingConstants::DATA_COLLECTION_HEADER_SIZE_BYTES)));
+	OutDefinitionSet.Defines.Add(FComputeKernelDefinition(TEXT("PCG_DATA_HEADER_SIZE_BYTES"),            FString::FromInt(PCGDataCollectionPackingConstants::DATA_HEADER_SIZE_BYTES)));
+	OutDefinitionSet.Defines.Add(FComputeKernelDefinition(TEXT("PCG_ATTRIBUTE_HEADER_SIZE_BYTES"),       FString::FromInt(PCGDataCollectionPackingConstants::ATTRIBUTE_HEADER_SIZE_BYTES)));
 
 	// Misc
-	OutDefinitionSet.Defines.Add(FComputeKernelDefinition(TEXT("PCG_KERNEL_EXECUTED_FLAG"), FString::FromInt(PCGComputeConstants::KernelExecutedFlag)));
-	OutDefinitionSet.Defines.Add(FComputeKernelDefinition(TEXT("PCG_INVALID_DENSITY"),      FString::Printf(TEXT("%f"), PCGComputeConstants::INVALID_DENSITY)));
-	OutDefinitionSet.Defines.Add(FComputeKernelDefinition(TEXT("PCG_NUM_RESERVED_ATTRS"),   FString::FromInt(PCGComputeConstants::NUM_RESERVED_ATTRS)));
+	OutDefinitionSet.Defines.Add(FComputeKernelDefinition(TEXT("PCG_KERNEL_EXECUTED_FLAG"), FString::FromInt(PCGDataCollectionPackingConstants::KernelExecutedFlag)));
+	OutDefinitionSet.Defines.Add(FComputeKernelDefinition(TEXT("PCG_INVALID_DENSITY"),      FString::Printf(TEXT("%f"), PCGDataCollectionPackingConstants::INVALID_DENSITY)));
+	OutDefinitionSet.Defines.Add(FComputeKernelDefinition(TEXT("PCG_NUM_RESERVED_ATTRS"),   FString::FromInt(PCGDataCollectionPackingConstants::NUM_RESERVED_ATTRS)));
 
-	OutDefinitionSet.Defines.Add(FComputeKernelDefinition(TEXT("PCG_ATTRIBUTE_STRIDE_MASK"), FString::FromInt(PCGComputeConstants::AttributeStrideMask)));
-	OutDefinitionSet.Defines.Add(FComputeKernelDefinition(TEXT("PCG_ATTRIBUTE_ALLOCATED_MASK"), FString::FromInt(PCGComputeConstants::AttributeAllocatedMask)));
+	OutDefinitionSet.Defines.Add(FComputeKernelDefinition(TEXT("PCG_ATTRIBUTE_STRIDE_MASK"), FString::FromInt(PCGDataCollectionPackingConstants::AttributeStrideMask)));
+	OutDefinitionSet.Defines.Add(FComputeKernelDefinition(TEXT("PCG_ATTRIBUTE_ALLOCATED_MASK"), FString::FromInt(PCGDataCollectionPackingConstants::AttributeAllocatedMask)));
 }
 
 UComputeDataProvider* UPCGDataCollectionDataInterface::CreateDataProvider() const
@@ -575,7 +577,7 @@ bool UPCGDataCollectionDataProvider::PrepareForExecute_GameThread(UPCGDataBindin
 
 	if (!bBufferSizeValidated)
 	{
-		SizeBytes = GetPinDescription()->ComputePackedSizeBytes();
+		SizeBytes = PCGDataCollectionPackingHelpers::ComputePackedSizeBytes(GetPinDescription());
 		bIsBufferSizeValid = !PCGComputeHelpers::IsBufferSizeTooLarge(SizeBytes);
 
 		if (!bIsBufferSizeValid)
@@ -759,8 +761,8 @@ void FPCGDataCollectionDataProviderProxy::AllocateResources(FRDGBuilder& GraphBu
 		// Initialize with an empty data collection. The kernel may not run, for example if indirect dispatch args end up being 0. This ensures
 		// there is something meaningful to readback.
 		TArray<uint32> PackedDataCollection;
-		PackedDataCollection.Reserve(bZeroInitialize ? (SizeBytes / PackedDataCollection.GetTypeSize()) : PinDesc->ComputePackedHeaderSizeBytes());
-		PinDesc->WriteHeader(PackedDataCollection);
+		PackedDataCollection.Reserve(bZeroInitialize ? (SizeBytes / PackedDataCollection.GetTypeSize()) : PCGDataCollectionPackingHelpers::ComputePackedHeaderSizeBytes(PinDesc));
+		PCGDataCollectionPackingHelpers::WriteHeader(PinDesc, PackedDataCollection);
 
 		if (bZeroInitialize)
 		{
@@ -851,7 +853,7 @@ FRDGBufferSRVRef FPCGDataCollectionDataProviderProxy::GetAttributeIdRemapBufferS
 	if (!bProducedByCPU || !DataAdaptor || !DataAdaptor->UsesAttributeIdRemap())
 	{
 		// Default implementation. No remap required, disabled remapping.
-		OutFirstRemappedAttributeId = PCGComputeConstants::MAX_NUM_ATTRS;
+		OutFirstRemappedAttributeId = PCGDataCollectionPackingConstants::MAX_NUM_ATTRS;
 
 		return GraphBuilder.CreateSRV(FRDGBufferSRVDesc(GSystemTextures.GetDefaultStructuredBuffer(GraphBuilder, sizeof(int32))));
 	}
diff --git a/Engine/Plugins/PCG/Source/PCG/Private/Compute/DataInterfaces/PCGDataCollectionDataInterface.h b/Engine/Plugins/PCG/Source/PCG/Private/Compute/DataInterfaces/PCGDataCollectionDataInterface.h
--- a/Engine/Plugins/PCG/Source/PCG/Private/Compute/DataInterfaces/PCGDataCollectionDataInterface.h
+++ b/Engine/Plugins/PCG/Source/PCG/Private/Compute/DataInterfaces/PCGDataCollectionDataInterface.h
@@ -2,7 +2,6 @@
 
 #pragma once
 
-#include "Compute/PCGDataForGPU.h"
 #include "Compute/DataInterfaces/PCGComputeDataInterface.h"
 
 #include "RenderGraphResources.h"
@@ -11,6 +10,7 @@
 #include "PCGDataCollectionDataInterface.generated.h"
 
 struct FPCGComputeGraphContext;
+struct FPCGDataCollectionDesc;
 class FPCGDataCollectionDataInterfaceParameters;
 class FPCGDataUploadAdaptor;
 class UPCGDataBinding;
diff --git a/Engine/Plugins/PCG/Source/PCG/Private/Compute/DataInterfaces/PCGRawBufferDataInterface.cpp b/Engine/Plugins/PCG/Source/PCG/Private/Compute/DataInterfaces/PCGRawBufferDataInterface.cpp
--- a/Engine/Plugins/PCG/Source/PCG/Private/Compute/DataInterfaces/PCGRawBufferDataInterface.cpp
+++ b/Engine/Plugins/PCG/Source/PCG/Private/Compute/DataInterfaces/PCGRawBufferDataInterface.cpp
@@ -18,6 +18,32 @@
 namespace PCGRawBufferDataInterface
 {
 	static TCHAR const* TemplateFilePath = TEXT("/Plugin/PCG/Private/PCGRawBufferDataInterface.ush");
+
+	/** Computes the size (in bytes) of the data description after packing. */
+	uint64 ComputePackedSizeBytes(const FPCGDataDesc& InDataDesc)
+	{
+		return InDataDesc.GetElementCount().X * sizeof(uint32);
+	}
+
+	/** Computes the size (in bytes) of the data collection after packing. */
+	uint64 ComputePackedSizeBytes(TSharedPtr<const FPCGDataCollectionDesc> InDataCollectionDesc)
+	{
+		TRACE_CPUPROFILER_EVENT_SCOPE(PCGRawBufferPacking::ComputePackedSizeBytes);
+
+		uint64 TotalDataSizeBytes = 0;
+
+		if (InDataCollectionDesc)
+		{
+			const TArray<FPCGDataDesc>& DataDescs = InDataCollectionDesc->DataDescs;
+
+			for (const FPCGDataDesc& DataDesc : DataDescs)
+			{
+				TotalDataSizeBytes += ComputePackedSizeBytes(DataDesc);
+			}
+		}
+	
+		return TotalDataSizeBytes;
+	}
 }
 
 void UPCGRawBufferDataInterface::GetSupportedOutputs(TArray<FShaderFunctionDefinition>& OutFunctions) const
@@ -119,7 +145,7 @@ bool UPCGRawBufferDataProvider::PrepareForExecute_GameThread(UPCGDataBinding* In
 
 	if (const TSharedPtr<const FPCGDataCollectionDesc> ProducerPinDesc = InBinding->GetCachedKernelPinDataDesc(GetGraphBindingIndex()); ensure(ProducerPinDesc))
 	{
-		const uint64 PackedSizeBytes = ProducerPinDesc->ComputePackedSizeBytes();
+		const uint64 PackedSizeBytes = PCGRawBufferDataInterface::ComputePackedSizeBytes(ProducerPinDesc);
 
 		if (!PCGComputeHelpers::IsBufferSizeTooLarge(PackedSizeBytes))
 		{
diff --git a/Engine/Plugins/PCG/Source/PCG/Private/Compute/DataInterfaces/PCGTextureDataInterface.h b/Engine/Plugins/PCG/Source/PCG/Private/Compute/DataInterfaces/PCGTextureDataInterface.h
--- a/Engine/Plugins/PCG/Source/PCG/Private/Compute/DataInterfaces/PCGTextureDataInterface.h
+++ b/Engine/Plugins/PCG/Source/PCG/Private/Compute/DataInterfaces/PCGTextureDataInterface.h
@@ -2,7 +2,7 @@
 
 #pragma once
 
-#include "Compute/PCGDataForGPU.h"
+#include "Compute/PCGDataDescription.h"
 #include "Compute/DataInterfaces/PCGComputeDataInterface.h"
 #include "Data/PCGTextureData.h"
 
diff --git a/Engine/Plugins/PCG/Source/PCG/Private/Compute/PCGComputeCommon.cpp b/Engine/Plugins/PCG/Source/PCG/Private/Compute/PCGComputeCommon.cpp
--- a/Engine/Plugins/PCG/Source/PCG/Private/Compute/PCGComputeCommon.cpp
+++ b/Engine/Plugins/PCG/Source/PCG/Private/Compute/PCGComputeCommon.cpp
@@ -18,6 +18,7 @@
 #include "Compute/DataInterfaces/PCGStaticMeshDataInterface.h"
 #include "Compute/DataInterfaces/PCGTextureDataInterface.h"
 #include "Compute/DataInterfaces/PCGVirtualTextureDataInterface.h"
+#include "Compute/Packing/PCGDataCollectionPacking.h"
 #include "Data/PCGBasePointData.h"
 #include "Data/PCGTextureData.h"
 #include "Graph/PCGGPUGraphCompilationContext.h"
@@ -29,6 +30,11 @@
 
 #include UE_INLINE_GENERATED_CPP_BY_NAME(PCGComputeCommon)
 
+namespace PCGComputeConstants
+{
+	constexpr TCHAR DataLabelTagPrefix[] = { TEXT("PCG_DATA_LABEL") };
+}
+
 namespace PCGComputeHelpers
 {
 	static TAutoConsoleVariable<float> CVarMaxGPUBufferSizeProportion(
@@ -217,12 +223,12 @@ namespace PCGComputeHelpers
 
 	int32 GetAttributeIdFromMetadataAttributeIndex(int32 InAttributeIndex)
 	{
-		return InAttributeIndex >= 0 ? (InAttributeIndex + PCGComputeConstants::NUM_RESERVED_ATTRS) : INDEX_NONE;
+		return InAttributeIndex >= 0 ? (InAttributeIndex + PCGDataCollectionPackingConstants::NUM_RESERVED_ATTRS) : INDEX_NONE;
 	}
 
 	int32 GetMetadataAttributeIndexFromAttributeId(int32 InAttributeId)
 	{
-		return InAttributeId >= PCGComputeConstants::NUM_RESERVED_ATTRS ? (InAttributeId - PCGComputeConstants::NUM_RESERVED_ATTRS) : INDEX_NONE;
+		return InAttributeId >= PCGDataCollectionPackingConstants::NUM_RESERVED_ATTRS ? (InAttributeId - PCGDataCollectionPackingConstants::NUM_RESERVED_ATTRS) : INDEX_NONE;
 	}
 
 	FString GetPrefixedDataLabel(const FString& InLabel)
diff --git a/Engine/Plugins/PCG/Source/PCG/Private/Compute/PCGComputeGraph.cpp b/Engine/Plugins/PCG/Source/PCG/Private/Compute/PCGComputeGraph.cpp
--- a/Engine/Plugins/PCG/Source/PCG/Private/Compute/PCGComputeGraph.cpp
+++ b/Engine/Plugins/PCG/Source/PCG/Private/Compute/PCGComputeGraph.cpp
@@ -10,8 +10,10 @@
 #include "Compute/PCGDataBinding.h"
 #include "Compute/PCGPinPropertiesGPU.h"
 #include "Compute/Data/PCGProxyForGPUData.h"
+#include "Compute/Packing/PCGDataCollectionPacking.h"
 #include "Metadata/PCGMetadata.h"
 
+#include "RenderGraphResources.h"
 #include "ComputeFramework/ComputeKernelCompileResult.h"
 
 #include UE_INLINE_GENERATED_CPP_BY_NAME(PCGComputeGraph)
@@ -253,7 +255,7 @@ void UPCGComputeGraph::DebugLogDataDescriptions(const UPCGDataBinding* InBinding
 				// If attribute comes from an upstream GPU buffer, display the attribute ID in the upstream buffer.
 				int32 SourceBufferAttributeId = INDEX_NONE;
 
-				if (AttributeDesc.AttributeId >= PCGComputeConstants::NUM_RESERVED_ATTRS && UpstreamDataCollectionInfo)
+				if (AttributeDesc.AttributeId >= PCGDataCollectionPackingConstants::NUM_RESERVED_ATTRS && UpstreamDataCollectionInfo)
 				{
 					const int32 GraphAttributeIndex = PCGComputeHelpers::GetMetadataAttributeIndexFromAttributeId(AttributeDesc.AttributeId);
 					const int32* FoundIndex = InBinding->GetGraphAttributeToSourceBufferAttributeIndex().Find({ GraphAttributeIndex, UpstreamDataCollectionInfo->GetBuffer() });
diff --git a/Engine/Plugins/PCG/Source/PCG/Private/Compute/PCGComputeGraph.h b/Engine/Plugins/PCG/Source/PCG/Private/Compute/PCGComputeGraph.h
--- a/Engine/Plugins/PCG/Source/PCG/Private/Compute/PCGComputeGraph.h
+++ b/Engine/Plugins/PCG/Source/PCG/Private/Compute/PCGComputeGraph.h
@@ -3,7 +3,7 @@
 #pragma once
 
 #include "Compute/PCGComputeCommon.h"
-#include "Compute/PCGDataForGPU.h"
+#include "Compute/PCGDataDescription.h"
 
 #include "ComputeFramework/ComputeGraph.h"
 #include "ComputeFramework/ComputeKernel.h"
diff --git a/Engine/Plugins/PCG/Source/PCG/Private/Compute/PCGComputeKernel.cpp b/Engine/Plugins/PCG/Source/PCG/Private/Compute/PCGComputeKernel.cpp
--- a/Engine/Plugins/PCG/Source/PCG/Private/Compute/PCGComputeKernel.cpp
+++ b/Engine/Plugins/PCG/Source/PCG/Private/Compute/PCGComputeKernel.cpp
@@ -12,6 +12,7 @@
 #include "Compute/PCGDataBinding.h"
 #include "Compute/PCGPinPropertiesGPU.h"
 #include "Compute/Elements/PCGComputeGraphElement.h"
+#include "Compute/Packing/PCGDataCollectionPacking.h"
 #include "Graph/PCGGPUGraphCompilationContext.h"
 #include "Helpers/PCGSettingsHelpers.h"
 
@@ -281,7 +282,7 @@ void UPCGComputeKernel::ComputeDataDescFromPinProperties(const FPCGPinProperties
 		{
 			for (const FPCGKernelAttributeDesc& InAttrDesc : InDataDesc.AttributeDescs)
 			{
-				if (InAttrDesc.AttributeId >= PCGComputeConstants::NUM_RESERVED_ATTRS
+				if (InAttrDesc.AttributeId >= PCGDataCollectionPackingConstants::NUM_RESERVED_ATTRS
 					&& !OutDataDesc.AttributeDescs.ContainsByPredicate([InAttrDesc](const FPCGKernelAttributeDesc& OutAttrDesc)
 						{
 							// Note: We shouldn't need to check for uniqueness of attr Index, since attributes should all have unique index via
diff --git a/Engine/Plugins/PCG/Source/PCG/Private/Compute/PCGDataBinding.cpp b/Engine/Plugins/PCG/Source/PCG/Private/Compute/PCGDataBinding.cpp
--- a/Engine/Plugins/PCG/Source/PCG/Private/Compute/PCGDataBinding.cpp
+++ b/Engine/Plugins/PCG/Source/PCG/Private/Compute/PCGDataBinding.cpp
@@ -9,11 +9,14 @@
 #include "Compute/PCGComputeGraph.h"
 #include "Compute/PCGComputeKernel.h"
 #include "Compute/Data/PCGProxyForGPUData.h"
+#include "Compute/Packing/PCGDataCollectionPacking.h"
 #include "Data/PCGPointData.h"
 #include "Elements/Metadata/PCGMetadataElementCommon.h"
 #include "Metadata/PCGMetadata.h"
 #include "Metadata/Accessors/PCGAttributeAccessorHelpers.h"
 
+#include "RenderGraphResources.h"
+
 #include UE_INLINE_GENERATED_CPP_BY_NAME(PCGDataBinding)
 
 void UPCGDataBinding::Initialize(const UPCGComputeGraph* InComputeGraph, FPCGContext* InContext)
@@ -436,18 +439,18 @@ void UPCGDataBinding::AddInputDataAttributesToTable()
 					AttributeIdentifier.MetadataDomain = PCGMetadataDomainID::Default;
 				}
 					
-				const EPCGKernelAttributeType AttributeType = PCGDataForGPUHelpers::GetAttributeTypeFromMetadataType(AttributeTypes[Index]);
+				const EPCGKernelAttributeType AttributeType = PCGDataDescriptionHelpers::GetAttributeTypeFromMetadataType(AttributeTypes[Index]);
 
 				if (AttributeTable.AddAttribute(AttributeIdentifier, AttributeType) == INDEX_NONE)
 				{
-					UE_LOG(LogPCG, Error, TEXT("Attribute table exceeded maximum size (%d), use the 'Dump Data Descriptions' setting on the GPU node(s) to list attributes that are present."), PCGComputeConstants::MAX_NUM_CUSTOM_ATTRS);
+					UE_LOG(LogPCG, Error, TEXT("Attribute table exceeded maximum size (%d), use the 'Dump Data Descriptions' setting on the GPU node(s) to list attributes that are present."), PCGDataCollectionPackingConstants::MAX_NUM_CUSTOM_ATTRS);
 					break;
 				}
 			}
 		}
 	}
 
-	ensure(AttributeTable.Num() <= PCGComputeConstants::MAX_NUM_CUSTOM_ATTRS);
+	ensure(AttributeTable.Num() <= PCGDataCollectionPackingConstants::MAX_NUM_CUSTOM_ATTRS);
 }
 
 void UPCGDataBinding::AddInputDataStringsToTable()
@@ -522,7 +525,7 @@ void UPCGDataBinding::AddInputDataStringsToTable()
 					AttributeIdentifier.MetadataDomain = PCGMetadataDomainID::Default;
 				}
 				
-				const EPCGKernelAttributeType AttributeType = PCGDataForGPUHelpers::GetAttributeTypeFromMetadataType(AttributeTypes[AttributeIndex]);
+				const EPCGKernelAttributeType AttributeType = PCGDataDescriptionHelpers::GetAttributeTypeFromMetadataType(AttributeTypes[AttributeIndex]);
 
 				if (AttributeType == EPCGKernelAttributeType::StringKey || AttributeType == EPCGKernelAttributeType::Name)
 				{
diff --git a/Engine/Plugins/PCG/Source/PCG/Private/Compute/PCGDataDescription.cpp b/Engine/Plugins/PCG/Source/PCG/Private/Compute/PCGDataDescription.cpp
new file mode 100644
--- /dev/null
+++ b/Engine/Plugins/PCG/Source/PCG/Private/Compute/PCGDataDescription.cpp
@@ -0,0 +1,761 @@
+// Copyright Epic Games, Inc. All Rights Reserved.
+
+#include "Compute/PCGDataDescription.h"
+
+#include "Compute/PCGDataBinding.h"
+#include "Compute/PCGPinPropertiesGPU.h"
+#include "Compute/Packing/PCGDataCollectionPacking.h"
+#include "Data/PCGBasePointData.h"
+#include "Elements/Metadata/PCGMetadataElementCommon.h"
+#include "Metadata/PCGMetadata.h"
+#include "Metadata/Accessors/PCGAttributeAccessorHelpers.h"
+
+#include UE_INLINE_GENERATED_CPP_BY_NAME(PCGDataDescription)
+
+#define LOCTEXT_NAMESPACE "PCGDataDescription"
+
+namespace PCGDataDescriptionConstants
+{
+	const static FPCGKernelAttributeDesc PointPropertyDescs[PCGDataCollectionPackingConstants::NUM_POINT_PROPERTIES] =
+	{
+		FPCGKernelAttributeDesc(PCGDataCollectionPackingConstants::POINT_POSITION_ATTRIBUTE_ID,   EPCGKernelAttributeType::Float3, TEXT("$Position")),
+		FPCGKernelAttributeDesc(PCGDataCollectionPackingConstants::POINT_ROTATION_ATTRIBUTE_ID,   EPCGKernelAttributeType::Quat,   TEXT("$Rotation")),
+		FPCGKernelAttributeDesc(PCGDataCollectionPackingConstants::POINT_SCALE_ATTRIBUTE_ID,      EPCGKernelAttributeType::Float3, TEXT("$Scale")),
+		FPCGKernelAttributeDesc(PCGDataCollectionPackingConstants::POINT_BOUNDS_MIN_ATTRIBUTE_ID, EPCGKernelAttributeType::Float3, TEXT("$BoundsMin")),
+		FPCGKernelAttributeDesc(PCGDataCollectionPackingConstants::POINT_BOUNDS_MAX_ATTRIBUTE_ID, EPCGKernelAttributeType::Float3, TEXT("$BoundsMax")),
+		FPCGKernelAttributeDesc(PCGDataCollectionPackingConstants::POINT_COLOR_ATTRIBUTE_ID,      EPCGKernelAttributeType::Float4, TEXT("$Color")),
+		FPCGKernelAttributeDesc(PCGDataCollectionPackingConstants::POINT_DENSITY_ATTRIBUTE_ID,    EPCGKernelAttributeType::Float,  TEXT("$Density")),
+		FPCGKernelAttributeDesc(PCGDataCollectionPackingConstants::POINT_SEED_ATTRIBUTE_ID,       EPCGKernelAttributeType::Int,    TEXT("$Seed")),
+		FPCGKernelAttributeDesc(PCGDataCollectionPackingConstants::POINT_STEEPNESS_ATTRIBUTE_ID,  EPCGKernelAttributeType::Float,  TEXT("$Steepness"))
+	};
+}
+
+namespace PCGDataDescriptionHelpers
+{
+	EPCGKernelAttributeType GetAttributeTypeFromMetadataType(EPCGMetadataTypes MetadataType)
+	{
+		switch (MetadataType)
+		{
+		case EPCGMetadataTypes::Boolean:
+			return EPCGKernelAttributeType::Bool;
+		case EPCGMetadataTypes::Float:
+		case EPCGMetadataTypes::Double:
+			return EPCGKernelAttributeType::Float;
+		case EPCGMetadataTypes::Integer32:
+		case EPCGMetadataTypes::Integer64:
+			return EPCGKernelAttributeType::Int;
+		case EPCGMetadataTypes::Vector2:
+			return EPCGKernelAttributeType::Float2;
+		case EPCGMetadataTypes::Vector:
+			return EPCGKernelAttributeType::Float3;
+		case EPCGMetadataTypes::Rotator:
+			return EPCGKernelAttributeType::Rotator;
+		case EPCGMetadataTypes::Vector4:
+			return EPCGKernelAttributeType::Float4;
+		case EPCGMetadataTypes::Quaternion:
+			return EPCGKernelAttributeType::Quat;
+		case EPCGMetadataTypes::Transform:
+			return EPCGKernelAttributeType::Transform;
+		case EPCGMetadataTypes::SoftObjectPath: // TODO: This collapses all StringKey types into String attributes, meaning we'll lose the original CPU type when doing readback.
+		case EPCGMetadataTypes::SoftClassPath:
+		case EPCGMetadataTypes::String:
+			return EPCGKernelAttributeType::StringKey;
+		case EPCGMetadataTypes::Name:
+			return EPCGKernelAttributeType::Name;
+		default:
+			return EPCGKernelAttributeType::Invalid;
+		}
+	}
+
+	int GetAttributeTypeStrideBytes(EPCGKernelAttributeType Type)
+	{
+		switch (Type)
+		{
+		case EPCGKernelAttributeType::Bool:
+		case EPCGKernelAttributeType::Int:
+		case EPCGKernelAttributeType::Float:
+		case EPCGKernelAttributeType::StringKey:
+		case EPCGKernelAttributeType::Name:
+			return 4;
+		case EPCGKernelAttributeType::Float2:
+			return 8;
+		case EPCGKernelAttributeType::Float3:
+		case EPCGKernelAttributeType::Rotator:
+			return 12;
+		case EPCGKernelAttributeType::Float4:
+		case EPCGKernelAttributeType::Quat:
+			return 16;
+		case EPCGKernelAttributeType::Transform:
+			return 64;
+		default:
+			checkNoEntry();
+			return 0;
+		}
+	}
+}
+
+FPCGKernelAttributeKey::FPCGKernelAttributeKey(const FPCGAttributePropertySelector& InSelector, EPCGKernelAttributeType InType)
+	: Type(InType)
+{
+	SetSelector(InSelector);
+}
+
+void FPCGKernelAttributeKey::SetSelector(const FPCGAttributePropertySelector& InSelector)
+{
+	Name.ImportFromOtherSelector(InSelector);
+	UpdateIdentifierFromSelector();
+}
+
+bool FPCGKernelAttributeKey::UpdateIdentifierFromSelector()
+{
+	// @todo_pcg: When we support more domains, we'll need to have a way to convert it to a proper identifier.
+	// For now, only support default metadata domain, and mark the identifier invalid if it is anything else.
+	if (Name.GetSelection() != EPCGAttributePropertySelection::Attribute)
+	{
+		UE_LOG(LogPCG, Error, TEXT("While updating FPCGKernelAttributeKey with %s, the selector was not targeting an attribute. Discarded."), *Name.ToString());
+		Identifier.MetadataDomain = PCGMetadataDomainID::Invalid;
+		return true;
+	}
+
+	const FName DomainName = Name.GetDomainName();
+	if (DomainName != NAME_None)// && DomainName != PCGDataConstants::DataDomainName)
+	{
+		//UE_LOG(LogPCG, Error, TEXT("While updating FPCGKernelAttributeKey with %s, the selector was targeting a domain that is not the @Data domain which is unsupported at the moment. Discarded."), *Name.ToString());
+		UE_LOG(LogPCG, Error, TEXT("While updating FPCGKernelAttributeKey with %s, the selector was targeting a domain which is unsupported at the moment. Discarded."), *Name.ToString());
+		Identifier.MetadataDomain = PCGMetadataDomainID::Invalid;
+		return true;
+	}
+
+	FPCGAttributeIdentifier NewIdentifier{Name.GetAttributeName(), DomainName == PCGDataConstants::DataDomainName ? PCGMetadataDomainID::Data : PCGMetadataDomainID::Default};
+	bool bHasChanged = NewIdentifier != Identifier;
+	if (bHasChanged)
+	{
+		Identifier = std::move(NewIdentifier);
+	}
+	
+	return bHasChanged;
+}
+
+bool FPCGKernelAttributeKey::IsValid() const
+{
+	return Identifier.MetadataDomain.IsValid();
+}
+
+bool FPCGKernelAttributeKey::operator==(const FPCGKernelAttributeKey& Other) const
+{
+	return Type == Other.Type && Identifier == Other.Identifier;
+}
+
+uint32 GetTypeHash(const FPCGKernelAttributeKey& In)
+{
+	return HashCombine(GetTypeHash(In.Type), GetTypeHash(In.Identifier));
+}
+
+int32 FPCGKernelAttributeTable::GetAttributeId(const FPCGKernelAttributeKey& InAttribute) const
+{
+	const int TableIndex = AttributeTable.IndexOfByKey(InAttribute);
+	if (TableIndex >= 0)
+	{
+		return PCGComputeHelpers::GetAttributeIdFromMetadataAttributeIndex(TableIndex);
+	}
+	else
+	{
+		// Attribute with the given name and type is not present. If this is unexpected, check this attribute table for attributes
+		// in case the supplied name or type is wrong, and ensure any attribute created by a kernel is declared using GetKernelAttributeKeys().
+		return INDEX_NONE;
+	}
+}
+
+int32 FPCGKernelAttributeTable::GetAttributeId(FPCGAttributeIdentifier InIdentifier, EPCGKernelAttributeType InType) const
+{
+	return GetAttributeId(FPCGKernelAttributeKey(InIdentifier, InType));
+}
+
+int32 FPCGKernelAttributeTable::AddAttribute(const FPCGKernelAttributeKey& Key)
+{
+	int32 Index = AttributeTable.Find(Key);
+
+	if (Index == INDEX_NONE && AttributeTable.Num() < PCGDataCollectionPackingConstants::MAX_NUM_CUSTOM_ATTRS)
+	{
+		Index = AttributeTable.Num();
+		AttributeTable.Add(Key);
+	}
+
+	return Index;
+}
+
+int32 FPCGKernelAttributeTable::AddAttribute(FPCGAttributeIdentifier InIdentifier, EPCGKernelAttributeType InType)
+{
+	return AddAttribute(FPCGKernelAttributeKey(InIdentifier, InType));
+}
+
+#if PCG_KERNEL_LOGGING_ENABLED
+void FPCGKernelAttributeTable::DebugLog() const
+{
+	const UEnum* PCGKernelAttributeTypeEnum = StaticEnum<EPCGKernelAttributeType>();
+
+	for (const FPCGKernelAttributeKey& Attribute : AttributeTable)
+	{
+		UE_LOG(LogPCG, Display, TEXT("\tName: %s\t\tID: %d\t\tType: %s"),
+			*Attribute.Identifier.ToString(),
+			GetAttributeId(Attribute),
+			*PCGKernelAttributeTypeEnum->GetDisplayNameTextByValue(static_cast<int64>(Attribute.Type)).ToString());
+	}
+}
+#endif
+
+void FPCGKernelAttributeDesc::AddUniqueStringKeys(const TArray<int32>& InOtherStringKeys)
+{
+	for (int32 OtherStringKey : InOtherStringKeys)
+	{
+		UniqueStringKeys.AddUnique(OtherStringKey);
+	}
+}
+
+void FPCGKernelAttributeDesc::SetStringKeys(const TConstArrayView<int32>& InStringKeys)
+{
+	UniqueStringKeys = InStringKeys;
+}
+
+bool FPCGKernelAttributeDesc::operator==(const FPCGKernelAttributeDesc& Other) const
+{
+	return AttributeId == Other.AttributeId && AttributeKey == Other.AttributeKey;
+}
+
+FPCGDataDesc::FPCGDataDesc(FPCGDataTypeIdentifier InType, int32 InElementCount)
+	: FPCGDataDesc(InType, FIntVector4(InElementCount, 0, 0, 0))
+{
+}
+
+FPCGDataDesc::FPCGDataDesc(FPCGDataTypeIdentifier InType, FIntPoint InElementCount)
+	: FPCGDataDesc(InType, FIntVector4(InElementCount.X, InElementCount.Y, 0, 0))
+{
+}
+
+FPCGDataDesc::FPCGDataDesc(FPCGDataTypeIdentifier InType, FIntVector3 InElementCount)
+	: FPCGDataDesc(InType, FIntVector4(InElementCount.X, InElementCount.Y, InElementCount.Z, 0))
+{
+}
+
+FPCGDataDesc::FPCGDataDesc(FPCGDataTypeIdentifier InType, FIntVector4 InElementCount)
+	: Type(InType)
+	, ElementCount(InElementCount)
+{
+	ElementDimension = PCGComputeHelpers::GetElementDimension(InType);
+
+	InitializeAttributeDescs(/*InData=*/nullptr, /*InBinding*/nullptr);
+}
+
+FPCGDataDesc::FPCGDataDesc(const FPCGTaggedData& InTaggedData, const UPCGDataBinding* InBinding)
+{
+	check(InTaggedData.Data);
+
+	Type = InTaggedData.Data->GetDataTypeId();
+	ElementCount = PCGComputeHelpers::GetElementCount(InTaggedData.Data);
+	ElementDimension = PCGComputeHelpers::GetElementDimension(InTaggedData.Data);
+
+	TagStringKeys.Reserve(InTaggedData.Tags.Num());
+
+	for (const FString& Tag : InTaggedData.Tags)
+	{
+		const int Index = InBinding->GetStringTable().IndexOfByKey(Tag);
+		if (Index != INDEX_NONE)
+		{
+			TagStringKeys.Add(Index);
+		}
+	}
+
+	InitializeAttributeDescs(InTaggedData.Data, InBinding);
+}
+
+bool FPCGDataDesc::HasElementsMetadataDomainAttributes() const
+{
+	return AttributeDescs.ContainsByPredicate([](const FPCGKernelAttributeDesc& AttributeDesc)
+	{
+		return AttributeDesc.AttributeId >= PCGDataCollectionPackingConstants::NUM_RESERVED_ATTRS && AttributeDesc.AttributeKey.Identifier.MetadataDomain != PCGMetadataDomainID::Data;
+	});
+}
+
+bool FPCGDataDesc::ContainsAttribute(FPCGAttributeIdentifier InAttributeIdentifier) const
+{
+	return AttributeDescs.ContainsByPredicate([InAttributeIdentifier](const FPCGKernelAttributeDesc& AttributeDesc)
+	{
+		return AttributeDesc.AttributeKey.Identifier == InAttributeIdentifier;
+	});
+}
+
+bool FPCGDataDesc::ContainsAttribute(FPCGAttributeIdentifier InAttributeIdentifier, EPCGKernelAttributeType InAttributeType) const
+{
+	return AttributeDescs.ContainsByPredicate([InAttributeKey = FPCGKernelAttributeKey(InAttributeIdentifier, InAttributeType)](const FPCGKernelAttributeDesc& AttributeDesc)
+	{
+		return AttributeDesc.AttributeKey == InAttributeKey;
+	});
+}
+
+void FPCGDataDesc::AddAttribute(FPCGKernelAttributeKey InAttribute, const UPCGDataBinding* InBinding, const TArray<int32>* InOptionalUniqueStringKeys)
+{
+	const int32 AttributeId = InBinding->GetAttributeId(InAttribute);
+	if (AttributeId == INDEX_NONE)
+	{
+		return;
+	}
+
+	// Remove existing attributes if they collide with the new attribute (that is, if they have the same Identifier). The new attribute will stomp them,
+	// which is consistent with the general behavior on CPU. The node is authoritative on its attributes types.
+	// @todo_pcg: We should probably encapsulate the AttributeDescs, and they can only be added through AddAttribute() to ensure consistent behavior (w.r.t. stomping existing attributes).
+	for (int AttributeIndex = AttributeDescs.Num() - 1; AttributeIndex >= 0; --AttributeIndex)
+	{
+		FPCGKernelAttributeDesc& ExistingAttributeDesc = AttributeDescs[AttributeIndex];
+
+		if (ExistingAttributeDesc.AttributeKey.Identifier == InAttribute.Identifier)
+		{
+			AttributeDescs.RemoveAtSwap(AttributeIndex);
+		}
+	}
+
+	AttributeDescs.Emplace(AttributeId, InAttribute.Type, InAttribute.Identifier, InOptionalUniqueStringKeys);
+}
+
+FIntVector4 FPCGDataDesc::GetElementCountForAttribute(const FPCGKernelAttributeDesc& AttributeDesc) const
+{
+	// A buffer that is not fully allocated uses only a single value to represent all elements.
+	// Data domain only has a single element.
+	if (!IsAttributeAllocated(AttributeDesc) || AttributeDesc.AttributeKey.Identifier.MetadataDomain == PCGMetadataDomainID::Data)
+	{
+		return FIntVector4(1, 0, 0, 0);
+	}
+	else
+	{
+		return ElementCount;
+	}
+}
+
+int32 FPCGDataDesc::ComputeTotalElementCount() const
+{
+	switch (ElementDimension)
+	{
+	default: // Fallthrough
+	case EPCGElementDimension::One:
+		return ElementCount.X;
+	case EPCGElementDimension::Two:
+		return ElementCount.X * ElementCount.Y;
+	case EPCGElementDimension::Three:
+		return ElementCount.X * ElementCount.Y * ElementCount.Z;
+	case EPCGElementDimension::Four:
+		return ElementCount.X * ElementCount.Y * ElementCount.Z * ElementCount.W;
+	}
+}
+
+void FPCGDataDesc::AddElementCount(int32 InElementCountToAdd)
+{
+	ensure(ElementDimension == EPCGElementDimension::One);
+	ElementCount.X += InElementCountToAdd;
+}
+
+void FPCGDataDesc::AddElementCount(FIntPoint InElementCountToAdd)
+{
+	ensure(ElementDimension == EPCGElementDimension::Two);
+	ElementCount.X += InElementCountToAdd.X;
+	ElementCount.Y += InElementCountToAdd.Y;
+}
+
+void FPCGDataDesc::AddElementCount(FIntVector3 InElementCountToAdd)
+{
+	ensure(ElementDimension == EPCGElementDimension::Three);
+	ElementCount.X += InElementCountToAdd.X;
+	ElementCount.Y += InElementCountToAdd.Y;
+	ElementCount.Z += InElementCountToAdd.Z;
+}
+
+void FPCGDataDesc::AddElementCount(FIntVector4 InElementCountToAdd)
+{
+	ensure(ElementDimension == EPCGElementDimension::Four);
+	ElementCount += InElementCountToAdd;
+}
+
+FPCGDataDesc& FPCGDataDesc::CombineElementCount(const FPCGDataDesc& Other, EPCGElementMultiplicity Multiplicity)
+{
+	if (Multiplicity == EPCGElementMultiplicity::Sum)
+	{
+		if (ElementDimension == EPCGElementDimension::One)
+		{
+			// Special case for 1D processing, sum by collapsing the other element count into a single value.
+			ElementCount.X += Other.ComputeTotalElementCount();
+		}
+		else if (ElementDimension == EPCGElementDimension::Two)
+		{
+			if (Other.ElementDimension == EPCGElementDimension::One)
+			{
+				ElementCount.X += Other.ElementCount.X;
+				ElementCount.Y += Other.ElementCount.X;
+			}
+			else
+			{
+				ElementCount += Other.ElementCount;
+				ElementCount.Z = 0;
+				ElementCount.W = 0;
+			}
+		}
+		else if (ElementDimension == EPCGElementDimension::Three)
+		{
+			if (Other.ElementDimension == EPCGElementDimension::One)
+			{
+				ElementCount.X += Other.ElementCount.X;
+				ElementCount.Y += Other.ElementCount.X;
+				ElementCount.Z += Other.ElementCount.X;
+			}
+			else
+			{
+				ElementCount += Other.ElementCount;
+				ElementCount.W = 0;
+			}
+		}
+		else if (ElementDimension == EPCGElementDimension::Four)
+		{
+			if (Other.ElementDimension == EPCGElementDimension::One)
+			{
+				ElementCount.X += Other.ElementCount.X;
+				ElementCount.Y += Other.ElementCount.X;
+				ElementCount.Z += Other.ElementCount.X;
+				ElementCount.W += Other.ElementCount.X;
+			}
+			else
+			{
+				ElementCount += Other.ElementCount;
+			}
+		}
+	}
+	else
+	{
+		if (ElementDimension == EPCGElementDimension::One)
+		{
+			// Special case for 1D processing, scale by collapsing the other element count into a single value.
+			ElementCount.X *= Other.ComputeTotalElementCount();
+		}
+		else if (ElementDimension == EPCGElementDimension::Two)
+		{
+			if (Other.ElementDimension == EPCGElementDimension::One)
+			{
+				ElementCount *= Other.ElementCount.X;
+			}
+			else
+			{
+				ElementCount *= Other.ElementCount;
+			}
+		}
+		else if (ElementDimension == EPCGElementDimension::Three)
+		{
+			if (Other.ElementDimension == EPCGElementDimension::One)
+			{
+				ElementCount *= Other.ElementCount.X;
+			}
+			else
+			{
+				ElementCount *= Other.ElementCount;
+			}
+		}
+		else if (ElementDimension == EPCGElementDimension::Four)
+		{
+			if (Other.ElementDimension == EPCGElementDimension::One)
+			{
+				ElementCount *= Other.ElementCount.X;
+			}
+			else
+			{
+				ElementCount *= Other.ElementCount;
+			}
+		}
+	}
+
+	return *this;
+}
+
+FPCGDataDesc FPCGDataDesc::CombineElementCount(const FPCGDataDesc& A, const FPCGDataDesc& B, EPCGElementMultiplicity Multiplicity)
+{
+	FPCGDataDesc Desc = A;
+	Desc.CombineElementCount(B, Multiplicity);
+	return Desc;
+}
+
+bool FPCGDataDesc::IsAttributeAllocated(int32 InAttributeId) const
+{
+	// @todo_pcg: Only point properties are supported for single value ranges. All other attributes are fully allocated.
+	if (Type == EPCGDataType::Point && InAttributeId < PCGDataCollectionPackingConstants::NUM_RESERVED_ATTRS)
+	{
+		switch (InAttributeId)
+		{
+		case PCGDataCollectionPackingConstants::POINT_POSITION_ATTRIBUTE_ID:
+			return !!(AllocatedPointProperties & EPCGPointNativeProperties::Transform);
+		case PCGDataCollectionPackingConstants::POINT_ROTATION_ATTRIBUTE_ID:
+			return !!(AllocatedPointProperties & EPCGPointNativeProperties::Transform);
+		case PCGDataCollectionPackingConstants::POINT_SCALE_ATTRIBUTE_ID:
+			return !!(AllocatedPointProperties & EPCGPointNativeProperties::Transform);
+		case PCGDataCollectionPackingConstants::POINT_BOUNDS_MIN_ATTRIBUTE_ID:
+			return !!(AllocatedPointProperties & EPCGPointNativeProperties::BoundsMin);
+		case PCGDataCollectionPackingConstants::POINT_BOUNDS_MAX_ATTRIBUTE_ID:
+			return !!(AllocatedPointProperties & EPCGPointNativeProperties::BoundsMax);
+		case PCGDataCollectionPackingConstants::POINT_COLOR_ATTRIBUTE_ID:
+			return !!(AllocatedPointProperties & EPCGPointNativeProperties::Color);
+		case PCGDataCollectionPackingConstants::POINT_SEED_ATTRIBUTE_ID:
+			return !!(AllocatedPointProperties & EPCGPointNativeProperties::Seed);
+		case PCGDataCollectionPackingConstants::POINT_STEEPNESS_ATTRIBUTE_ID:
+			return !!(AllocatedPointProperties & EPCGPointNativeProperties::Steepness);
+		// Note: Because density is used to represent invalid points, it must always be allocated.
+		case PCGDataCollectionPackingConstants::POINT_DENSITY_ATTRIBUTE_ID:
+			return true;
+		default:
+			return true;
+		}
+	}
+
+	return true;
+}
+
+void FPCGDataDesc::InitializeAttributeDescs(const UPCGData* InData, const UPCGDataBinding* InBinding)
+{
+	TRACE_CPUPROFILER_EVENT_SCOPE(FPCGDataDesc::InitializeAttributeDescs);
+
+	if (Type == EPCGDataType::Point)
+	{
+		AttributeDescs.Append(PCGDataDescriptionConstants::PointPropertyDescs, PCGDataCollectionPackingConstants::NUM_POINT_PROPERTIES);
+
+		if (const UPCGBasePointData* PointData = Cast<UPCGBasePointData>(InData))
+		{
+			AllocatedPointProperties = PointData->GetAllocatedProperties();
+		}
+	}
+	else { /* TODO: More types! */ }
+
+	const UPCGMetadata* Metadata = (InData && PCGComputeHelpers::ShouldImportAttributesFromData(InData)) ? InData->ConstMetadata() : nullptr;
+
+	if (InBinding && Metadata)
+	{
+		const FPCGMetadataDomainID MetadataDefaultDomainID = Metadata->GetConstDefaultMetadataDomain()->GetDomainID();
+		
+		TArray<FPCGAttributeIdentifier> AttributeIdentifiers;
+		TArray<EPCGMetadataTypes> AttributeTypes;
+		Metadata->GetAllAttributes(AttributeIdentifiers, AttributeTypes);
+
+		// Cache the keys to a given domain, so we don't recreate them
+		TMap<FPCGMetadataDomainID, TUniquePtr<const IPCGAttributeAccessorKeys>> AllKeys;
+
+		for (int CustomAttributeIndex = 0; CustomAttributeIndex < AttributeIdentifiers.Num(); ++CustomAttributeIndex)
+		{
+			// @todo_pcg: Attributes on other domains than the default are ignored at the moment, until we have a better way of representing
+			// different domains in the GPU header.
+			// It means those are lost.
+			FPCGAttributeIdentifier AttributeIdentifier = AttributeIdentifiers[CustomAttributeIndex];
+			if (AttributeIdentifier.MetadataDomain != PCGMetadataDomainID::Default && AttributeIdentifier.MetadataDomain != MetadataDefaultDomainID)
+			{
+				continue;
+			}
+
+			// If the domain is the default domain, force it to the default identifier.
+			if (AttributeIdentifier.MetadataDomain == MetadataDefaultDomainID)
+			{
+				AttributeIdentifier.MetadataDomain = PCGMetadataDomainID::Default;
+			}
+			
+			const EPCGKernelAttributeType AttributeType = PCGDataDescriptionHelpers::GetAttributeTypeFromMetadataType(AttributeTypes[CustomAttributeIndex]);
+
+			if (AttributeType == EPCGKernelAttributeType::Invalid)
+			{
+				const UEnum* EnumClass = StaticEnum<EPCGMetadataTypes>();
+				check(EnumClass);
+
+				FPCGAttributePropertySelector Selector = FPCGAttributePropertySelector::CreateAttributeSelector(AttributeIdentifier.Name);
+				InData->SetDomainFromDomainID(AttributeIdentifier.MetadataDomain, Selector);
+
+				UE_LOG(LogPCG, Warning, TEXT("Skipping attribute '%s'. '%s' type attributes are not supported on GPU."),
+					*Selector.ToString(),
+					*EnumClass->GetNameStringByValue(static_cast<int64>(AttributeTypes[CustomAttributeIndex])));
+
+				continue;
+			}
+
+			// Ignore excess attributes.
+			if (CustomAttributeIndex >= PCGDataCollectionPackingConstants::MAX_NUM_CUSTOM_ATTRS)
+			{
+				// TODO: Would be nice to include the pin label for debug purposes
+				UE_LOG(LogPCG, Warning, TEXT("Attempted to exceed max number of custom attributes (%d). Additional attributes will be ignored."), PCGDataCollectionPackingConstants::MAX_NUM_CUSTOM_ATTRS);
+				break;
+			}
+
+			TArray<int32> UniqueStringKeys;
+
+			if (AttributeType == EPCGKernelAttributeType::StringKey || AttributeType == EPCGKernelAttributeType::Name)
+			{
+				const FPCGMetadataDomain* MetadataDomain = Metadata->GetConstMetadataDomain(AttributeIdentifier.MetadataDomain);
+				const FPCGMetadataAttributeBase* AttributeBase = MetadataDomain->GetConstAttribute(AttributeIdentifier.Name);
+				check(MetadataDomain && AttributeBase);
+
+				const TUniquePtr<const IPCGAttributeAccessor> Accessor = PCGAttributeAccessorHelpers::CreateConstAccessor(AttributeBase, MetadataDomain);
+				TUniquePtr<const IPCGAttributeAccessorKeys>& Keys = AllKeys.FindOrAdd(AttributeIdentifier.MetadataDomain);
+				if (!Keys.IsValid())
+				{
+					FPCGAttributePropertySelector Selector = FPCGAttributePropertySelector::CreateAttributeSelector(AttributeIdentifier.Name);
+					InData->SetDomainFromDomainID(AttributeIdentifier.MetadataDomain, Selector);
+					Keys = PCGAttributeAccessorHelpers::CreateConstKeys(InData, Selector);
+				}
+
+				check(Accessor && Keys);
+
+				PCGMetadataElementCommon::ApplyOnAccessor<FString>(*Keys, *Accessor, [&UniqueStringKeys, &StringTable = InBinding->GetStringTable()](const FString& InValue, int32)
+				{
+					const int StringTableIndex = StringTable.IndexOfByKey(InValue);
+					if (StringTableIndex != INDEX_NONE)
+					{
+						UniqueStringKeys.AddUnique(StringTableIndex);
+					}
+				});
+			}
+
+			const int32 AttributeId = InBinding->GetAttributeId(AttributeIdentifier, AttributeType);
+			ensureMsgf(AttributeId != INDEX_NONE, TEXT("Attribute '%s' type %d was missing from attribute table."), *AttributeIdentifier.Name.ToString(), (int32)AttributeType);
+
+			if (AttributeId != INDEX_NONE)
+			{
+				AttributeDescs.Emplace(AttributeId, AttributeType, AttributeIdentifier, &UniqueStringKeys);
+			}
+		}
+	}
+}
+
+TSharedPtr<FPCGDataCollectionDesc> FPCGDataCollectionDesc::MakeShared()
+{
+	return MakeShareable(new FPCGDataCollectionDesc());
+}
+
+TSharedPtr<FPCGDataCollectionDesc> FPCGDataCollectionDesc::MakeSharedFrom(const TSharedPtr<const FPCGDataCollectionDesc> InOtherDataDesc)
+{
+	TRACE_CPUPROFILER_EVENT_SCOPE(FPCGDataCollectionDesc::MakeSharedFrom);
+
+	TSharedPtr<FPCGDataCollectionDesc> NewDesc = MakeShareable(new FPCGDataCollectionDesc());
+
+	if (InOtherDataDesc)
+	{
+		NewDesc->DataDescs = InOtherDataDesc->DataDescs;
+	}
+
+	return NewDesc;
+}
+
+uint32 FPCGDataCollectionDesc::ComputeTotalElementCount() const
+{
+	uint32 ProcessingElementCount = 0;
+
+	for (const FPCGDataDesc& DataDesc : DataDescs)
+	{
+		ProcessingElementCount += DataDesc.ComputeTotalElementCount();
+	}
+
+	return ProcessingElementCount;
+}
+
+uint32 FPCGDataCollectionDesc::ComputeTotalElementCount(FPCGDataTypeIdentifier InDataType) const
+{
+	// InDataType not used anymore
+	// @todo_pcg Remove in future pass.
+	return ComputeTotalElementCount();
+}
+
+bool FPCGDataCollectionDesc::GetAttributeDesc(FPCGAttributeIdentifier InAttributeIdentifier, FPCGKernelAttributeDesc& OutAttributeDesc, bool& bOutConflictingTypesFound, bool& bOutPresentOnAllData) const
+{
+	// Will be set to the type of the first attribute across all the data which has a matching name, and then is used to detect if subsequent
+	// attributes are found with conflicting type.
+	EPCGKernelAttributeType FoundType = EPCGKernelAttributeType::Invalid;
+	bOutPresentOnAllData = false;
+
+	bool bFoundOnAllData = true;
+
+	for (const FPCGDataDesc& DataDesc : DataDescs)
+	{
+		bool bFoundOnData = false;
+
+		for (const FPCGKernelAttributeDesc& AttributeDesc : DataDesc.AttributeDescs)
+		{
+			if (AttributeDesc.AttributeKey.Identifier == InAttributeIdentifier && ensure(AttributeDesc.AttributeKey.Type != EPCGKernelAttributeType::Invalid))
+			{
+				bFoundOnData = true;
+
+				if (FoundType == EPCGKernelAttributeType::Invalid)
+				{
+					// Take the first matching attribute.
+					OutAttributeDesc = AttributeDesc;
+					FoundType = AttributeDesc.AttributeKey.Type;
+				}
+				else if (FoundType != AttributeDesc.AttributeKey.Type)
+				{
+					// Signal conflict found.
+					bOutConflictingTypesFound = true;
+
+					// Can stop iterating once we find a conflict. Signal that a matching attribute was found (despite the conflict).
+					return true;
+				}
+			}
+		}
+
+		bFoundOnAllData &= bFoundOnData;
+	}
+
+	// Signal no conflict found.
+	bOutConflictingTypesFound = false;
+
+	bOutPresentOnAllData = bFoundOnAllData;
+
+	// Signal attribute found or not.
+	return FoundType != EPCGKernelAttributeType::Invalid;
+}
+
+bool FPCGDataCollectionDesc::ContainsAttributeOnAnyData(FPCGAttributeIdentifier InAttributeIdentifier) const
+{
+	return DataDescs.ContainsByPredicate([InAttributeIdentifier](const FPCGDataDesc& DataDesc)
+	{
+		return DataDesc.ContainsAttribute(InAttributeIdentifier);
+	});
+}
+
+void FPCGDataCollectionDesc::AddAttributeToAllData(FPCGKernelAttributeKey InAttribute, const UPCGDataBinding* InBinding, const TArray<int32>* InOptionalUniqueStringKeys)
+{
+	for (FPCGDataDesc& DataDesc : DataDescs)
+	{
+		DataDesc.AddAttribute(InAttribute, InBinding, InOptionalUniqueStringKeys);
+	}
+}
+
+void FPCGDataCollectionDesc::AllocatePropertiesForAllData(EPCGPointNativeProperties InProperties)
+{
+	for (FPCGDataDesc& DataDesc : DataDescs)
+	{
+		DataDesc.AllocateProperties(InProperties);
+	}
+}
+
+void FPCGDataCollectionDesc::GetUniqueStringKeyValues(int32 InAttributeId, TArray<int32>& OutUniqueStringKeys) const
+{
+	for (const FPCGDataDesc& DataDesc : DataDescs)
+	{
+		for (const FPCGKernelAttributeDesc& AttributeDesc : DataDesc.AttributeDescs)
+		{
+			if (AttributeDesc.AttributeId == InAttributeId)
+			{
+				for (int32 StringKey : AttributeDesc.GetUniqueStringKeys())
+				{
+					OutUniqueStringKeys.AddUnique(StringKey);
+				}
+
+				break;
+			}
+		}
+	}
+}
+
+int FPCGDataCollectionDesc::GetNumStringKeyValues(int32 InAttributeId) const
+{
+	TArray<int32> UniqueStringKeys;
+	GetUniqueStringKeyValues(InAttributeId, UniqueStringKeys);
+
+	return UniqueStringKeys.Num();
+}
+
+#undef LOCTEXT_NAMESPACE
diff --git a/Engine/Plugins/PCG/Source/PCG/Private/Compute/PCGDataForGPU.cpp b/Engine/Plugins/PCG/Source/PCG/Private/Compute/PCGDataForGPU.cpp
deleted file mode 100644
diff --git a/Engine/Plugins/PCG/Source/PCG/Private/Compute/PCGDataUploadAdaptor.cpp b/Engine/Plugins/PCG/Source/PCG/Private/Compute/PCGDataUploadAdaptor.cpp
--- a/Engine/Plugins/PCG/Source/PCG/Private/Compute/PCGDataUploadAdaptor.cpp
+++ b/Engine/Plugins/PCG/Source/PCG/Private/Compute/PCGDataUploadAdaptor.cpp
@@ -8,6 +8,7 @@
 #include "Compute/PCGDataBinding.h"
 #include "Compute/PCGKernelHelpers.h"
 #include "Compute/Data/PCGProxyForGPUData.h"
+#include "Compute/Packing/PCGDataCollectionPacking.h"
 
 #include "RenderCaptureInterface.h"
 #include "RenderGraphBuilder.h"
@@ -144,7 +145,7 @@ bool FPCGDataUploadAdaptor::PrepareData_GameThread(const UPCGComputeKernel* InPr
 	}
 	else
 	{
-		const uint64 BufferSize = TargetDataCollectionDesc->ComputePackedSizeBytes();
+		const uint64 BufferSize = PCGDataCollectionPackingHelpers::ComputePackedSizeBytes(TargetDataCollectionDesc);
 		if (PCGComputeHelpers::IsBufferSizeTooLarge(BufferSize))
 		{
 			return true;
@@ -188,7 +189,7 @@ bool FPCGDataUploadAdaptor::PrepareData_GameThread(const UPCGComputeKernel* InPr
 		PCGComputeHelpers::NotifyCPUToGPUUpload(Binding, InProducerKernel, InProducerSettings);
 #endif
 
-		TargetDataCollectionDesc->PackDataCollection(Binding->GetInputDataCollection(), InputPinLabel, Binding, PackedDataCollection);
+		PCGDataCollectionPackingHelpers::PackDataCollection(Binding->GetInputDataCollection(), TargetDataCollectionDesc, InputPinLabel, Binding, PackedDataCollection);
 	}
 
 	return true;
@@ -200,7 +201,7 @@ FRDGBufferSRVRef FPCGDataUploadAdaptor::GetAttributeRemapBufferSRV(FRDGBuilder&
 	check(ExternalBufferForReuse);
 
 	// Mapping created for all non-system attributes (system attributes always have fixed IDs).
-	OutFirstRemappedAttributeId = PCGComputeConstants::NUM_RESERVED_ATTRS;
+	OutFirstRemappedAttributeId = PCGDataCollectionPackingConstants::NUM_RESERVED_ATTRS;
 
 	FRDGBufferRef Buffer = nullptr;
 
diff --git a/Engine/Plugins/PCG/Source/PCG/Private/Compute/PCGDataUploadAdaptor.h b/Engine/Plugins/PCG/Source/PCG/Private/Compute/PCGDataUploadAdaptor.h
--- a/Engine/Plugins/PCG/Source/PCG/Private/Compute/PCGDataUploadAdaptor.h
+++ b/Engine/Plugins/PCG/Source/PCG/Private/Compute/PCGDataUploadAdaptor.h
@@ -2,7 +2,7 @@
 
 #pragma once
 
-#include "Compute/PCGDataForGPU.h"
+#include "Compute/PCGDataDescription.h"
 
 #include "CoreMinimal.h"
 #include "RendererInterface.h"
diff --git a/Engine/Plugins/PCG/Source/PCG/Private/Compute/Packing/PCGDataCollectionPacking.cpp b/Engine/Plugins/PCG/Source/PCG/Private/Compute/Packing/PCGDataCollectionPacking.cpp
new file mode 100644
--- /dev/null
+++ b/Engine/Plugins/PCG/Source/PCG/Private/Compute/Packing/PCGDataCollectionPacking.cpp
@@ -0,0 +1,1349 @@
+// Copyright Epic Games, Inc. All Rights Reserved.
+
+#include "Compute/Packing/PCGDataCollectionPacking.h"
+
+#include "PCGContext.h"
+#include "PCGParamData.h"
+#include "Compute/PCGComputeCommon.h"
+#include "Compute/PCGDataBinding.h"
+#include "Compute/PCGDataDescription.h"
+#include "Compute/PCGPinPropertiesGPU.h"
+#include "Compute/Data/PCGProxyForGPUData.h"
+#include "Data/PCGBasePointData.h"
+#include "Data/PCGPointData.h"
+#include "Elements/Metadata/PCGMetadataElementCommon.h"
+#include "Helpers/PCGAsync.h"
+#include "Metadata/Accessors/PCGAttributeAccessorHelpers.h"
+
+#include "Async/ParallelFor.h"
+
+namespace PCGDataCollectionPackingHelpers
+{
+	uint32 GetElementDataStartAddressUints(const uint32* InPackedDataCollection, uint32 InDataIndex, uint32 InAttributeId)
+	{
+		uint32 ReadAddressBytes = PCGDataCollectionPackingConstants::DATA_COLLECTION_HEADER_SIZE_BYTES + InDataIndex * PCGDataCollectionPackingConstants::DATA_HEADER_SIZE_BYTES;
+		ReadAddressBytes += /*TypeId*/4 + /*Attribute Count*/4 + /*Element Count*/4;
+	
+		ReadAddressBytes += InAttributeId * PCGDataCollectionPackingConstants::ATTRIBUTE_HEADER_SIZE_BYTES;
+		ReadAddressBytes += /*PackedIdAndStride*/4;
+	
+		return InPackedDataCollection[ReadAddressBytes >> 2] >> 2;
+	}
+
+	bool PackAttributeHelper(const FPCGMetadataAttributeBase* InAttributeBase, const FPCGKernelAttributeDesc& InAttributeDesc, PCGMetadataEntryKey InEntryKey, const TArray<FString>& InStringTable, TArray<uint32>& OutPackedDataCollection, uint32& InOutAddressUints)
+	{
+		check(InAttributeBase);
+
+		const PCGMetadataValueKey ValueKey = InAttributeBase->GetValueKey(InEntryKey);
+		const int16 TypeId = InAttributeBase->GetTypeId();
+		const int StrideBytes = PCGDataDescriptionHelpers::GetAttributeTypeStrideBytes(InAttributeDesc.AttributeKey.Type);
+
+		switch (TypeId)
+		{
+		case PCG::Private::MetadataTypes<bool>::Id:
+		{
+			const FPCGMetadataAttribute<bool>* Attribute = static_cast<const FPCGMetadataAttribute<bool>*>(InAttributeBase);
+			const bool Value = Attribute->GetValue(ValueKey);
+			check(StrideBytes == 4);
+			OutPackedDataCollection[InOutAddressUints++] = Value;
+			break;
+		}
+		case PCG::Private::MetadataTypes<float>::Id:
+		{
+			const FPCGMetadataAttribute<float>* Attribute = static_cast<const FPCGMetadataAttribute<float>*>(InAttributeBase);
+			const float Value = Attribute->GetValue(ValueKey);
+			check(StrideBytes == 4);
+			OutPackedDataCollection[InOutAddressUints++] = FMath::AsUInt(Value);
+			break;
+		}
+		case PCG::Private::MetadataTypes<double>::Id:
+		{
+			const FPCGMetadataAttribute<double>* Attribute = static_cast<const FPCGMetadataAttribute<double>*>(InAttributeBase);
+			const double Value = Attribute->GetValue(ValueKey);
+			check(StrideBytes == 4);
+			OutPackedDataCollection[InOutAddressUints++] = FMath::AsUInt(static_cast<float>(Value));
+			break;
+		}
+		case PCG::Private::MetadataTypes<int32>::Id:
+		{
+			const FPCGMetadataAttribute<int32>* Attribute = static_cast<const FPCGMetadataAttribute<int32>*>(InAttributeBase);
+			const int32 Value = Attribute->GetValue(ValueKey);
+			check(StrideBytes == 4);
+			OutPackedDataCollection[InOutAddressUints++] = Value;
+			break;
+		}
+		case PCG::Private::MetadataTypes<int64>::Id:
+		{
+			const FPCGMetadataAttribute<int64>* Attribute = static_cast<const FPCGMetadataAttribute<int64>*>(InAttributeBase);
+			const int64 Value = Attribute->GetValue(ValueKey);
+			check(StrideBytes == 4);
+			OutPackedDataCollection[InOutAddressUints++] = Value;
+			break;
+		}
+		case PCG::Private::MetadataTypes<FVector2D>::Id:
+		{
+			const FPCGMetadataAttribute<FVector2D>* Attribute = static_cast<const FPCGMetadataAttribute<FVector2D>*>(InAttributeBase);
+			const FVector2D Value = Attribute->GetValue(ValueKey);
+			check(StrideBytes == 8);
+			OutPackedDataCollection[InOutAddressUints++] = FMath::AsUInt(static_cast<float>(Value.X));
+			OutPackedDataCollection[InOutAddressUints++] = FMath::AsUInt(static_cast<float>(Value.Y));
+			break;
+		}
+		case PCG::Private::MetadataTypes<FRotator>::Id:
+		{
+			const FPCGMetadataAttribute<FRotator>* Attribute = static_cast<const FPCGMetadataAttribute<FRotator>*>(InAttributeBase);
+			const FRotator Value = Attribute->GetValue(ValueKey);
+			check(StrideBytes == 12);
+			OutPackedDataCollection[InOutAddressUints++] = FMath::AsUInt(static_cast<float>(Value.Pitch));
+			OutPackedDataCollection[InOutAddressUints++] = FMath::AsUInt(static_cast<float>(Value.Yaw));
+			OutPackedDataCollection[InOutAddressUints++] = FMath::AsUInt(static_cast<float>(Value.Roll));
+			break;
+		}
+		case PCG::Private::MetadataTypes<FVector>::Id:
+		{
+			const FPCGMetadataAttribute<FVector>* Attribute = static_cast<const FPCGMetadataAttribute<FVector>*>(InAttributeBase);
+			const FVector Value = Attribute->GetValue(ValueKey);
+			check(StrideBytes == 12);
+			OutPackedDataCollection[InOutAddressUints++] = FMath::AsUInt(static_cast<float>(Value.X));
+			OutPackedDataCollection[InOutAddressUints++] = FMath::AsUInt(static_cast<float>(Value.Y));
+			OutPackedDataCollection[InOutAddressUints++] = FMath::AsUInt(static_cast<float>(Value.Z));
+			break;
+		}
+		case PCG::Private::MetadataTypes<FVector4>::Id:
+		{
+			const FPCGMetadataAttribute<FVector4>* Attribute = static_cast<const FPCGMetadataAttribute<FVector4>*>(InAttributeBase);
+			const FVector4 Value = Attribute->GetValue(ValueKey);
+			check(StrideBytes == 16);
+			OutPackedDataCollection[InOutAddressUints++] = FMath::AsUInt(static_cast<float>(Value.X));
+			OutPackedDataCollection[InOutAddressUints++] = FMath::AsUInt(static_cast<float>(Value.Y));
+			OutPackedDataCollection[InOutAddressUints++] = FMath::AsUInt(static_cast<float>(Value.Z));
+			OutPackedDataCollection[InOutAddressUints++] = FMath::AsUInt(static_cast<float>(Value.W));
+			break;
+		}
+		case PCG::Private::MetadataTypes<FQuat>::Id:
+		{
+			const FPCGMetadataAttribute<FQuat>* Attribute = static_cast<const FPCGMetadataAttribute<FQuat>*>(InAttributeBase);
+			const FQuat Value = Attribute->GetValue(ValueKey);
+			check(StrideBytes == 16);
+			OutPackedDataCollection[InOutAddressUints++] = FMath::AsUInt(static_cast<float>(Value.X));
+			OutPackedDataCollection[InOutAddressUints++] = FMath::AsUInt(static_cast<float>(Value.Y));
+			OutPackedDataCollection[InOutAddressUints++] = FMath::AsUInt(static_cast<float>(Value.Z));
+			OutPackedDataCollection[InOutAddressUints++] = FMath::AsUInt(static_cast<float>(Value.W));
+			break;
+		}
+		case PCG::Private::MetadataTypes<FTransform>::Id:
+		{
+			const FPCGMetadataAttribute<FTransform>* Attribute = static_cast<const FPCGMetadataAttribute<FTransform>*>(InAttributeBase);
+			const FTransform Transform = Attribute->GetValue(ValueKey);
+
+			const bool bIsRotationNormalized = Transform.IsRotationNormalized();
+			if (!bIsRotationNormalized)
+			{
+				UE_LOG(LogPCG, Error, TEXT("Tried to pack attribute '%s' of type Transform for GPU data collection, but the transform's rotation is not normalized (%f, %f, %f, %f). Using identity instead."),
+					*InAttributeBase->Name.ToString(), Transform.GetRotation().X, Transform.GetRotation().Y, Transform.GetRotation().Z, Transform.GetRotation().W);
+			}
+
+			// Note: ToMatrixWithScale() crashes if the transform is not normalized.
+			const FMatrix Matrix = bIsRotationNormalized ? Transform.ToMatrixWithScale() : FMatrix::Identity;
+
+			check(StrideBytes == 64);
+			OutPackedDataCollection[InOutAddressUints++] = FMath::AsUInt(static_cast<float>(Matrix.M[0][0]));
+			OutPackedDataCollection[InOutAddressUints++] = FMath::AsUInt(static_cast<float>(Matrix.M[1][0]));
+			OutPackedDataCollection[InOutAddressUints++] = FMath::AsUInt(static_cast<float>(Matrix.M[2][0]));
+			OutPackedDataCollection[InOutAddressUints++] = FMath::AsUInt(static_cast<float>(Matrix.M[3][0]));
+			OutPackedDataCollection[InOutAddressUints++] = FMath::AsUInt(static_cast<float>(Matrix.M[0][1]));
+			OutPackedDataCollection[InOutAddressUints++] = FMath::AsUInt(static_cast<float>(Matrix.M[1][1]));
+			OutPackedDataCollection[InOutAddressUints++] = FMath::AsUInt(static_cast<float>(Matrix.M[2][1]));
+			OutPackedDataCollection[InOutAddressUints++] = FMath::AsUInt(static_cast<float>(Matrix.M[3][1]));
+			OutPackedDataCollection[InOutAddressUints++] = FMath::AsUInt(static_cast<float>(Matrix.M[0][2]));
+			OutPackedDataCollection[InOutAddressUints++] = FMath::AsUInt(static_cast<float>(Matrix.M[1][2]));
+			OutPackedDataCollection[InOutAddressUints++] = FMath::AsUInt(static_cast<float>(Matrix.M[2][2]));
+			OutPackedDataCollection[InOutAddressUints++] = FMath::AsUInt(static_cast<float>(Matrix.M[3][2]));
+			OutPackedDataCollection[InOutAddressUints++] = FMath::AsUInt(static_cast<float>(Matrix.M[0][3]));
+			OutPackedDataCollection[InOutAddressUints++] = FMath::AsUInt(static_cast<float>(Matrix.M[1][3]));
+			OutPackedDataCollection[InOutAddressUints++] = FMath::AsUInt(static_cast<float>(Matrix.M[2][3]));
+			OutPackedDataCollection[InOutAddressUints++] = FMath::AsUInt(static_cast<float>(Matrix.M[3][3]));
+
+			break;
+		}
+		case PCG::Private::MetadataTypes<FString>::Id:
+		{
+			// String stored as an integer for reading/writing in kernel, and accompanying string table in data description.
+			const FPCGMetadataAttribute<FString>* Attribute = static_cast<const FPCGMetadataAttribute<FString>*>(InAttributeBase);
+			const int32 Value = InStringTable.IndexOfByKey(Attribute->GetValue(ValueKey));
+			check(StrideBytes == 4);
+			OutPackedDataCollection[InOutAddressUints++] = Value;
+			break;
+		}
+		case PCG::Private::MetadataTypes<FSoftObjectPath>::Id:
+		{
+			// SOP path string stored as an integer for reading/writing in kernel, and accompanying string table in data description.
+			const FPCGMetadataAttribute<FSoftObjectPath>* Attribute = static_cast<const FPCGMetadataAttribute<FSoftObjectPath>*>(InAttributeBase);
+			const int32 Value = InStringTable.IndexOfByKey(Attribute->GetValue(ValueKey).ToString());
+			check(StrideBytes == 4);
+			OutPackedDataCollection[InOutAddressUints++] = Value;
+			break;
+		}
+		case PCG::Private::MetadataTypes<FSoftClassPath>::Id:
+		{
+			// SCP path string stored as an integer for reading/writing in kernel, and accompanying string table in data description.
+			const FPCGMetadataAttribute<FSoftClassPath>* Attribute = static_cast<const FPCGMetadataAttribute<FSoftClassPath>*>(InAttributeBase);
+			const int32 Value = InStringTable.IndexOfByKey(Attribute->GetValue(ValueKey).ToString());
+			check(StrideBytes == 4);
+			OutPackedDataCollection[InOutAddressUints++] = Value;
+			break;
+		}
+		case PCG::Private::MetadataTypes<FName>::Id:
+		{
+			// FNames are currently stored in string table so use same logic as string.
+			const FPCGMetadataAttribute<FName>* Attribute = static_cast<const FPCGMetadataAttribute<FName>*>(InAttributeBase);
+			const int32 Value = InStringTable.IndexOfByKey(Attribute->GetValue(ValueKey).ToString());
+			check(StrideBytes == 4);
+			OutPackedDataCollection[InOutAddressUints++] = Value;
+			break;
+		}
+		default:
+			return false;
+		}
+
+		return true;
+	}
+
+	FPCGMetadataAttributeBase* CreateAttributeFromAttributeDesc(UPCGMetadata* Metadata, const FPCGKernelAttributeDesc& AttributeDesc)
+	{
+		check(Metadata);
+		FPCGMetadataDomain* MetadataDomain = Metadata->GetMetadataDomain(AttributeDesc.AttributeKey.Identifier.MetadataDomain);
+		if (!MetadataDomain)
+		{
+			return nullptr;
+		}
+		
+		switch (AttributeDesc.AttributeKey.Type)
+		{
+		case EPCGKernelAttributeType::Bool:
+		{
+			return MetadataDomain->FindOrCreateAttribute<bool>(AttributeDesc.AttributeKey.Identifier.Name);
+		}
+		case EPCGKernelAttributeType::Int:
+		{
+			return MetadataDomain->FindOrCreateAttribute<int>(AttributeDesc.AttributeKey.Identifier.Name);
+		}
+		case EPCGKernelAttributeType::Float:
+		{
+			return MetadataDomain->FindOrCreateAttribute<float>(AttributeDesc.AttributeKey.Identifier.Name);
+		}
+		case EPCGKernelAttributeType::Float2:
+		{
+			return MetadataDomain->FindOrCreateAttribute<FVector2D>(AttributeDesc.AttributeKey.Identifier.Name);
+		}
+		case EPCGKernelAttributeType::Float3:
+		{
+			return MetadataDomain->FindOrCreateAttribute<FVector>(AttributeDesc.AttributeKey.Identifier.Name);
+		}
+		case EPCGKernelAttributeType::Float4:
+		{
+			return MetadataDomain->FindOrCreateAttribute<FVector4>(AttributeDesc.AttributeKey.Identifier.Name);
+		}
+		case EPCGKernelAttributeType::Rotator:
+		{
+			return MetadataDomain->FindOrCreateAttribute<FRotator>(AttributeDesc.AttributeKey.Identifier.Name);
+		}
+		case EPCGKernelAttributeType::Quat:
+		{
+			return MetadataDomain->FindOrCreateAttribute<FQuat>(AttributeDesc.AttributeKey.Identifier.Name);
+		}
+		case EPCGKernelAttributeType::Transform:
+		{
+			return MetadataDomain->FindOrCreateAttribute<FTransform>(AttributeDesc.AttributeKey.Identifier.Name);
+		}
+		case EPCGKernelAttributeType::StringKey:
+		{
+			return MetadataDomain->FindOrCreateAttribute<FString>(AttributeDesc.AttributeKey.Identifier.Name);
+		}
+		case EPCGKernelAttributeType::Name:
+		{
+			return MetadataDomain->FindOrCreateAttribute<FName>(AttributeDesc.AttributeKey.Identifier.Name);
+		}
+		default:
+			return nullptr;
+		}
+	}
+
+	bool UnpackAttributeHelper(FPCGContext* InContext, const void* InPackedData, const FPCGKernelAttributeDesc& InAttributeDesc, const TArray<FString>& InStringTable, uint32 InAddressUints, uint32 InNumElements, UPCGData* OutData)
+	{
+		check(InPackedData);
+		check(OutData);
+
+		const float* DataAsFloat = static_cast<const float*>(InPackedData);
+		const int32* DataAsInt = static_cast<const int32*>(InPackedData);
+
+		FPCGAttributePropertyOutputSelector Selector = FPCGAttributePropertySelector::CreateAttributeSelector<FPCGAttributePropertyOutputSelector>(InAttributeDesc.AttributeKey.Identifier.Name);
+		OutData->SetDomainFromDomainID(InAttributeDesc.AttributeKey.Identifier.MetadataDomain, Selector);
+
+		switch (InAttributeDesc.AttributeKey.Type)
+		{
+		case EPCGKernelAttributeType::Bool:
+		{
+			TArray<bool> Values;
+			Values.SetNumUninitialized(InNumElements);
+
+			for (uint32 ElementIndex = 0; ElementIndex < InNumElements; ++ElementIndex)
+			{
+				const uint32 PackedElementIndex = InAddressUints + ElementIndex;
+				Values[ElementIndex] = static_cast<bool>(DataAsFloat[PackedElementIndex]);
+			}
+
+			ensure(PCGAttributeAccessorHelpers::WriteAllValues<bool>(OutData, Selector, Values, /*SourceSelector=*/nullptr, InContext));
+			break;
+		}
+		case EPCGKernelAttributeType::Int:
+		{
+			TArray<int> Values;
+			Values.SetNumUninitialized(InNumElements);
+
+			for (uint32 ElementIndex = 0; ElementIndex < InNumElements; ++ElementIndex)
+			{
+				const uint32 PackedElementIndex = InAddressUints + ElementIndex;
+				Values[ElementIndex] = DataAsInt[PackedElementIndex];
+			}
+
+			ensure(PCGAttributeAccessorHelpers::WriteAllValues<int>(OutData, Selector, Values, /*SourceSelector=*/nullptr, InContext));
+			break;
+		}
+		case EPCGKernelAttributeType::Float:
+		{
+			TArray<float> Values;
+			Values.SetNumUninitialized(InNumElements);
+
+			for (uint32 ElementIndex = 0; ElementIndex < InNumElements; ++ElementIndex)
+			{
+				const uint32 PackedElementIndex = InAddressUints + ElementIndex;
+				Values[ElementIndex] = DataAsFloat[PackedElementIndex];
+			}
+
+			ensure(PCGAttributeAccessorHelpers::WriteAllValues<float>(OutData, Selector, Values, /*SourceSelector=*/nullptr, InContext));
+			break;
+		}
+		case EPCGKernelAttributeType::Float2:
+		{
+			TArray<FVector2D> Values;
+			Values.SetNumUninitialized(InNumElements);
+
+			for (uint32 ElementIndex = 0; ElementIndex < InNumElements; ++ElementIndex)
+			{
+				const uint32 PackedElementIndex = InAddressUints + ElementIndex * 2;
+				Values[ElementIndex].X = DataAsFloat[PackedElementIndex + 0];
+				Values[ElementIndex].Y = DataAsFloat[PackedElementIndex + 1];
+			}
+
+			ensure(PCGAttributeAccessorHelpers::WriteAllValues<FVector2D>(OutData, Selector, Values, /*SourceSelector=*/nullptr, InContext));
+			break;
+		}
+		case EPCGKernelAttributeType::Float3:
+		{
+			TArray<FVector> Values;
+			Values.SetNumUninitialized(InNumElements);
+
+			for (uint32 ElementIndex = 0; ElementIndex < InNumElements; ++ElementIndex)
+			{
+				const uint32 PackedElementIndex = InAddressUints + ElementIndex * 3;
+				Values[ElementIndex].X = DataAsFloat[PackedElementIndex + 0];
+				Values[ElementIndex].Y = DataAsFloat[PackedElementIndex + 1];
+				Values[ElementIndex].Z = DataAsFloat[PackedElementIndex + 2];
+			}
+
+			ensure(PCGAttributeAccessorHelpers::WriteAllValues<FVector>(OutData, Selector, Values, /*SourceSelector=*/nullptr, InContext));
+			break;
+		}
+		case EPCGKernelAttributeType::Float4:
+		{
+			TArray<FVector4> Values;
+			Values.SetNumUninitialized(InNumElements);
+
+			for (uint32 ElementIndex = 0; ElementIndex < InNumElements; ++ElementIndex)
+			{
+				const uint32 PackedElementIndex = InAddressUints + ElementIndex * 4;
+				Values[ElementIndex].X = DataAsFloat[PackedElementIndex + 0];
+				Values[ElementIndex].Y = DataAsFloat[PackedElementIndex + 1];
+				Values[ElementIndex].Z = DataAsFloat[PackedElementIndex + 2];
+				Values[ElementIndex].W = DataAsFloat[PackedElementIndex + 3];
+			}
+
+			ensure(PCGAttributeAccessorHelpers::WriteAllValues<FVector4>(OutData, Selector, Values, /*SourceSelector=*/nullptr, InContext));
+			break;
+		}
+		case EPCGKernelAttributeType::Rotator:
+		{
+			TArray<FRotator> Values;
+			Values.SetNumUninitialized(InNumElements);
+
+			for (uint32 ElementIndex = 0; ElementIndex < InNumElements; ++ElementIndex)
+			{
+				const uint32 PackedElementIndex = InAddressUints + ElementIndex * 3;
+				Values[ElementIndex].Pitch = DataAsFloat[PackedElementIndex + 0];
+				Values[ElementIndex].Yaw = DataAsFloat[PackedElementIndex + 1];
+				Values[ElementIndex].Roll = DataAsFloat[PackedElementIndex + 2];
+			}
+
+			ensure(PCGAttributeAccessorHelpers::WriteAllValues<FRotator>(OutData, Selector, Values, /*SourceSelector=*/nullptr, InContext));
+			break;
+		}
+		case EPCGKernelAttributeType::Quat:
+		{
+			TArray<FQuat> Values;
+			Values.SetNumUninitialized(InNumElements);
+
+			for (uint32 ElementIndex = 0; ElementIndex < InNumElements; ++ElementIndex)
+			{
+				const uint32 PackedElementIndex = InAddressUints + ElementIndex * 4;
+				Values[ElementIndex].X = DataAsFloat[PackedElementIndex + 0];
+				Values[ElementIndex].Y = DataAsFloat[PackedElementIndex + 1];
+				Values[ElementIndex].Z = DataAsFloat[PackedElementIndex + 2];
+				Values[ElementIndex].W = DataAsFloat[PackedElementIndex + 3];
+			}
+
+			ensure(PCGAttributeAccessorHelpers::WriteAllValues<FQuat>(OutData, Selector, Values, /*SourceSelector=*/nullptr, InContext));
+			break;
+		}
+		case EPCGKernelAttributeType::Transform:
+		{
+			TArray<FTransform> Values;
+			Values.SetNumUninitialized(InNumElements);
+
+			FMatrix Matrix;
+
+			for (uint32 ElementIndex = 0; ElementIndex < InNumElements; ++ElementIndex)
+			{
+				const uint32 PackedElementIndex = InAddressUints + ElementIndex * 16;
+
+				Matrix.M[0][0] = DataAsFloat[PackedElementIndex + 0];
+				Matrix.M[1][0] = DataAsFloat[PackedElementIndex + 1];
+				Matrix.M[2][0] = DataAsFloat[PackedElementIndex + 2];
+				Matrix.M[3][0] = DataAsFloat[PackedElementIndex + 3];
+				Matrix.M[0][1] = DataAsFloat[PackedElementIndex + 4];
+				Matrix.M[1][1] = DataAsFloat[PackedElementIndex + 5];
+				Matrix.M[2][1] = DataAsFloat[PackedElementIndex + 6];
+				Matrix.M[3][1] = DataAsFloat[PackedElementIndex + 7];
+				Matrix.M[0][2] = DataAsFloat[PackedElementIndex + 8];
+				Matrix.M[1][2] = DataAsFloat[PackedElementIndex + 9];
+				Matrix.M[2][2] = DataAsFloat[PackedElementIndex + 10];
+				Matrix.M[3][2] = DataAsFloat[PackedElementIndex + 11];
+				Matrix.M[0][3] = DataAsFloat[PackedElementIndex + 12];
+				Matrix.M[1][3] = DataAsFloat[PackedElementIndex + 13];
+				Matrix.M[2][3] = DataAsFloat[PackedElementIndex + 14];
+				Matrix.M[3][3] = DataAsFloat[PackedElementIndex + 15];
+
+				new (&Values[ElementIndex]) FTransform(Matrix);
+			}
+
+			ensure(PCGAttributeAccessorHelpers::WriteAllValues<FTransform>(OutData, Selector, Values, /*SourceSelector=*/nullptr, InContext));
+			break;
+		}
+		case EPCGKernelAttributeType::StringKey:
+		{
+			check(!InStringTable.IsEmpty());
+
+			TArray<FString> Values;
+			Values.Reserve(InNumElements);
+
+			for (uint32 ElementIndex = 0; ElementIndex < InNumElements; ++ElementIndex)
+			{
+				const uint32 PackedElementIndex = InAddressUints + ElementIndex;
+				const int32 StringKey = InStringTable.IsValidIndex(DataAsInt[PackedElementIndex]) ? DataAsInt[PackedElementIndex] : 0;
+				Values.Add(InStringTable[StringKey]);
+			}
+
+			ensure(PCGAttributeAccessorHelpers::WriteAllValues<FString>(OutData, Selector, Values, /*SourceSelector=*/nullptr, InContext));
+			break;
+		}
+		case EPCGKernelAttributeType::Name:
+		{
+			check(!InStringTable.IsEmpty());
+
+			TArray<FName> Values;
+			Values.SetNumUninitialized(InNumElements);
+
+			for (uint32 ElementIndex = 0; ElementIndex < InNumElements; ++ElementIndex)
+			{
+				const uint32 PackedElementIndex = InAddressUints + ElementIndex;
+
+				// FNames currently stored in string table.
+				const int32 StringKey = InStringTable.IsValidIndex(DataAsInt[PackedElementIndex]) ? DataAsInt[PackedElementIndex] : 0;
+				Values[ElementIndex] = *InStringTable[StringKey];
+			}
+
+			ensure(PCGAttributeAccessorHelpers::WriteAllValues<FName>(OutData, Selector, Values, /*SourceSelector=*/nullptr, InContext));
+			break;
+		}
+		default:
+			return false;
+		}
+		return true;
+	}
+
+	uint64 ComputePackedSizeBytes(const FPCGDataDesc& InDataDesc)
+	{
+		ensure(InDataDesc.GetElementDimension() == EPCGElementDimension::One);
+
+		const FPCGDataTypeIdentifier& Type = InDataDesc.GetType();
+
+		check(PCGComputeHelpers::IsTypeAllowedInDataCollection(Type));
+		uint64 DataSizeBytes = PCGDataCollectionPackingConstants::DATA_HEADER_SIZE_BYTES;
+
+		for (const FPCGKernelAttributeDesc& AttributeDesc : InDataDesc.AttributeDescs)
+		{
+			const uint64 NumValues = static_cast<uint64>(InDataDesc.GetElementCountForAttribute(AttributeDesc).X);
+			DataSizeBytes += static_cast<uint64>(PCGDataDescriptionHelpers::GetAttributeTypeStrideBytes(AttributeDesc.AttributeKey.Type)) * NumValues;
+		}
+
+		return DataSizeBytes;
+	}
+
+	uint32 ComputePackedHeaderSizeBytes(TSharedPtr<const FPCGDataCollectionDesc> InDataCollectionDesc)
+	{
+		return InDataCollectionDesc ? PCGDataCollectionPackingConstants::DATA_COLLECTION_HEADER_SIZE_BYTES + PCGDataCollectionPackingConstants::DATA_HEADER_SIZE_BYTES * InDataCollectionDesc->DataDescs.Num() : 0;
+	}
+	
+	uint64 ComputePackedSizeBytes(TSharedPtr<const FPCGDataCollectionDesc> InDataCollectionDesc)
+	{
+		TRACE_CPUPROFILER_EVENT_SCOPE(PCGDataCollectionPacking::ComputePackedSizeBytes);
+
+		if (!InDataCollectionDesc)
+		{
+			return 0;
+		}
+	
+		bool bRequiresHeader = false;
+		uint64 TotalDataSizeBytes = 0;
+
+		const TArray<FPCGDataDesc>& DataDescs = InDataCollectionDesc->DataDescs;
+	
+		for (const FPCGDataDesc& DataDesc : DataDescs)
+		{
+			bRequiresHeader |= PCGComputeHelpers::IsTypeAllowedInDataCollection(DataDesc.GetType());
+			TotalDataSizeBytes += ComputePackedSizeBytes(DataDesc);
+		}
+	
+		// Assumption: If no data items then pack as empty data collection, so include header.
+		if (bRequiresHeader || DataDescs.IsEmpty())
+		{
+			TotalDataSizeBytes += ComputePackedHeaderSizeBytes(InDataCollectionDesc);
+		}
+	
+		return TotalDataSizeBytes;
+	}
+	
+	void WriteHeader(TSharedPtr<const FPCGDataCollectionDesc> InDataCollectionDesc, TArray<uint32>& OutPackedDataCollectionHeader)
+	{
+		TRACE_CPUPROFILER_EVENT_SCOPE(PCGDataCollectionPacking::WriteHeader);
+
+		if (!InDataCollectionDesc)
+		{
+			return;
+		}
+
+		const TArray<FPCGDataDesc>& DataDescs = InDataCollectionDesc->DataDescs;
+	
+		const uint32 HeaderSizeBytes = ComputePackedHeaderSizeBytes(InDataCollectionDesc);
+		const uint32 HeaderSizeUints = HeaderSizeBytes >> 2;
+	
+		if (OutPackedDataCollectionHeader.Num() < static_cast<int32>(HeaderSizeUints))
+		{
+			OutPackedDataCollectionHeader.SetNumUninitialized(HeaderSizeUints);
+		}
+	
+		// Zero-initialize header portion. We detect absent attributes using 0s.
+		for (uint32 Index = 0; Index < HeaderSizeUints; ++Index)
+		{
+			OutPackedDataCollectionHeader[Index] = 0;
+		}
+	
+		uint32 WriteAddressUints = 0;
+	
+		// Start at end of header
+		uint32 DataStartAddressBytes = HeaderSizeBytes;
+	
+		// Num data
+		OutPackedDataCollectionHeader[WriteAddressUints++] = DataDescs.Num();
+	
+		for (int32 DataIndex = 0; DataIndex < DataDescs.Num(); ++DataIndex)
+		{
+			const FPCGDataDesc& DataDesc = DataDescs[DataIndex];
+			ensure(DataDesc.GetElementDimension() == EPCGElementDimension::One);
+	
+			// Data i: type ID
+			if (DataDesc.GetType() == EPCGDataType::Param)
+			{
+				OutPackedDataCollectionHeader[WriteAddressUints++] = PCGDataCollectionPackingConstants::PARAM_DATA_TYPE_ID;
+			}
+			else
+			{
+				ensure(DataDesc.GetType() == EPCGDataType::Point);
+				OutPackedDataCollectionHeader[WriteAddressUints++] = PCGDataCollectionPackingConstants::POINT_DATA_TYPE_ID;
+			}
+	
+			// Data i: attribute count (including intrinsic point properties)
+			OutPackedDataCollectionHeader[WriteAddressUints++] = DataDesc.AttributeDescs.Num();
+	
+			// Data i: element count
+			OutPackedDataCollectionHeader[WriteAddressUints++] = DataDesc.GetElementCount().X;
+	
+			const uint32 DataAttributesHeaderStartAddressBytes = WriteAddressUints << 2;
+	
+			for (int32 AttrIndex = 0; AttrIndex < DataDesc.AttributeDescs.Num(); ++AttrIndex)
+			{
+				const FPCGKernelAttributeDesc& AttributeDesc = DataDesc.AttributeDescs[AttrIndex];
+				const int AttributeStride = PCGDataDescriptionHelpers::GetAttributeTypeStrideBytes(AttributeDesc.AttributeKey.Type);
+				const int AttributeElementsCount = DataDesc.GetElementCountForAttribute(AttributeDesc).X;
+	
+				// Scatter from attributes that are present into header which has slots for all possible attributes.
+				WriteAddressUints = (AttributeDesc.AttributeId * PCGDataCollectionPackingConstants::ATTRIBUTE_HEADER_SIZE_BYTES + DataAttributesHeaderStartAddressBytes) >> 2;
+	
+				// Data i element j: packed attribute info and data address.
+				const uint32 PackedAttributeInfo = (DataDesc.IsAttributeAllocated(AttributeDesc) ? PCGDataCollectionPackingConstants::AttributeAllocatedMask : 0) + (AttributeDesc.AttributeId << 8) + AttributeStride;
+				OutPackedDataCollectionHeader[WriteAddressUints++] = PackedAttributeInfo;
+				OutPackedDataCollectionHeader[WriteAddressUints++] = DataStartAddressBytes;
+	
+				// Move the DataStartAddress to the beginning of the next attribute by skipping all the bytes of the current attribute.
+				DataStartAddressBytes += AttributeElementsCount * AttributeStride;
+			}
+	
+			// After scattering in attribute headers, fast forward to end of section.
+			WriteAddressUints = (PCGDataCollectionPackingConstants::MAX_NUM_ATTRS * PCGDataCollectionPackingConstants::ATTRIBUTE_HEADER_SIZE_BYTES + DataAttributesHeaderStartAddressBytes) >> 2;
+		}
+	
+		check(WriteAddressUints * 4 == HeaderSizeBytes);
+	}
+	
+	void PackDataCollection(const FPCGDataCollection& InDataCollection, TSharedPtr<const FPCGDataCollectionDesc> InDataCollectionDesc, FName InPin, const UPCGDataBinding* InDataBinding, TArray<uint32>& OutPackedDataCollection)
+	{
+		TRACE_CPUPROFILER_EVENT_SCOPE(PCGDataCollectionPacking::PackDataCollection);
+
+		if (!ensure(InDataCollectionDesc))
+		{
+			return;
+		}
+
+		const TArray<FPCGDataDesc>& DataDescs = InDataCollectionDesc->DataDescs;
+		const TArray<FPCGTaggedData> InputData = InDataCollection.GetInputsByPin(InPin);
+	
+		for (const FPCGTaggedData& TaggedData : InputData)
+		{
+			if (TaggedData.Data && !PCGComputeHelpers::IsTypeAllowedInDataCollection(TaggedData.Data->GetDataType()))
+			{
+				UE_LOG(LogPCG, Error,
+					TEXT("PackDataCollection: Encountered data '%s' which does not support packing into data collections. ")
+					TEXT("Ensure the correct data interface has been created for the pin ('%s'). Packing will fail."),
+					*TaggedData.Data->GetName(), *InPin.ToString());
+				return;
+			}
+		}
+	
+		const uint32 PackedDataCollectionSizeBytes = ComputePackedSizeBytes(InDataCollectionDesc);
+		check(PackedDataCollectionSizeBytes >= 0);
+	
+		// Uninitialized is fine, all data is initialized explicitly.
+		OutPackedDataCollection.SetNumUninitialized(PackedDataCollectionSizeBytes >> 2);
+	
+	#if !UE_BUILD_SHIPPING
+		if (PCGSystemSwitches::CVarFuzzGPUMemory.GetValueOnAnyThread())
+		{
+			FRandomStream Rand(GFrameCounter);
+			for (uint32& Value : OutPackedDataCollection)
+			{
+				Value = Rand.GetUnsignedInt();
+			}
+		}
+	#endif // !UE_BUILD_SHIPPING
+	
+		// Data addresses are written to the header and will be used during packing below.
+		WriteHeader(InDataCollectionDesc, OutPackedDataCollection);
+	
+		for (int32 DataIndex = 0; DataIndex < InputData.Num(); ++DataIndex)
+		{
+			const FPCGDataDesc& DataDesc = DataDescs[DataIndex];
+			ensure(DataDesc.GetElementDimension() == EPCGElementDimension::One);
+	
+			if (Cast<UPCGProxyForGPUData>(InputData[DataIndex].Data))
+			{
+				UE_LOG(LogPCG, Error, TEXT("Attempted to pack a data that is not resident on the CPU. Uploaded data will be uninitialized."));
+				continue;
+			}
+	
+			// No work to do if there are no elements to process.
+			if (DataDesc.GetElementCount().X <= 0)
+			{
+				continue;
+			}
+	
+			const UPCGMetadata* Metadata = InputData[DataIndex].Data ? InputData[DataIndex].Data->ConstMetadata() : nullptr;
+			if (!ensure(Metadata))
+			{
+				continue;
+			}
+	
+			if (const UPCGBasePointData* PointData = Cast<UPCGBasePointData>(InputData[DataIndex].Data))
+			{
+				const uint32 NumElements = PointData->GetNumPoints();
+	
+				if (NumElements == 0)
+				{
+					continue;
+				}
+	
+				for (const FPCGKernelAttributeDesc& AttributeDesc : DataDesc.AttributeDescs)
+				{
+					const uint32 AttributeId = AttributeDesc.AttributeId;
+	
+					const FPCGMetadataDomain* MetadataDomain = Metadata->GetConstMetadataDomain(AttributeDesc.AttributeKey.Identifier.MetadataDomain);
+					const FPCGMetadataAttributeBase* AttributeBase = (AttributeId >= PCGDataCollectionPackingConstants::NUM_RESERVED_ATTRS && MetadataDomain) ? MetadataDomain->GetConstAttribute(AttributeDesc.AttributeKey.Identifier.Name) : nullptr;
+	
+					uint32 AddressUints = GetElementDataStartAddressUints(OutPackedDataCollection.GetData(), DataIndex, AttributeId);
+	
+					if (AttributeId < PCGDataCollectionPackingConstants::NUM_RESERVED_ATTRS)
+					{
+						const uint32 NumElementsForAttribute = DataDesc.IsAttributeAllocated(AttributeDesc) ? NumElements : 1u;
+	
+						// Point property.
+						switch (AttributeId)
+						{
+						case PCGDataCollectionPackingConstants::POINT_POSITION_ATTRIBUTE_ID:
+						{
+							const TConstPCGValueRange<FTransform> ReadValueRange = PointData->GetConstTransformValueRange();
+	
+							for (uint32 ElementIndex = 0; ElementIndex < NumElementsForAttribute; ++ElementIndex)
+							{
+								const FVector Position = ReadValueRange[ElementIndex].GetLocation();
+								OutPackedDataCollection[AddressUints++] = FMath::AsUInt(static_cast<float>(Position.X));
+								OutPackedDataCollection[AddressUints++] = FMath::AsUInt(static_cast<float>(Position.Y));
+								OutPackedDataCollection[AddressUints++] = FMath::AsUInt(static_cast<float>(Position.Z));
+							}
+							break;
+						}
+						case PCGDataCollectionPackingConstants::POINT_ROTATION_ATTRIBUTE_ID:
+						{
+							const TConstPCGValueRange<FTransform> ReadValueRange = PointData->GetConstTransformValueRange();
+	
+							for (uint32 ElementIndex = 0; ElementIndex < NumElementsForAttribute; ++ElementIndex)
+							{
+								const FQuat Rotation = ReadValueRange[ElementIndex].GetRotation();
+								OutPackedDataCollection[AddressUints++] = FMath::AsUInt(static_cast<float>(Rotation.X));
+								OutPackedDataCollection[AddressUints++] = FMath::AsUInt(static_cast<float>(Rotation.Y));
+								OutPackedDataCollection[AddressUints++] = FMath::AsUInt(static_cast<float>(Rotation.Z));
+								OutPackedDataCollection[AddressUints++] = FMath::AsUInt(static_cast<float>(Rotation.W));
+							}
+							break;
+						}
+						case PCGDataCollectionPackingConstants::POINT_SCALE_ATTRIBUTE_ID:
+						{
+							const TConstPCGValueRange<FTransform> ReadValueRange = PointData->GetConstTransformValueRange();
+	
+							for (uint32 ElementIndex = 0; ElementIndex < NumElementsForAttribute; ++ElementIndex)
+							{
+								const FVector Scale = ReadValueRange[ElementIndex].GetScale3D();
+								OutPackedDataCollection[AddressUints++] = FMath::AsUInt(static_cast<float>(Scale.X));
+								OutPackedDataCollection[AddressUints++] = FMath::AsUInt(static_cast<float>(Scale.Y));
+								OutPackedDataCollection[AddressUints++] = FMath::AsUInt(static_cast<float>(Scale.Z));
+							}
+							break;
+						}
+						case PCGDataCollectionPackingConstants::POINT_BOUNDS_MIN_ATTRIBUTE_ID:
+						{
+							const TConstPCGValueRange<FVector> ReadValueRange = PointData->GetConstBoundsMinValueRange();
+	
+							for (uint32 ElementIndex = 0; ElementIndex < NumElementsForAttribute; ++ElementIndex)
+							{
+								const FVector& BoundsMin = ReadValueRange[ElementIndex];
+								OutPackedDataCollection[AddressUints++] = FMath::AsUInt(static_cast<float>(BoundsMin.X));
+								OutPackedDataCollection[AddressUints++] = FMath::AsUInt(static_cast<float>(BoundsMin.Y));
+								OutPackedDataCollection[AddressUints++] = FMath::AsUInt(static_cast<float>(BoundsMin.Z));
+							}
+							break;
+						}
+						case PCGDataCollectionPackingConstants::POINT_BOUNDS_MAX_ATTRIBUTE_ID:
+						{
+							const TConstPCGValueRange<FVector> ReadValueRange = PointData->GetConstBoundsMaxValueRange();
+	
+							for (uint32 ElementIndex = 0; ElementIndex < NumElementsForAttribute; ++ElementIndex)
+							{
+								const FVector& BoundsMax = ReadValueRange[ElementIndex];
+								OutPackedDataCollection[AddressUints++] = FMath::AsUInt(static_cast<float>(BoundsMax.X));
+								OutPackedDataCollection[AddressUints++] = FMath::AsUInt(static_cast<float>(BoundsMax.Y));
+								OutPackedDataCollection[AddressUints++] = FMath::AsUInt(static_cast<float>(BoundsMax.Z));
+							}
+							break;
+						}
+						case PCGDataCollectionPackingConstants::POINT_COLOR_ATTRIBUTE_ID:
+						{
+							const TConstPCGValueRange<FVector4> ReadValueRange = PointData->GetConstColorValueRange();
+	
+							for (uint32 ElementIndex = 0; ElementIndex < NumElementsForAttribute; ++ElementIndex)
+							{
+								const FVector4& Color = ReadValueRange[ElementIndex];
+								OutPackedDataCollection[AddressUints++] = FMath::AsUInt(static_cast<float>(Color.X));
+								OutPackedDataCollection[AddressUints++] = FMath::AsUInt(static_cast<float>(Color.Y));
+								OutPackedDataCollection[AddressUints++] = FMath::AsUInt(static_cast<float>(Color.Z));
+								OutPackedDataCollection[AddressUints++] = FMath::AsUInt(static_cast<float>(Color.W));
+							}
+							break;
+						}
+						case PCGDataCollectionPackingConstants::POINT_DENSITY_ATTRIBUTE_ID:
+						{
+							const TConstPCGValueRange<float> ReadValueRange = PointData->GetConstDensityValueRange();
+	
+							for (uint32 ElementIndex = 0; ElementIndex < NumElementsForAttribute; ++ElementIndex)
+							{
+								const float Density = ReadValueRange[ElementIndex];
+								OutPackedDataCollection[AddressUints++] = FMath::AsUInt(Density);
+							}
+							break;
+						}
+						case PCGDataCollectionPackingConstants::POINT_SEED_ATTRIBUTE_ID:
+						{
+							const TConstPCGValueRange<int32> ReadValueRange = PointData->GetConstSeedValueRange();
+	
+							for (uint32 ElementIndex = 0; ElementIndex < NumElementsForAttribute; ++ElementIndex)
+							{
+								const int Seed = ReadValueRange[ElementIndex];
+								OutPackedDataCollection[AddressUints++] = Seed;
+							}
+							break;
+						}
+						case PCGDataCollectionPackingConstants::POINT_STEEPNESS_ATTRIBUTE_ID:
+						{
+							const TConstPCGValueRange<float> ReadValueRange = PointData->GetConstSteepnessValueRange();
+	
+							for (uint32 ElementIndex = 0; ElementIndex < NumElementsForAttribute; ++ElementIndex)
+							{
+								const float Steepness = ReadValueRange[ElementIndex];
+								OutPackedDataCollection[AddressUints++] = FMath::AsUInt(Steepness);
+							}
+							break;
+						}
+						default:
+							checkNoEntry();
+							break;
+						}
+					}
+					else
+					{
+						// Pack attribute. Validate first element only for perf.
+						if (AttributeBase->GetMetadataDomain()->GetDomainID() == PCGMetadataDomainID::Data)
+						{
+							ensure(PackAttributeHelper(AttributeBase, AttributeDesc, /*InEntryKey*/0, InDataBinding->GetStringTable(), OutPackedDataCollection, AddressUints));
+						}
+						else
+						{
+							const TConstPCGValueRange<int64> ReadValueRange = PointData->GetConstMetadataEntryValueRange();
+	
+							ensure(PackAttributeHelper(AttributeBase, AttributeDesc, ReadValueRange[0], InDataBinding->GetStringTable(), OutPackedDataCollection, AddressUints));
+							for (uint32 ElementIndex = 1; ElementIndex < NumElements; ++ElementIndex)
+							{
+								PackAttributeHelper(AttributeBase, AttributeDesc, ReadValueRange[ElementIndex], InDataBinding->GetStringTable(), OutPackedDataCollection, AddressUints);
+							}
+						}
+					}
+				}
+			}
+			else if (const UPCGParamData* ParamData = Cast<UPCGParamData>(InputData[DataIndex].Data))
+			{
+				for (const FPCGKernelAttributeDesc& AttributeDesc : DataDesc.AttributeDescs)
+				{
+					const FPCGMetadataDomain* MetadataDomain = Metadata->GetConstMetadataDomain(AttributeDesc.AttributeKey.Identifier.MetadataDomain);
+					const FPCGMetadataAttributeBase* AttributeBase = MetadataDomain ? MetadataDomain->GetConstAttribute(AttributeDesc.AttributeKey.Identifier.Name) : nullptr;
+					if (!AttributeBase)
+					{
+						continue;
+					}
+	
+					uint32 AddressUints = GetElementDataStartAddressUints(OutPackedDataCollection.GetData(), DataIndex, AttributeDesc.AttributeId);
+	
+					// Pack attribute. Validate first element only for perf.
+					ensure(PackAttributeHelper(AttributeBase, AttributeDesc, /*InEntryKey*/0, InDataBinding->GetStringTable(), OutPackedDataCollection, AddressUints));
+					for (int32 ElementIndex = 1; ElementIndex < DataDesc.GetElementCountForAttribute(AttributeDesc).X; ++ElementIndex)
+					{
+						PackAttributeHelper(AttributeBase, AttributeDesc, /*InEntryKey*/ElementIndex, InDataBinding->GetStringTable(), OutPackedDataCollection, AddressUints);
+					}
+				}
+			}
+			else { /* TODO: Support additional data types. */ }
+		}
+	}
+	
+	EPCGUnpackDataCollectionResult UnpackDataCollection(FPCGContext* InContext, TSharedPtr<const FPCGDataCollectionDesc> InDataCollectionDesc, const TArray<uint8>& InPackedData, FName InPin, const TArray<FString>& InStringTable, FPCGDataCollection& OutDataCollection)
+	{
+		TRACE_CPUPROFILER_EVENT_SCOPE(PCGDataCollectionPacking::UnpackDataCollection);
+	
+		if (InPackedData.IsEmpty())
+		{
+			ensureMsgf(false, TEXT("Tried to unpack a GPU data collection, but the readback buffer was empty."));
+			return EPCGUnpackDataCollectionResult::NoData;
+		}
+
+		if (!ensure(InDataCollectionDesc))
+		{
+			return EPCGUnpackDataCollectionResult::NoData;
+		}
+
+		const TArray<FPCGDataDesc>& DataDescs = InDataCollectionDesc->DataDescs;
+	
+		const void* PackedData = InPackedData.GetData();
+		const float* DataAsFloat = static_cast<const float*>(PackedData);
+		const uint32* DataAsUint = static_cast<const uint32*>(PackedData);
+		const int32* DataAsInt = static_cast<const int32*>(PackedData);
+	
+		const uint32 PackedExecutionFlagAndNumData = DataAsUint[0];
+	
+		// Most significant bit of NumData is reserved to flag whether or not the kernel executed.
+		ensureMsgf(PackedExecutionFlagAndNumData & PCGDataCollectionPackingConstants::KernelExecutedFlag, TEXT("Tried to unpack a GPU data collection, but the compute shader did not execute."));
+		const uint32 NumData = PackedExecutionFlagAndNumData & ~PCGDataCollectionPackingConstants::KernelExecutedFlag;
+	
+		if (NumData != DataDescs.Num())
+		{
+			return EPCGUnpackDataCollectionResult::DataMismatch;
+		}
+	
+		TArray<FPCGTaggedData>& OutData = OutDataCollection.TaggedData;
+		TStaticArray<FPCGMetadataAttributeBase*, PCGDataCollectionPackingConstants::MAX_NUM_ATTRS> MetadataAttributes;
+	
+		for (uint32 DataIndex = 0; DataIndex < NumData; ++DataIndex)
+		{
+			const FPCGDataDesc& DataDesc = DataDescs[DataIndex];
+			ensure(DataDesc.GetElementDimension() == EPCGElementDimension::One);
+	
+			const uint32 DataHeaderAddress = (PCGDataCollectionPackingConstants::DATA_COLLECTION_HEADER_SIZE_BYTES + PCGDataCollectionPackingConstants::DATA_HEADER_SIZE_BYTES * DataIndex) / sizeof(uint32);
+			const uint32 TypeId =        DataAsUint[DataHeaderAddress + 0];
+			const uint32 NumAttributes = DataAsUint[DataHeaderAddress + 1];
+			
+			uint32 NumElements = DataAsUint[DataHeaderAddress + 2];
+			// Use tighter/more refined element count if we have one:
+			if (DataDesc.GetElementCount().X >= 0)
+			{
+				NumElements = FMath::Min(NumElements, static_cast<uint32>(DataDesc.GetElementCount().X));
+			}
+	
+			const TArray<FPCGKernelAttributeDesc>& AttributeDescs = DataDesc.AttributeDescs;
+			check(NumAttributes == AttributeDescs.Num());
+			check(AttributeDescs.Num() <= PCGDataCollectionPackingConstants::MAX_NUM_ATTRS);
+	
+			if (TypeId == PCGDataCollectionPackingConstants::POINT_DATA_TYPE_ID)
+			{
+				TRACE_CPUPROFILER_EVENT_SCOPE(UnpackPointDataItem);
+	
+				UPCGBasePointData* OutPointData = FPCGContext::NewPointData_AnyThread(InContext);
+				UPCGMetadata* Metadata = OutPointData->MutableMetadata();
+				OutPointData->SetNumPoints(NumElements, /*bInitializeValues=*/false);
+	
+				{
+					TRACE_CPUPROFILER_EVENT_SCOPE(InitializeOutput);
+	
+					// We only need to add the entry keys if there are actually attributes to unpack.
+					if (DataDesc.HasElementsMetadataDomainAttributes())
+					{
+						TPCGValueRange<int64> MetadataEntryRange = OutPointData->GetMetadataEntryValueRange(/*bAllocate=*/true);
+						TArray<PCGMetadataEntryKey*> ParentEntryKeys;
+						ParentEntryKeys.Reserve(NumElements);
+	
+						for (uint32 ElementIndex = 0; ElementIndex < NumElements; ++ElementIndex)
+						{
+							PCGMetadataEntryKey& EntryKey = MetadataEntryRange[ElementIndex];
+							EntryKey = PCGInvalidEntryKey;
+							ParentEntryKeys.Add(&EntryKey);
+						}
+	
+						Metadata->AddEntriesInPlace(ParentEntryKeys);
+					}
+					else if(OutPointData->IsA<UPCGPointData>())
+					{
+						// Special case so that we can use ParallelFor (UPCGPointArrayData doesn't need this)
+						TPCGValueRange<int64> MetadataEntryRange = OutPointData->GetMetadataEntryValueRange(/*bAllocate=*/true);
+						ParallelFor(NumElements, [&MetadataEntryRange](int32 ElementIndex)
+						{
+							MetadataEntryRange[ElementIndex] = -1;
+						});
+					}
+					else
+					{
+						OutPointData->SetMetadataEntry(-1);
+					}
+				}
+	
+				FPCGTaggedData& OutTaggedData = OutData.Emplace_GetRef();
+				OutTaggedData.Data = OutPointData;
+				OutTaggedData.Pin = InPin;
+	
+				for (const int32 TagStringKey : DataDesc.TagStringKeys)
+				{
+					if (ensure(InStringTable.IsValidIndex(TagStringKey)))
+					{
+						OutTaggedData.Tags.Add(InStringTable[TagStringKey]);
+					}
+				}
+	
+				for (const FPCGKernelAttributeDesc& AttributeDesc : AttributeDescs)
+				{
+					if (AttributeDesc.AttributeId >= PCGDataCollectionPackingConstants::NUM_RESERVED_ATTRS)
+					{
+						MetadataAttributes[AttributeDesc.AttributeId] = CreateAttributeFromAttributeDesc(Metadata, AttributeDesc);
+					}
+				}
+	
+				// No work to do if there are no elements to process.
+				if (NumElements == 0)
+				{
+					continue;
+				}
+	
+				FTransform DefaultTransform = FTransform::Identity;
+				bool bTransformIsAllocated = false;
+	
+				// Loop over attributes.
+				for (const FPCGKernelAttributeDesc& AttributeDesc : AttributeDescs)
+				{
+					const uint32 AttributeId = AttributeDesc.AttributeId;
+					const uint32 AttributeStrideUints = PCGDataDescriptionHelpers::GetAttributeTypeStrideBytes(AttributeDesc.AttributeKey.Type) >> 2;
+					const uint32 AddressUints = GetElementDataStartAddressUints(DataAsUint, DataIndex, AttributeId);
+	
+					if (AttributeId < PCGDataCollectionPackingConstants::NUM_RESERVED_ATTRS)
+					{
+						TRACE_CPUPROFILER_EVENT_SCOPE(UnpackPointProperty);
+	
+						const bool bAllocated = DataDesc.IsAttributeAllocated(AttributeDesc);
+	
+						// We tried hoisting this decision to a lambda but it didn't appear to help.
+						switch (AttributeId)
+						{
+						case PCGDataCollectionPackingConstants::POINT_POSITION_ATTRIBUTE_ID:
+						{
+							bTransformIsAllocated = bAllocated;
+	
+							if (bAllocated)
+							{
+								TPCGValueRange<FTransform> TransformRange = OutPointData->GetTransformValueRange(/*bAllocate=*/true);
+								ParallelFor(NumElements, [DataAsFloat, AddressUints, &TransformRange](int32 ElementIndex)
+								{
+									const FVector Location = FVector(
+										DataAsFloat[AddressUints + ElementIndex * 3 + 0],
+										DataAsFloat[AddressUints + ElementIndex * 3 + 1],
+										DataAsFloat[AddressUints + ElementIndex * 3 + 2]);
+	
+									TransformRange[ElementIndex].SetLocation(Location);
+								});
+							}
+							else
+							{
+								const FVector Location = FVector(
+									DataAsFloat[AddressUints + /*ElementIndex=*/0 * 3 + 0],
+									DataAsFloat[AddressUints + /*ElementIndex=*/0 * 3 + 1],
+									DataAsFloat[AddressUints + /*ElementIndex=*/0 * 3 + 2]);
+	
+								DefaultTransform.SetLocation(Location);
+							}
+	
+							break;
+						}
+						case PCGDataCollectionPackingConstants::POINT_ROTATION_ATTRIBUTE_ID:
+						{
+							if (bAllocated)
+							{
+								TPCGValueRange<FTransform> TransformRange = OutPointData->GetTransformValueRange(/*bAllocate=*/true);
+								ParallelFor(NumElements, [DataAsFloat, AddressUints, &TransformRange](int32 ElementIndex)
+								{
+									const FQuat Rotation = FQuat(
+										DataAsFloat[AddressUints + ElementIndex * 4 + 0],
+										DataAsFloat[AddressUints + ElementIndex * 4 + 1],
+										DataAsFloat[AddressUints + ElementIndex * 4 + 2],
+										DataAsFloat[AddressUints + ElementIndex * 4 + 3]);
+	
+									// Normalize here with default tolerance (zero quat will return identity).
+									TransformRange[ElementIndex].SetRotation(Rotation.GetNormalized());
+								});
+							}
+							else
+							{
+								const FQuat Rotation = FQuat(
+									DataAsFloat[AddressUints + /*ElementIndex=*/0 * 4 + 0],
+									DataAsFloat[AddressUints + /*ElementIndex=*/0 * 4 + 1],
+									DataAsFloat[AddressUints + /*ElementIndex=*/0 * 4 + 2],
+									DataAsFloat[AddressUints + /*ElementIndex=*/0 * 4 + 3]);
+	
+								DefaultTransform.SetRotation(Rotation);
+							}
+							break;
+						}
+						case PCGDataCollectionPackingConstants::POINT_SCALE_ATTRIBUTE_ID:
+						{
+							if (bAllocated)
+							{
+								TPCGValueRange<FTransform> TransformRange = OutPointData->GetTransformValueRange(/*bAllocate=*/true);
+								ParallelFor(NumElements, [DataAsFloat, AddressUints, &TransformRange](int32 ElementIndex)
+								{
+									const FVector Scale = FVector(
+										DataAsFloat[AddressUints + ElementIndex * 3 + 0],
+										DataAsFloat[AddressUints + ElementIndex * 3 + 1],
+										DataAsFloat[AddressUints + ElementIndex * 3 + 2]);
+	
+									TransformRange[ElementIndex].SetScale3D(Scale);
+								});
+							}
+							else
+							{
+								const FVector Scale = FVector(
+									DataAsFloat[AddressUints + /*ElementIndex=*/0 * 3 + 0],
+									DataAsFloat[AddressUints + /*ElementIndex=*/0 * 3 + 1],
+									DataAsFloat[AddressUints + /*ElementIndex=*/0 * 3 + 2]);
+	
+								DefaultTransform.SetScale3D(Scale);
+							}
+							break;
+						}
+						case PCGDataCollectionPackingConstants::POINT_BOUNDS_MIN_ATTRIBUTE_ID:
+						{
+							if (bAllocated)
+							{
+								TPCGValueRange<FVector> BoundsMinRange = OutPointData->GetBoundsMinValueRange(/*bAllocate=*/true);
+								ParallelFor(NumElements, [DataAsFloat, AddressUints, &BoundsMinRange](int32 ElementIndex)
+								{
+									const FVector BoundsMin = FVector(
+										DataAsFloat[AddressUints + ElementIndex * 3 + 0],
+										DataAsFloat[AddressUints + ElementIndex * 3 + 1],
+										DataAsFloat[AddressUints + ElementIndex * 3 + 2]);
+	
+									BoundsMinRange[ElementIndex] = BoundsMin;
+								});
+							}
+							else
+							{
+								const FVector BoundsMin = FVector(
+									DataAsFloat[AddressUints + /*ElementIndex=*/0 * 3 + 0],
+									DataAsFloat[AddressUints + /*ElementIndex=*/0 * 3 + 1],
+									DataAsFloat[AddressUints + /*ElementIndex=*/0 * 3 + 2]);
+	
+								OutPointData->SetBoundsMin(BoundsMin);
+							}
+							break;
+						}
+						case PCGDataCollectionPackingConstants::POINT_BOUNDS_MAX_ATTRIBUTE_ID:
+						{
+							if (bAllocated)
+							{
+								TPCGValueRange<FVector> BoundsMaxRange = OutPointData->GetBoundsMaxValueRange(/*bAllocate=*/true);
+								ParallelFor(NumElements, [DataAsFloat, AddressUints, &BoundsMaxRange](int32 ElementIndex)
+								{
+									const FVector BoundsMax = FVector(
+										DataAsFloat[AddressUints + ElementIndex * 3 + 0],
+										DataAsFloat[AddressUints + ElementIndex * 3 + 1],
+										DataAsFloat[AddressUints + ElementIndex * 3 + 2]);
+	
+									BoundsMaxRange[ElementIndex] = BoundsMax;
+								});
+							}
+							else
+							{
+								const FVector BoundsMax = FVector(
+									DataAsFloat[AddressUints + /*ElementIndex=*/0 * 3 + 0],
+									DataAsFloat[AddressUints + /*ElementIndex=*/0 * 3 + 1],
+									DataAsFloat[AddressUints + /*ElementIndex=*/0 * 3 + 2]);
+	
+								OutPointData->SetBoundsMax(BoundsMax);
+							}
+							break;
+						}
+						case PCGDataCollectionPackingConstants::POINT_COLOR_ATTRIBUTE_ID:
+						{
+							if (bAllocated)
+							{
+								TPCGValueRange<FVector4> ColorRange = OutPointData->GetColorValueRange(/*bAllocate=*/true);
+								ParallelFor(NumElements, [DataAsFloat, AddressUints, &ColorRange](int32 ElementIndex)
+								{
+									const FVector4 Color = FVector4(
+										DataAsFloat[AddressUints + ElementIndex * 4 + 0],
+										DataAsFloat[AddressUints + ElementIndex * 4 + 1],
+										DataAsFloat[AddressUints + ElementIndex * 4 + 2],
+										DataAsFloat[AddressUints + ElementIndex * 4 + 3]);
+	
+									ColorRange[ElementIndex] = Color;
+								});
+							}
+							else
+							{
+								const FVector4 Color = FVector4(
+									DataAsFloat[AddressUints + /*ElementIndex=*/0 * 4 + 0],
+									DataAsFloat[AddressUints + /*ElementIndex=*/0 * 4 + 1],
+									DataAsFloat[AddressUints + /*ElementIndex=*/0 * 4 + 2],
+									DataAsFloat[AddressUints + /*ElementIndex=*/0 * 4 + 3]);
+	
+								OutPointData->SetColor(Color);
+							}
+							break;
+						}
+						case PCGDataCollectionPackingConstants::POINT_DENSITY_ATTRIBUTE_ID:
+						{
+							if (bAllocated)
+							{
+								TPCGValueRange<float> DensityRange = OutPointData->GetDensityValueRange(/*bAllocate=*/true);
+								ParallelFor(NumElements, [DataAsFloat, AddressUints, &DensityRange](int32 ElementIndex)
+								{
+									DensityRange[ElementIndex] = DataAsFloat[AddressUints + ElementIndex];
+								});
+							}
+							else
+							{
+								OutPointData->SetDensity(DataAsFloat[AddressUints]);
+							}
+							break;
+						}
+						case PCGDataCollectionPackingConstants::POINT_SEED_ATTRIBUTE_ID:
+						{
+							if (bAllocated)
+							{
+								TPCGValueRange<int32> SeedRange = OutPointData->GetSeedValueRange(/*bAllocate=*/true);
+								ParallelFor(NumElements, [DataAsInt, AddressUints, &SeedRange](int32 ElementIndex)
+								{
+									SeedRange[ElementIndex] = DataAsInt[AddressUints + ElementIndex];
+								});
+							}
+							else
+							{
+								OutPointData->SetSeed(DataAsInt[AddressUints]);
+							}
+							break;
+						}
+						case PCGDataCollectionPackingConstants::POINT_STEEPNESS_ATTRIBUTE_ID:
+						{
+							if (bAllocated)
+							{
+								TPCGValueRange<float> SteepnessRange = OutPointData->GetSteepnessValueRange(/*bAllocate=*/true);
+								ParallelFor(NumElements, [DataAsFloat, AddressUints, &SteepnessRange](int32 ElementIndex)
+								{
+									SteepnessRange[ElementIndex] = DataAsFloat[AddressUints + ElementIndex];
+								});
+							}
+							else
+							{
+								OutPointData->SetSteepness(DataAsFloat[AddressUints]);
+							}
+							break;
+						}
+						default:
+							checkNoEntry();
+							break;
+						}
+					}
+					else
+					{
+						TRACE_CPUPROFILER_EVENT_SCOPE(UnpackAttribute);
+	
+						if (FPCGMetadataAttributeBase* AttributeBase = MetadataAttributes[AttributeDesc.AttributeId])
+						{
+							ensure(UnpackAttributeHelper(InContext, PackedData, AttributeDesc, InStringTable, AddressUints, DataDesc.GetElementCountForAttribute(AttributeDesc).X, OutPointData));
+						}
+					}
+				}
+	
+				if (!bTransformIsAllocated)
+				{
+					OutPointData->SetTransform(DefaultTransform);
+				}
+	
+				// TODO: It may be more efficient to create a mapping from input point index to final output point index and do everything in one pass.
+				auto ProcessRangeFunc = [OutPointData](int32 StartReadIndex, int32 StartWriteIndex, int32 Count)
+				{
+					int32 NumWritten = 0;
+	
+					const FConstPCGPointValueRanges InRanges(OutPointData);
+					FPCGPointValueRanges OutRanges(OutPointData, /*bAllocate=*/false);
+	
+					for (int32 ReadIndex = StartReadIndex; ReadIndex < StartReadIndex + Count; ++ReadIndex)
+					{
+						if (InRanges.DensityRange[ReadIndex] == PCGDataCollectionPackingConstants::INVALID_DENSITY)
+						{
+							continue;
+						}
+	
+						const int32 WriteIndex = StartWriteIndex + NumWritten;
+						OutRanges.SetFromValueRanges(WriteIndex, InRanges, ReadIndex);
+						++NumWritten;
+					}
+	
+					return NumWritten;
+				};
+	
+				auto MoveDataRangeFunc = [OutPointData](int32 RangeStartIndex, int32 MoveToIndex, int32 NumElements)
+				{
+					OutPointData->MoveRange(RangeStartIndex, MoveToIndex, NumElements);
+				};
+	
+				auto FinishedFunc = [OutPointData](int32 NumWritten)
+				{
+					OutPointData->SetNumPoints(NumWritten);
+				};
+	
+				{
+					TRACE_CPUPROFILER_EVENT_SCOPE(DiscardInvalidPoints);
+					FPCGAsyncState* AsyncState = InContext ? &InContext->AsyncState : nullptr;
+					FPCGAsync::AsyncProcessingRangeEx(
+						AsyncState, 
+						OutPointData->GetNumPoints(),
+						[]{},
+						ProcessRangeFunc,
+						MoveDataRangeFunc,
+						FinishedFunc,
+						/*bEnableTimeSlicing=*/false);
+				}
+			}
+			else if (TypeId == PCGDataCollectionPackingConstants::PARAM_DATA_TYPE_ID)
+			{
+				TRACE_CPUPROFILER_EVENT_SCOPE(UnpackParamDataItem);
+	
+				UPCGParamData* OutParamData = FPCGContext::NewObject_AnyThread<UPCGParamData>(InContext);
+				UPCGMetadata* Metadata = OutParamData->MutableMetadata();
+	
+				{
+					TRACE_CPUPROFILER_EVENT_SCOPE(InitializeOutput);
+	
+					TArray<TTuple</*EntryKey=*/int64, /*ParentEntryKey=*/int64>> AllMetadataEntries;
+					AllMetadataEntries.SetNumUninitialized(NumElements);
+	
+					ParallelFor(NumElements, [&](int32 ElementIndex)
+					{
+						//LLM_SCOPE_BYTAG(PCG); // AddEntryPlaceholder() does not seem to allocate but consider chunking this loop and enabling this LLM tag if this changes.
+						AllMetadataEntries[ElementIndex] = MakeTuple(Metadata->AddEntryPlaceholder(), PCGInvalidEntryKey);
+					});
+	
+					Metadata->AddDelayedEntries(AllMetadataEntries);
+				}
+	
+				FPCGTaggedData& OutTaggedData = OutData.Emplace_GetRef();
+				OutTaggedData.Data = OutParamData;
+				OutTaggedData.Pin = InPin;
+	
+				for (const FPCGKernelAttributeDesc& AttributeDesc : AttributeDescs)
+				{
+					MetadataAttributes[AttributeDesc.AttributeId] = CreateAttributeFromAttributeDesc(Metadata, AttributeDesc);
+				}
+	
+				// No work to do if there are no elements to process.
+				if (NumElements == 0)
+				{
+					continue;
+				}
+	
+				// Loop over attributes.
+				for (const FPCGKernelAttributeDesc& AttributeDesc : AttributeDescs)
+				{
+					TRACE_CPUPROFILER_EVENT_SCOPE(UnpackAttribute);
+	
+					if (FPCGMetadataAttributeBase* AttributeBase = MetadataAttributes[AttributeDesc.AttributeId])
+					{
+						const uint32 AddressUints = GetElementDataStartAddressUints(DataAsUint, DataIndex, AttributeDesc.AttributeId);
+						ensure(UnpackAttributeHelper(InContext, PackedData, AttributeDesc, InStringTable, AddressUints, NumElements, OutParamData));
+					}
+				}
+			}
+			else { /* TODO: Support additional data types. */ }
+		}
+	
+		return EPCGUnpackDataCollectionResult::Success;
+	}
+}
diff --git a/Engine/Plugins/PCG/Source/PCG/Private/Compute/Packing/PCGDataCollectionPacking.h b/Engine/Plugins/PCG/Source/PCG/Private/Compute/Packing/PCGDataCollectionPacking.h
new file mode 100644
--- /dev/null
+++ b/Engine/Plugins/PCG/Source/PCG/Private/Compute/Packing/PCGDataCollectionPacking.h
@@ -0,0 +1,74 @@
+// Copyright Epic Games, Inc. All Rights Reserved.
+
+#pragma once
+
+#include "CoreMinimal.h"
+
+class UPCGDataBinding;
+struct FPCGContext;
+struct FPCGDataCollection;
+struct FPCGDataCollectionDesc;
+struct FPCGDataDesc;
+
+enum class EPCGUnpackDataCollectionResult
+{
+	Success,
+	DataMismatch,
+	NoData
+};
+
+namespace PCGDataCollectionPackingConstants
+{
+	constexpr int32 MAX_NUM_ATTRS = 128;
+	constexpr int32 NUM_RESERVED_ATTRS = 32; // Reserved for point properties, spline accessors, etc.
+	constexpr int32 MAX_NUM_CUSTOM_ATTRS = MAX_NUM_ATTRS - NUM_RESERVED_ATTRS; // Reserved for custom attributes
+
+	constexpr int32 DATA_COLLECTION_HEADER_SIZE_BYTES = 4; // 4 bytes for NumData
+	constexpr int32 DATA_HEADER_PREAMBLE_SIZE_BYTES = 12; // 4 bytes for DataType, 4 bytes for NumAttrs, 4 bytes for NumElements
+	constexpr int32 ATTRIBUTE_HEADER_SIZE_BYTES = 8; // 4 bytes for PackedIdAndStride, 4 bytes for data start address
+	constexpr int32 DATA_HEADER_SIZE_BYTES = DATA_HEADER_PREAMBLE_SIZE_BYTES + MAX_NUM_ATTRS * ATTRIBUTE_HEADER_SIZE_BYTES;
+
+	constexpr int32 POINT_DATA_TYPE_ID = 0;
+	constexpr int32 PARAM_DATA_TYPE_ID = 1;
+
+	constexpr int32 NUM_POINT_PROPERTIES = 9;
+	constexpr int32 POINT_POSITION_ATTRIBUTE_ID = 0;
+	constexpr int32 POINT_ROTATION_ATTRIBUTE_ID = 1;
+	constexpr int32 POINT_SCALE_ATTRIBUTE_ID = 2;
+	constexpr int32 POINT_BOUNDS_MIN_ATTRIBUTE_ID = 3;
+	constexpr int32 POINT_BOUNDS_MAX_ATTRIBUTE_ID = 4;
+	constexpr int32 POINT_COLOR_ATTRIBUTE_ID = 5;
+	constexpr int32 POINT_DENSITY_ATTRIBUTE_ID = 6;
+	constexpr int32 POINT_SEED_ATTRIBUTE_ID = 7;
+	constexpr int32 POINT_STEEPNESS_ATTRIBUTE_ID = 8;
+
+	constexpr uint32 KernelExecutedFlag = 1 << 31;
+
+	// PackedAttributeInfo bitmasks
+	constexpr uint32 AttributeStrideMask = 0xFF;
+	constexpr uint32 AttributeAllocatedMask = 1 << 31;
+
+	// Used to represent invalid/removed points. We use a value slightly less than max float,
+	// as not all platforms support float infinity in shaders.
+	constexpr float INVALID_DENSITY = 3.402823e+38f;
+}
+
+namespace PCGDataCollectionPackingHelpers
+{
+	/** Computes the size (in bytes) of the data description after packing. */
+	uint64 ComputePackedSizeBytes(const FPCGDataDesc& InDataDesc);
+
+	/** Computes the size (in bytes) of the header portion of the packed data collection buffer. */
+	uint32 ComputePackedHeaderSizeBytes(TSharedPtr<const FPCGDataCollectionDesc> InDataCollectionDesc);
+
+	/** Computes the size (in bytes) of the data collection after packing. */
+	uint64 ComputePackedSizeBytes(TSharedPtr<const FPCGDataCollectionDesc> InDataCollectionDesc);
+
+	void WriteHeader(TSharedPtr<const FPCGDataCollectionDesc> InDataCollectionDesc, TArray<uint32>& OutPackedDataCollectionHeader);
+
+	/** Pack a data collection into the GPU data format. DataDescs defines which attributes are packed. */
+	void PackDataCollection(const FPCGDataCollection& InDataCollection, TSharedPtr<const FPCGDataCollectionDesc> InDataCollectionDesc, FName InPin, const UPCGDataBinding* InDataBinding, TArray<uint32>& OutPackedDataCollection);
+
+	/** Unpack a buffer of 8-bit uints to a data collection. */
+	EPCGUnpackDataCollectionResult UnpackDataCollection(FPCGContext* InContext, TSharedPtr<const FPCGDataCollectionDesc> InDataCollectionDesc, const TArray<uint8>& InPackedData, FName InPin, const TArray<FString>& InStringTable, FPCGDataCollection& OutDataCollection);
+}
diff --git a/Engine/Plugins/PCG/Source/PCG/Private/Compute/Packing/PCGMeshSpawnerPacking.cpp b/Engine/Plugins/PCG/Source/PCG/Private/Compute/Packing/PCGMeshSpawnerPacking.cpp
new file mode 100644
--- /dev/null
+++ b/Engine/Plugins/PCG/Source/PCG/Private/Compute/Packing/PCGMeshSpawnerPacking.cpp
@@ -0,0 +1,59 @@
+// Copyright Epic Games, Inc. All Rights Reserved.
+
+#include "Compute/Packing/PCGMeshSpawnerPacking.h"
+
+#include "Compute/PCGComputeCommon.h"
+#include "Compute/PCGDataDescription.h"
+#include "Metadata/PCGAttributePropertySelector.h"
+
+#define LOCTEXT_NAMESPACE "PCGMeshSpawnerPacking"
+
+namespace PCGMeshSpawnerPackingHelpers
+{
+	void ComputeCustomFloatPacking(
+		FPCGContext* InContext,
+		const UPCGSettings* InSettings,
+		TConstArrayView<FPCGAttributePropertyInputSelector> InAttributeSelectors,
+		const TSharedPtr<const FPCGDataCollectionDesc> InDataCollectionDescription,
+		uint32& OutCustomFloatCount,
+		TArray<FUint32Vector4>& OutAttributeIdOffsetStrides)
+	{
+		check(InDataCollectionDescription);
+
+		uint32 OffsetFloats = 0;
+
+		for (const FPCGAttributePropertyInputSelector& AttributeSelector : InAttributeSelectors)
+		{
+			const FName AttributeName = AttributeSelector.GetAttributeName();
+
+			// We need to do a lookup here as the user does not provide the full attribute description, only the attribute name. So we see what we can
+			// find in the input data similar to the CPU SM Spawner.
+			FPCGKernelAttributeDesc AttributeDesc;
+			bool bConflictingTypesInData = false;
+			bool bPresentOnAllData = false;
+			InDataCollectionDescription->GetAttributeDesc(AttributeName, AttributeDesc, bConflictingTypesInData, bPresentOnAllData);
+
+			if (bConflictingTypesInData)
+			{
+				// TODO: in future we could partition the execution - run it once per attribute type in input data.
+				PCG_KERNEL_VALIDATION_ERR(InContext, InSettings,
+					FText::Format(LOCTEXT("AttributePackingTypeConflict", "Attribute '{0}' encountered with multiple different types in input data, custom float packing failed."), FText::FromName(AttributeName)));
+
+				OutAttributeIdOffsetStrides.Empty();
+				OutCustomFloatCount = 0;
+
+				return;
+			}
+
+			const uint32 StrideFloats = PCGDataDescriptionHelpers::GetAttributeTypeStrideBytes(AttributeDesc.AttributeKey.Type) / sizeof(float);
+
+			OutAttributeIdOffsetStrides.Emplace(static_cast<uint32>(AttributeDesc.AttributeId), OffsetFloats, StrideFloats, /*Unused*/0);
+
+			OffsetFloats += StrideFloats;
+		}
+
+		OutCustomFloatCount = OffsetFloats;
+	}
+}
+
+#undef LOCTEXT_NAMESPACE
diff --git a/Engine/Plugins/PCG/Source/PCG/Private/Compute/Packing/PCGMeshSpawnerPacking.h b/Engine/Plugins/PCG/Source/PCG/Private/Compute/Packing/PCGMeshSpawnerPacking.h
new file mode 100644
--- /dev/null
+++ b/Engine/Plugins/PCG/Source/PCG/Private/Compute/Packing/PCGMeshSpawnerPacking.h
@@ -0,0 +1,22 @@
+// Copyright Epic Games, Inc. All Rights Reserved.
+
+#pragma once
+
+#include "CoreMinimal.h"
+
+class UPCGSettings;
+struct FPCGAttributePropertyInputSelector;
+struct FPCGContext;
+struct FPCGDataCollectionDesc;
+
+namespace PCGMeshSpawnerPackingHelpers
+{
+	/** Compute how attributes can be packed to custom floats. */
+	void ComputeCustomFloatPacking(
+		FPCGContext* InContext,
+		const UPCGSettings* InSettings,
+		TConstArrayView<FPCGAttributePropertyInputSelector> InAttributeSelectors,
+		const TSharedPtr<const FPCGDataCollectionDesc> InDataCollectionDescription,
+		uint32& OutCustomFloatCount,
+		TArray<FUint32Vector4>& OutAttributeIdOffsetStrides);
+}
diff --git a/Engine/Plugins/PCG/Source/PCG/Private/Elements/PCGReroute.cpp b/Engine/Plugins/PCG/Source/PCG/Private/Elements/PCGReroute.cpp
--- a/Engine/Plugins/PCG/Source/PCG/Private/Elements/PCGReroute.cpp
+++ b/Engine/Plugins/PCG/Source/PCG/Private/Elements/PCGReroute.cpp
@@ -8,7 +8,7 @@
 #include "PCGModule.h"
 #include "PCGNode.h"
 #include "PCGPin.h"
-#include "Compute/PCGDataForGPU.h"
+#include "Compute/PCGDataDescription.h"
 #include "Helpers/PCGConversion.h"
 #include "Helpers/PCGSettingsHelpers.h"
 
diff --git a/Engine/Plugins/PCG/Source/PCG/Private/Graph/PCGGraphCompilerGPU.cpp b/Engine/Plugins/PCG/Source/PCG/Private/Graph/PCGGraphCompilerGPU.cpp
--- a/Engine/Plugins/PCG/Source/PCG/Private/Graph/PCGGraphCompilerGPU.cpp
+++ b/Engine/Plugins/PCG/Source/PCG/Private/Graph/PCGGraphCompilerGPU.cpp
@@ -14,6 +14,7 @@
 #include "Compute/DataInterfaces/PCGDataLabelResolverDataInterface.h"
 #include "Compute/DataInterfaces/PCGDebugDataInterface.h"
 #include "Compute/Elements/PCGComputeGraphElement.h"
+#include "Compute/Packing/PCGDataCollectionPacking.h"
 #include "Elements/PCGStaticMeshSpawnerKernel.h"
 #include "Graph/PCGGPUGraphCompilationContext.h"
 #include "Graph/PCGGraphCompiler.h"
@@ -909,7 +910,7 @@ void FPCGGraphCompilerGPU::BuildComputeGraphStaticData(
 			{
 				if (InOutComputeGraph->StaticAttributeTable.AddAttribute(Key) == INDEX_NONE)
 				{
-					UE_LOG(LogPCG, Error, TEXT("FPCGGraphCompilerGPU: Static attribute table exceeded maximum size (%d), use the 'Dump Data Descriptions' setting on the GPU node(s) to list attributes that are present."), PCGComputeConstants::MAX_NUM_CUSTOM_ATTRS);
+					UE_LOG(LogPCG, Error, TEXT("FPCGGraphCompilerGPU: Static attribute table exceeded maximum size (%d), use the 'Dump Data Descriptions' setting on the GPU node(s) to list attributes that are present."), PCGDataCollectionPackingConstants::MAX_NUM_CUSTOM_ATTRS);
 					break;
 				}
 			}
@@ -919,7 +920,7 @@ void FPCGGraphCompilerGPU::BuildComputeGraphStaticData(
 		}
 	}
 
-	ensure(InOutComputeGraph->StaticAttributeTable.Num() <= PCGComputeConstants::MAX_NUM_CUSTOM_ATTRS);
+	ensure(InOutComputeGraph->StaticAttributeTable.Num() <= PCGDataCollectionPackingConstants::MAX_NUM_CUSTOM_ATTRS);
 }
 
 void FPCGGraphCompilerGPU::CreateDataInterfaces(
diff --git a/Engine/Plugins/PCG/Source/PCG/Public/Compute/Data/PCGProxyForGPUData.h b/Engine/Plugins/PCG/Source/PCG/Public/Compute/Data/PCGProxyForGPUData.h
--- a/Engine/Plugins/PCG/Source/PCG/Public/Compute/Data/PCGProxyForGPUData.h
+++ b/Engine/Plugins/PCG/Source/PCG/Public/Compute/Data/PCGProxyForGPUData.h
@@ -4,13 +4,17 @@
 
 #include "PCGData.h"
 
-#include "Compute/PCGDataForGPU.h"
+#include "Compute/PCGDataDescription.h"
 
-#include "PCGProxyForGPUData.generated.h"
+#include "Misc/SpinLock.h"
 
-#define UE_API PCG_API
+#include "PCGProxyForGPUData.generated.h"
 
 struct FPCGContext;
+struct FPCGProxyForGPUDataCollection;
+
+class FRDGPooledBuffer;
+class FRHIGPUBufferReadback;
 
 USTRUCT()
 struct FPCGDataTypeInfoProxyForGPU : public FPCGDataTypeInfo
@@ -31,24 +35,24 @@ class UPCGProxyForGPUData : public UPCGData
 	GENERATED_BODY()
 
 public:
-	UE_API void Initialize(TSharedPtr<FPCGProxyForGPUDataCollection> InDataCollection, int InDataIndexInCollection);
+	void Initialize(TSharedPtr<FPCGProxyForGPUDataCollection> InDataCollection, int InDataIndexInCollection);
 
 	//~Begin UPCGData interface
 	PCG_ASSIGN_DEFAULT_TYPE_INFO(FPCGDataTypeInfoProxyForGPU)
 
-	UE_API virtual void AddToCrc(FArchiveCrc32& Ar, bool bFullDataCrc) const override;
+	PCG_API virtual void AddToCrc(FArchiveCrc32& Ar, bool bFullDataCrc) const override;
 	virtual bool CanBeSerialized() const override { return false; }
 	virtual bool HoldsTransientResources() const override { return true; }
 	virtual bool IsCacheable() const override { return false; }
-	UE_API virtual void ReleaseTransientResources(const TCHAR* InReason = nullptr);
-	UE_API virtual EPCGDataType GetUnderlyingDataType() const override;
-	UE_API virtual FPCGDataTypeBaseId GetUnderlyingDataTypeId() const override;
+	PCG_API virtual void ReleaseTransientResources(const TCHAR* InReason = nullptr) override;
+	PCG_API virtual EPCGDataType GetUnderlyingDataType() const override;
+	PCG_API virtual FPCGDataTypeBaseId GetUnderlyingDataTypeId() const override;
 	//~End UPCGData interface
 
 	/** Returns the GPU info. Returns null if the buffer has been discarded. */
-	UE_API TSharedPtr<const FPCGProxyForGPUDataCollection> GetInputDataCollectionInfo() const;
+	TSharedPtr<const FPCGProxyForGPUDataCollection> GetInputDataCollectionInfo() const;
 
-	int GetDataIndexInCollection() const { return DataIndexInCollection; }
+	int GetDataIndexInCollection() const;
 
 	struct FReadbackResult
 	{
@@ -60,24 +64,65 @@ public:
 	};
 
 	/** Populates a CPU data object representing the GPU data, performing a readback from GPU->CPU if required. */
-	UE_API FReadbackResult GetCPUData(FPCGContext* InContext) const;
+	PCG_API FReadbackResult GetCPUData(FPCGContext* InContext) const;
 
 	/** Returns the element count for this data. Does not trigger a GPU->CPU readback. */
-	UE_API FIntVector4 GetElementCount() const;
+	FIntVector4 GetElementCount() const;
 
 	/** Returns a description of this data. Does not trigger a GPU->CPU readback. */
-	UE_API bool GetDescription(FPCGDataDesc& OutDescription) const;
+	bool GetDescription(FPCGDataDesc& OutDescription) const;
 
-	UE_API TSharedPtr<const FPCGProxyForGPUDataCollection> GetGPUInfo() const;
+	TSharedPtr<const FPCGProxyForGPUDataCollection> GetGPUInfo() const;
 
-protected:
-	UE_API TSharedPtr<FPCGProxyForGPUDataCollection> GetGPUInfoMutable() const;
+private:
+	TSharedPtr<FPCGProxyForGPUDataCollection> GetGPUInfoMutable() const;
 
-protected:
+private:
 	UPROPERTY()
 	int32 DataIndexInCollection = INDEX_NONE;
 
 	TSharedPtr<FPCGProxyForGPUDataCollection> DataCollectionOnGPU;
 };
 
-#undef UE_API
+/** A proxy for a data collection residing in a GPU buffer along with functionality to retrieve the data on the CPU. Holds onto GPU memory. */
+struct FPCGProxyForGPUDataCollection : public TSharedFromThis<FPCGProxyForGPUDataCollection>
+{
+public:
+	FPCGProxyForGPUDataCollection(TRefCountPtr<FRDGPooledBuffer> InBuffer, uint32 InBufferSizeBytes, TSharedPtr<const FPCGDataCollectionDesc> InDescription, const TArray<FString>& InStringTable);
+
+	/** Populates a CPU data object representing the GPU data for the given index, performing a readback from GPU->CPU if required. */
+	bool GetCPUData(FPCGContext* InContext, int32 InDataIndex, FPCGTaggedData& OutData);
+
+	const TRefCountPtr<FRDGPooledBuffer>& GetBuffer() const { return Buffer; }
+
+	uint32 GetBufferSizeBytes() const { return BufferSizeBytes; }
+
+	TSharedPtr<const FPCGDataCollectionDesc> GetDescription() const { return Description; }
+
+	const TArray<FString>& GetStringTable() const { return StringTable; }
+
+private:
+	/** Persistent GPU buffer that can be read back. Buffer will be freed when this ref count is 0. */
+	TRefCountPtr<FRDGPooledBuffer> Buffer;
+
+	uint32 BufferSizeBytes = 0;
+
+	TSharedPtr<const FPCGDataCollectionDesc> Description = nullptr;
+
+	/** Used to comprehend string IDs in buffer. */
+	TArray<FString> StringTable;
+
+	/** Read back data. Populated once upon first readback request. */
+	TArray<const FPCGTaggedData> ReadbackData;
+	TArray<TStrongObjectPtr<const UPCGData>> ReadbackDataRefs;
+
+	TSharedPtr<FRHIGPUBufferReadback> ReadbackRequest;
+
+	TArray<uint8> RawReadbackData;
+
+	std::atomic<bool> bReadbackDataProcessed = false;
+
+	bool bReadbackDataArrived = false;
+
+	UE::FSpinLock ReadbackLock;
+};
diff --git a/Engine/Plugins/PCG/Source/PCG/Public/Compute/DataInterfaces/PCGComputeDataInterface.h b/Engine/Plugins/PCG/Source/PCG/Public/Compute/DataInterfaces/PCGComputeDataInterface.h
--- a/Engine/Plugins/PCG/Source/PCG/Public/Compute/DataInterfaces/PCGComputeDataInterface.h
+++ b/Engine/Plugins/PCG/Source/PCG/Public/Compute/DataInterfaces/PCGComputeDataInterface.h
@@ -11,6 +11,7 @@
 
 class UPCGComputeKernel;
 class UPCGDataBinding;
+struct FPCGDataCollectionDesc;
 
 UCLASS(MinimalAPI, Abstract, ClassGroup = (Procedural))
 class UPCGComputeDataInterface : public UComputeDataInterface
diff --git a/Engine/Plugins/PCG/Source/PCG/Public/Compute/PCGComputeCommon.h b/Engine/Plugins/PCG/Source/PCG/Public/Compute/PCGComputeCommon.h
--- a/Engine/Plugins/PCG/Source/PCG/Public/Compute/PCGComputeCommon.h
+++ b/Engine/Plugins/PCG/Source/PCG/Public/Compute/PCGComputeCommon.h
@@ -16,11 +16,7 @@ class UObject;
 class UPCGSettings;
 struct FPCGGPUCompilationContext;
 struct FPCGContext;
-struct FPCGDataCollectionDesc;
-struct FPCGDataDesc;
 struct FPCGPinProperties;
-struct FPCGPinPropertiesGPU;
-struct FPCGTaggedData;
 
 class FRHIShaderResourceView;
 
@@ -61,44 +57,8 @@ enum class EPCGElementDimension : uint8
 
 namespace PCGComputeConstants
 {
-	constexpr int MAX_NUM_ATTRS = 128;
-	constexpr int NUM_RESERVED_ATTRS = 32; // Reserved for point properties, spline accessors, etc.
-	constexpr int MAX_NUM_CUSTOM_ATTRS = MAX_NUM_ATTRS - NUM_RESERVED_ATTRS; // Reserved for custom attributes
-
-	constexpr int DATA_COLLECTION_HEADER_SIZE_BYTES = 4; // 4 bytes for NumData
-	constexpr int DATA_HEADER_PREAMBLE_SIZE_BYTES = 12; // 4 bytes for DataType, 4 bytes for NumAttrs, 4 bytes for NumElements
-	constexpr int ATTRIBUTE_HEADER_SIZE_BYTES = 8; // 4 bytes for PackedIdAndStride, 4 bytes for data start address
-	constexpr int DATA_HEADER_SIZE_BYTES = DATA_HEADER_PREAMBLE_SIZE_BYTES + MAX_NUM_ATTRS * ATTRIBUTE_HEADER_SIZE_BYTES;
-
-	constexpr int POINT_DATA_TYPE_ID = 0;
-	constexpr int PARAM_DATA_TYPE_ID = 1;
-
-	constexpr int NUM_POINT_PROPERTIES = 9;
-	constexpr int POINT_POSITION_ATTRIBUTE_ID = 0;
-	constexpr int POINT_ROTATION_ATTRIBUTE_ID = 1;
-	constexpr int POINT_SCALE_ATTRIBUTE_ID = 2;
-	constexpr int POINT_BOUNDS_MIN_ATTRIBUTE_ID = 3;
-	constexpr int POINT_BOUNDS_MAX_ATTRIBUTE_ID = 4;
-	constexpr int POINT_COLOR_ATTRIBUTE_ID = 5;
-	constexpr int POINT_DENSITY_ATTRIBUTE_ID = 6;
-	constexpr int POINT_SEED_ATTRIBUTE_ID = 7;
-	constexpr int POINT_STEEPNESS_ATTRIBUTE_ID = 8;
-
-	constexpr uint32 KernelExecutedFlag = 1 << 31;
-
-	// PackedAttributeInfo bitmasks
-	constexpr uint32 AttributeStrideMask = 0xFF;
-	constexpr uint32 AttributeAllocatedMask = 1 << 31;
-
 	constexpr int32 MAX_PRIMITIVE_COMPONENTS_PER_SPAWNER = 256;
-
 	constexpr uint32 THREAD_GROUP_SIZE = 64;
-
-	// Used to represent invalid/removed points. We use a value slightly less than max float,
-	// as not all platforms support float infinity in shaders.
-	constexpr float INVALID_DENSITY = 3.402823e+38f;
-
-	constexpr TCHAR DataLabelTagPrefix[] = { TEXT("PCG_DATA_LABEL") };
 }
 
 namespace PCGComputeHelpers
@@ -130,10 +90,10 @@ namespace PCGComputeHelpers
 
 #if PCG_KERNEL_LOGGING_ENABLED
 	/** Logs a warning on a GPU node in the graph and console. */
-	void LogKernelWarning(const FPCGContext* Context, const UPCGSettings* Settings, const FText& InText);
+	PCG_API void LogKernelWarning(const FPCGContext* Context, const UPCGSettings* Settings, const FText& InText);
 
 	/** Logs an error on a GPU node in the graph and console. */
-	void LogKernelError(const FPCGContext* Context, const UPCGSettings* Settings, const FText& InText);
+	PCG_API void LogKernelError(const FPCGContext* Context, const UPCGSettings* Settings, const FText& InText);
 #endif
 
 	/** Returns true if the given buffer size is dangerously large. Optionally emits error log. */
@@ -164,8 +124,8 @@ namespace PCGComputeHelpers
 
 	UPCGComputeDataInterface* CreateOutputPinDataInterface(const FCreateDataInterfaceParams& InParams);
 
-	void NotifyGPUToCPUReadback(UPCGDataBinding* InBinding, const UPCGComputeKernel* InKernel, const UPCGSettings* InSettings);
-	void NotifyCPUToGPUUpload(UPCGDataBinding* InBinding, const UPCGComputeKernel* InKernel, const UPCGSettings* InSettings);
+	PCG_API void NotifyGPUToCPUReadback(UPCGDataBinding* InBinding, const UPCGComputeKernel* InKernel, const UPCGSettings* InSettings);
+	PCG_API void NotifyCPUToGPUUpload(UPCGDataBinding* InBinding, const UPCGComputeKernel* InKernel, const UPCGSettings* InSettings);
 #endif
 }
 
@@ -206,7 +166,7 @@ struct FPCGPinReference
 	FName Label;
 };
 
-uint32 GetTypeHash(const FPCGPinReference& In);
+PCG_API uint32 GetTypeHash(const FPCGPinReference& In);
 
 /** A connection for wiring kernels within a node. */
 struct FPCGKernelEdge
@@ -243,7 +203,7 @@ public:
 	{
 	}
 
-	bool operator==(const FPCGKernelPin& Other) const = default;
+	PCG_API bool operator==(const FPCGKernelPin& Other) const = default;
 
 public:
 	UPROPERTY()
@@ -255,7 +215,7 @@ public:
 	UPROPERTY()
 	bool bIsInput = false;
 
-	friend uint32 GetTypeHash(const FPCGKernelPin& In);
+	PCG_API friend uint32 GetTypeHash(const FPCGKernelPin& In);
 };
 
 /** Helper struct for serializing data labels. */
diff --git a/Engine/Plugins/PCG/Source/PCG/Public/Compute/PCGDataBinding.h b/Engine/Plugins/PCG/Source/PCG/Public/Compute/PCGDataBinding.h
--- a/Engine/Plugins/PCG/Source/PCG/Public/Compute/PCGDataBinding.h
+++ b/Engine/Plugins/PCG/Source/PCG/Public/Compute/PCGDataBinding.h
@@ -4,7 +4,7 @@
 
 #include "PCGData.h"
 #include "PCGSettings.h"
-#include "Compute/PCGDataForGPU.h"
+#include "Compute/PCGDataDescription.h"
 #include "Compute/PCGKernelParamsCache.h"
 
 #include "PCGDataBinding.generated.h"
@@ -15,6 +15,9 @@ class UPCGComputeGraph;
 class UPCGComputeKernel;
 struct FPCGContextHandle;
 struct FPCGKernelPin;
+struct FPCGProxyForGPUDataCollection;
+
+class FRDGPooledBuffer;
 
 USTRUCT()
 struct FPCGSpawnerPrimitives
diff --git a/Engine/Plugins/PCG/Source/PCG/Public/Compute/PCGDataDescription.h b/Engine/Plugins/PCG/Source/PCG/Public/Compute/PCGDataDescription.h
new file mode 100644
--- /dev/null
+++ b/Engine/Plugins/PCG/Source/PCG/Public/Compute/PCGDataDescription.h
@@ -0,0 +1,248 @@
+// Copyright Epic Games, Inc. All Rights Reserved.
+
+#pragma once
+
+#include "PCGCommon.h"
+#include "PCGData.h"
+#include "Compute/PCGComputeCommon.h"
+
+#include "PCGDataDescription.generated.h"
+
+class UPCGData;
+class UPCGDataBinding;
+enum class EPCGElementMultiplicity : uint8;
+enum class EPCGMetadataTypes : uint8;
+enum class EPCGPointNativeProperties : uint32;
+
+UENUM()
+enum class EPCGKernelAttributeType : uint8
+{
+	Bool = 1,
+	Int,
+	Float,
+	Float2,
+	Float3,
+	Float4,
+	Rotator,
+	Quat,
+	Transform,
+	StringKey,
+	Name,
+
+	Invalid = std::numeric_limits<uint8>::max() UMETA(Hidden),
+};
+
+/** Attribute name and type which uniquely identify an attribute in a compute graph. */
+USTRUCT()
+struct FPCGKernelAttributeKey
+{
+	GENERATED_BODY()
+
+	FPCGKernelAttributeKey() = default;
+
+	explicit FPCGKernelAttributeKey(FPCGAttributeIdentifier InIdentifier, EPCGKernelAttributeType InType)
+		: Identifier(InIdentifier)
+		, Type(InType)
+	{}
+
+	explicit FPCGKernelAttributeKey(const FPCGAttributePropertySelector& InSelector, EPCGKernelAttributeType InType);
+
+	/** To be called everytime the Selector changes, to update the Identifier. Return true if it has changed. */
+	bool UpdateIdentifierFromSelector();
+	
+	void SetSelector(const FPCGAttributePropertySelector& InSelector);
+
+	/** Cached identifier. Need to be updated if the Selector ever change. */
+	UPROPERTY()
+	FPCGAttributeIdentifier Identifier;
+
+	UPROPERTY(EditAnywhere, Category = "Settings", meta = (DisplayAfter = "Attribute"))
+	EPCGKernelAttributeType Type = EPCGKernelAttributeType::Float;
+
+	bool IsValid() const;
+
+	bool operator==(const FPCGKernelAttributeKey& Other) const;
+
+	friend uint32 GetTypeHash(const FPCGKernelAttributeKey& In);
+
+protected:
+	/** Selector to specify which attribute to create and on which domain. At the moment, only support `@Data` domains or no domain (default one). */
+	UPROPERTY(EditAnywhere, Category = "Settings", meta = (PCG_DiscardPropertySelection, PCG_DiscardExtraSelection))
+	FPCGAttributePropertyOutputNoSourceSelector Name;
+};
+
+/** Table of attributes used in compute graph with helpers to get unique attribute ID used to read/write attribute in data collection buffers. */
+USTRUCT()
+struct FPCGKernelAttributeTable
+{
+	GENERATED_BODY()
+
+	FPCGKernelAttributeTable() = default;
+
+	int32 GetAttributeId(const FPCGKernelAttributeKey& InAttribute) const;
+
+	int32 GetAttributeId(FPCGAttributeIdentifier InIdentifier, EPCGKernelAttributeType InType) const;
+
+	/** Adds an attribute of given name and type. Returns index or INDEX_NONE if add failed (happens if max table size reached). */
+	int32 AddAttribute(const FPCGKernelAttributeKey& Key);
+	int32 AddAttribute(FPCGAttributeIdentifier InIdentifier, EPCGKernelAttributeType InType);
+
+	int32 Num() const { return AttributeTable.Num(); }
+
+#if PCG_KERNEL_LOGGING_ENABLED
+	void DebugLog() const;
+#endif
+
+protected:
+	UPROPERTY()
+	TArray<FPCGKernelAttributeKey> AttributeTable;
+};
+
+/** Data description for a metadata attribute. Stores identifying name and type as well as the unique attribute ID. */
+struct FPCGKernelAttributeDesc
+{
+public:
+	FPCGKernelAttributeDesc() = default;
+
+	explicit FPCGKernelAttributeDesc(int32 InIndex, EPCGKernelAttributeType InType, FPCGAttributeIdentifier InIdentifier)
+		: AttributeId(InIndex)
+		, AttributeKey(InIdentifier, InType)
+	{
+	}
+
+	explicit FPCGKernelAttributeDesc(int32 InIndex, EPCGKernelAttributeType InType, FPCGAttributeIdentifier InIdentifier, const TArray<int32>* InUniqueStringKeys)
+		: AttributeId(InIndex)
+		, AttributeKey(InIdentifier, InType)
+	{
+		if (InUniqueStringKeys)
+		{
+			UniqueStringKeys = *InUniqueStringKeys;
+		}
+	}
+
+	const TArray<int32>& GetUniqueStringKeys() const { return UniqueStringKeys; }
+
+	void AddUniqueStringKeys(const TArray<int32>& InOtherStringKeys);
+
+	void SetStringKeys(const TConstArrayView<int32>& InStringKeys);
+
+	bool operator==(const FPCGKernelAttributeDesc& Other) const;
+
+public:
+	int32 AttributeId = INDEX_NONE;
+
+	FPCGKernelAttributeKey AttributeKey{};
+
+private:
+	/* All possible string keys arriving on this attribute (string keys are indices into the string table in the data binding). */
+	TArray<int32> UniqueStringKeys;
+};
+
+/** Data description for a single data object (UPCGData). */
+struct FPCGDataDesc
+{
+public:
+	FPCGDataDesc() = default;
+	FPCGDataDesc(FPCGDataTypeIdentifier InType, int32 InElementCount);
+	FPCGDataDesc(FPCGDataTypeIdentifier InType, FIntPoint InElementCount);
+	FPCGDataDesc(FPCGDataTypeIdentifier InType, FIntVector3 InElementCount);
+	FPCGDataDesc(FPCGDataTypeIdentifier InType, FIntVector4 InElementCount);
+	FPCGDataDesc(const FPCGTaggedData& InTaggedData, const UPCGDataBinding* InBinding);
+
+	bool HasElementsMetadataDomainAttributes() const;
+
+	bool ContainsAttribute(FPCGAttributeIdentifier InAttributeIdentifier) const;
+
+	bool ContainsAttribute(FPCGAttributeIdentifier InAttributeIdentifier, EPCGKernelAttributeType InAttributeType) const;
+
+	void AddAttribute(FPCGKernelAttributeKey InAttribute, const UPCGDataBinding* InBinding, const TArray<int32>* InOptionalUniqueStringKeys = nullptr);
+
+	const FPCGDataTypeIdentifier& GetType() const { return Type; }
+
+	FIntVector4 GetElementCount() const { return ElementCount; }
+	FIntVector4 GetElementCountForAttribute(const FPCGKernelAttributeDesc& AttributeDesc) const;
+	EPCGElementDimension GetElementDimension() const { return ElementDimension; }
+	int32 ComputeTotalElementCount() const;
+
+	void AddElementCount(int32 InElementCountToAdd);
+	void AddElementCount(FIntPoint InElementCountToAdd);
+	void AddElementCount(FIntVector3 InElementCountToAdd);
+	void AddElementCount(FIntVector4 InElementCountToAdd);
+	void ScaleElementCount(int32 InMultiplier) { ElementCount *= InMultiplier; }
+
+	FPCGDataDesc& CombineElementCount(const FPCGDataDesc& Other, EPCGElementMultiplicity Multiplicity);
+	static FPCGDataDesc CombineElementCount(const FPCGDataDesc& A, const FPCGDataDesc& B, EPCGElementMultiplicity Multiplicity);
+
+	void AllocateProperties(EPCGPointNativeProperties InProperties) { AllocatedPointProperties |= InProperties; }
+	EPCGPointNativeProperties GetAllocatedProperties() const { return AllocatedPointProperties; }
+
+	bool IsAttributeAllocated(const FPCGKernelAttributeDesc& InAttributeDesc) const { return IsAttributeAllocated(InAttributeDesc.AttributeId); }
+	bool IsAttributeAllocated(int32 InAttributeId) const;
+
+private:
+	void InitializeAttributeDescs(const UPCGData* InData, const UPCGDataBinding* InBinding);
+
+public:
+	// @todo_pcg: Encapsulate these members.
+	TArray<FPCGKernelAttributeDesc> AttributeDescs;
+	TArray<int32, TInlineAllocator<4>> TagStringKeys;
+
+private:
+	FPCGDataTypeIdentifier Type = EPCGDataType::Point;
+	FIntVector4 ElementCount = FIntVector4::ZeroValue;
+	EPCGElementDimension ElementDimension = EPCGElementDimension::One;
+
+	/** Properties that should be fully allocated for this data. Only applies to EPCGDataType::Point currently. */
+	EPCGPointNativeProperties AllocatedPointProperties = EPCGPointNativeProperties::None;
+};
+
+/** Data description for a data collection (FPCGDataCollection). */
+struct FPCGDataCollectionDesc
+{
+	/** FPCGDataCollectionDesc should not be created directly, use MakeShared(). */
+	FPCGDataCollectionDesc(const FPCGDataCollectionDesc&) = delete;
+	FPCGDataCollectionDesc& operator=(const FPCGDataCollectionDesc&) = delete;
+	FPCGDataCollectionDesc(FPCGDataCollectionDesc&&) = delete;
+	FPCGDataCollectionDesc& operator=(FPCGDataCollectionDesc&&) = delete;
+
+	/** Create an FPCGDataCollectionDesc. */
+	static PCG_API TSharedPtr<FPCGDataCollectionDesc> MakeShared();
+
+	/** Create an FPCGDataCollectionDesc from another FPCGDataCollectionDesc. */
+	static PCG_API TSharedPtr<FPCGDataCollectionDesc> MakeSharedFrom(const TSharedPtr<const FPCGDataCollectionDesc> InOtherDataDesc);
+
+	/** Compute total number of processing elements. */
+	PCG_API uint32 ComputeTotalElementCount() const;
+
+	/** Compute total number of processing elements of the given type. */
+	uint32 ComputeTotalElementCount(FPCGDataTypeIdentifier InDataType) const;
+
+	/** Get description of first attribute with matching identifier in input data. Returns true if such an attribute found and also signals whether multiple matching attributes
+	* with conflicting names are present. */
+	bool GetAttributeDesc(FPCGAttributeIdentifier InAttributeIdentifier, FPCGKernelAttributeDesc& OutAttributeDesc, bool& bOutConflictingTypesFound, bool& bOutPresentOnAllData) const;
+
+	bool ContainsAttributeOnAnyData(FPCGAttributeIdentifier InAttributeIdentifier) const;
+
+	/** Makes attribute present on all data. If data has existing attribute with same name then the given type will be applied. */
+	void AddAttributeToAllData(FPCGKernelAttributeKey InAttribute, const UPCGDataBinding* InBinding, const TArray<int32>* InOptionalUniqueStringKeys = nullptr);
+
+	void AllocatePropertiesForAllData(EPCGPointNativeProperties InProperties);
+
+	void GetUniqueStringKeyValues(int32 InAttributeId, TArray<int32>& OutUniqueStringKeys) const;
+
+	int GetNumStringKeyValues(int32 InAttributeId) const;
+
+	/** Description of each data in this data collection. */
+	TArray<FPCGDataDesc> DataDescs;
+
+private:
+	FPCGDataCollectionDesc() = default;
+};
+
+namespace PCGDataDescriptionHelpers
+{
+	/** Returns GPU type that will be used to represent the given metadata type. */
+	EPCGKernelAttributeType GetAttributeTypeFromMetadataType(EPCGMetadataTypes MetadataType);
+
+	int GetAttributeTypeStrideBytes(EPCGKernelAttributeType Type);
+}
diff --git a/Engine/Plugins/PCG/Source/PCG/Public/Compute/PCGDataForGPU.h b/Engine/Plugins/PCG/Source/PCG/Public/Compute/PCGDataForGPU.h
deleted file mode 100644
diff --git a/Engine/Plugins/PCG/Source/PCG/Public/Compute/PCGPinPropertiesGPU.h b/Engine/Plugins/PCG/Source/PCG/Public/Compute/PCGPinPropertiesGPU.h
--- a/Engine/Plugins/PCG/Source/PCG/Public/Compute/PCGPinPropertiesGPU.h
+++ b/Engine/Plugins/PCG/Source/PCG/Public/Compute/PCGPinPropertiesGPU.h
@@ -3,7 +3,7 @@
 #pragma once
 
 #include "PCGPin.h"
-#include "Compute/PCGDataForGPU.h"
+#include "Compute/PCGDataDescription.h"
 
 #include "PCGPinPropertiesGPU.generated.h"
 
