[PCG] Added ellipsis to data entry combox box. Fixed a crash when opening the save data dialog and nothing was selected.

#rnx

diff --git a/Engine/Plugins/PCG/Source/PCGEditor/Private/Widgets/SPCGEditorGraphAttributeListView.cpp b/Engine/Plugins/PCG/Source/PCGEditor/Private/Widgets/SPCGEditorGraphAttributeListView.cpp
--- a/Engine/Plugins/PCG/Source/PCGEditor/Private/Widgets/SPCGEditorGraphAttributeListView.cpp
+++ b/Engine/Plugins/PCG/Source/PCGEditor/Private/Widgets/SPCGEditorGraphAttributeListView.cpp
@@ -398,6 +398,7 @@ void SPCGEditorGraphAttributeListView::Construct(const FArguments& InArgs, TShar
 		.OnSelectionChanged(this, &SPCGEditorGraphAttributeListView::OnSelectionChanged)
 		[
 			SNew(STextBlock)
+			.OverflowPolicy(ETextOverflowPolicy::Ellipsis)
 			.Text(this, &SPCGEditorGraphAttributeListView::OnGenerateSelectedDataText)
 		];
 
@@ -539,7 +540,7 @@ void SPCGEditorGraphAttributeListView::Construct(const FArguments& InArgs, TShar
 				PinComboBox->AsShared()
 			]
 			+SHorizontalBox::Slot()
-			.AutoWidth()
+			.FillWidth(1.0)
 			.Padding(1.0f, 0.0f)
 			[
 				DataComboBox->AsShared()
@@ -1397,10 +1398,9 @@ FPCGDataCollection SPCGEditorGraphAttributeListView::BuildDataCollectionForSave(
 			{
 				Collection = *PinInspectionData;
 			}
-			else
+			else if(int SelectedIndex = GetSelectedDataIndex(); PinInspectionData->TaggedData.IsValidIndex(SelectedIndex))
 			{
-
-				Collection.TaggedData.Add(PinInspectionData->TaggedData[GetSelectedDataIndex()]);
+				Collection.TaggedData.Add(PinInspectionData->TaggedData[SelectedIndex]);
 			}
 		}
 	}
@@ -1469,7 +1469,7 @@ TSharedRef<SWidget> SPCGEditorGraphAttributeListView::OnGeneratePinWidget(TShare
 
 TSharedRef<SWidget> SPCGEditorGraphAttributeListView::OnGenerateDataWidget(TSharedPtr<FString> InItem) const
 {
-	return SNew(STextBlock).Text(InItem.IsValid() ? FText::FromString(*InItem) : FText());
+	return SNew(STextBlock).Text(InItem.IsValid() ? FText::FromString(*InItem) : FText()).OverflowPolicy(ETextOverflowPolicy::Ellipsis);
 }
 
 void SPCGEditorGraphAttributeListView::OnSelectionChanged(TSharedPtr<FString> Item, ESelectInfo::Type SelectInfo)
