[PCG] Add selector actor validation for editor tools

* Spline and Volume cannot add a component to an actor that doesn't have a root component.
* Make sure the spline component is added without inheriting the scale/rotation.

#rnx

diff --git a/Engine/Plugins/PCG/Source/PCG/Private/Data/Tool/PCGToolSplineData.cpp b/Engine/Plugins/PCG/Source/PCG/Private/Data/Tool/PCGToolSplineData.cpp
--- a/Engine/Plugins/PCG/Source/PCG/Private/Data/Tool/PCGToolSplineData.cpp
+++ b/Engine/Plugins/PCG/Source/PCG/Private/Data/Tool/PCGToolSplineData.cpp
@@ -106,9 +106,17 @@ USplineComponent* FPCGInteractiveToolWorkingData_Spline::FindMatchingSplineCompo
 
 USplineComponent* FPCGInteractiveToolWorkingData_Spline::GenerateMatchingSplineComponent(const FPCGInteractiveToolWorkingDataContext& Context)
 {
+	// Can't add a spline component to an actor that have no root.
+	if (!Context.OwningActor.IsValid() || !Context.OwningActor->GetRootComponent())
+	{
+		return nullptr;
+	}
+
 	USplineComponent* MatchingSplineComponent = NewObject<USplineComponent>(Context.OwningActor.Get());
 	MatchingSplineComponent->SetMobility(Context.OwningActor->GetRootComponent()->GetMobility());
-	MatchingSplineComponent->SetupAttachment(Context.OwningActor->GetRootComponent());
+	// Make sure we do not inherit from the actor rotation and scale
+	MatchingSplineComponent->AttachToComponent(Context.OwningActor->GetRootComponent(), FAttachmentTransformRules(EAttachmentRule::KeepRelative, EAttachmentRule::KeepWorld, EAttachmentRule::KeepWorld, false));
+
 	// Transient until Applied. 
 	MatchingSplineComponent->SetFlags(RF_Transient);
 
diff --git a/Engine/Plugins/PCG/Source/PCG/Private/Data/Tool/PCGToolVolumeData.cpp b/Engine/Plugins/PCG/Source/PCG/Private/Data/Tool/PCGToolVolumeData.cpp
--- a/Engine/Plugins/PCG/Source/PCG/Private/Data/Tool/PCGToolVolumeData.cpp
+++ b/Engine/Plugins/PCG/Source/PCG/Private/Data/Tool/PCGToolVolumeData.cpp
@@ -109,7 +109,8 @@ UBrushComponent* FPCGInteractiveToolWorkingData_Volume::FindMatchingBrushCompone
 
 UBrushComponent* FPCGInteractiveToolWorkingData_Volume::GenerateMatchingBrushComponent(const FPCGInteractiveToolWorkingDataContext& Context)
 {
-	if (!Context.OwningActor.IsValid())
+	// Can't add a brush component to an actor that have no root.
+	if (!Context.OwningActor.IsValid() || !Context.OwningActor->GetRootComponent())
 	{
 		return nullptr;
 	}
diff --git a/Engine/Plugins/PCG/Source/PCGEditor/Private/EditorMode/Tools/Line/PCGDrawSplineTool.cpp b/Engine/Plugins/PCG/Source/PCGEditor/Private/EditorMode/Tools/Line/PCGDrawSplineTool.cpp
--- a/Engine/Plugins/PCG/Source/PCGEditor/Private/EditorMode/Tools/Line/PCGDrawSplineTool.cpp
+++ b/Engine/Plugins/PCG/Source/PCGEditor/Private/EditorMode/Tools/Line/PCGDrawSplineTool.cpp
@@ -200,6 +200,12 @@ namespace PCGDrawSplineTool
 	}
 } // namespace PCGDrawSplineTool
 
+bool UPCGInteractiveToolSettings_SplineBase::IsWorkingActorCompatible(const AActor* InActor) const
+{
+	// Working actor needs to have a root component to attach a spline to it.
+	return Super::IsWorkingActorCompatible(InActor) && InActor && InActor->GetRootComponent();
+}
+
 UPCGInteractiveToolSettings_SplineBase::UPCGInteractiveToolSettings_SplineBase()
 	: SplineName(GetDefault<UPCGEditorModeSettings>()->DefaultNewSplineComponentName)
 {
@@ -730,7 +736,16 @@ void UPCGDrawSplineSurfaceTool::Setup()
 
 bool UPCGDrawSplineToolBuilder::CanBuildTool(const FToolBuilderState& SceneState) const
 {
-	return SplineToolClass != nullptr;
+	// @todo_pcg: To be factorized in a base class
+	if (SplineToolClass == nullptr)
+	{
+		return false;
+	}
+
+	const UPCGInteractiveToolSettings_SplineBase* SettingsCDO = GetDefault<UPCGInteractiveToolSettings_SplineBase>();
+	check(SettingsCDO)
+	return SceneState.SelectedActors.IsEmpty()
+		|| SceneState.SelectedActors.ContainsByPredicate([SettingsCDO](const AActor* InActor) { return SettingsCDO->IsWorkingActorCompatible(InActor); });
 }
 
 UInteractiveTool* UPCGDrawSplineToolBuilder::BuildTool(const FToolBuilderState& SceneState) const
diff --git a/Engine/Plugins/PCG/Source/PCGEditor/Private/EditorMode/Tools/Line/PCGDrawSplineTool.h b/Engine/Plugins/PCG/Source/PCGEditor/Private/EditorMode/Tools/Line/PCGDrawSplineTool.h
--- a/Engine/Plugins/PCG/Source/PCGEditor/Private/EditorMode/Tools/Line/PCGDrawSplineTool.h
+++ b/Engine/Plugins/PCG/Source/PCGEditor/Private/EditorMode/Tools/Line/PCGDrawSplineTool.h
@@ -29,6 +29,14 @@ class PCGEDITOR_API UPCGInteractiveToolSettings_SplineBase : public UPCGInteract
 	GENERATED_BODY()
 
 public:
+	// ~Begin UPCGInteractiveToolSettings Interface
+
+	/** This function is called by the builder using UPCGInteractiveToolSettings_SplineBase default class which prohibits to have override
+	 * implementation of this function by child classes. Mark it final for that reason, to be revisited if we change the tool classes hierarchy. */
+	virtual bool IsWorkingActorCompatible(const AActor* InActor) const override final;
+
+	// ~End UPCGInteractiveToolSettings Interface
+
 	UPCGInteractiveToolSettings_SplineBase();
 	
 	USplineComponent* GetWorkingSplineComponent(FName DataInstanceIdentifier) const;
diff --git a/Engine/Plugins/PCG/Source/PCGEditor/Private/EditorMode/Tools/PCGInteractiveToolBuilder.cpp b/Engine/Plugins/PCG/Source/PCGEditor/Private/EditorMode/Tools/PCGInteractiveToolBuilder.cpp
--- a/Engine/Plugins/PCG/Source/PCGEditor/Private/EditorMode/Tools/PCGInteractiveToolBuilder.cpp
+++ b/Engine/Plugins/PCG/Source/PCGEditor/Private/EditorMode/Tools/PCGInteractiveToolBuilder.cpp
@@ -67,11 +67,6 @@ bool UE::PCG::EditorMode::Tool::BuildTool(UInteractiveTool* InTool)
 	return true;
 }
 
-bool UPCGInteractiveToolBuilder::CanBuildTool(const FToolBuilderState& SceneState) const
-{
-	return Super::CanBuildTool(SceneState);
-}
-
 UInteractiveTool* UPCGInteractiveToolBuilder::BuildTool(const FToolBuilderState& SceneState) const
 {
 	return Super::BuildTool(SceneState);
diff --git a/Engine/Plugins/PCG/Source/PCGEditor/Private/EditorMode/Tools/PCGInteractiveToolBuilder.h b/Engine/Plugins/PCG/Source/PCGEditor/Private/EditorMode/Tools/PCGInteractiveToolBuilder.h
--- a/Engine/Plugins/PCG/Source/PCGEditor/Private/EditorMode/Tools/PCGInteractiveToolBuilder.h
+++ b/Engine/Plugins/PCG/Source/PCGEditor/Private/EditorMode/Tools/PCGInteractiveToolBuilder.h
@@ -19,7 +19,6 @@ class PCGEDITOR_API UPCGInteractiveToolBuilder : public UInteractiveToolBuilder
 
 public:
 	// ~Begin UInteractiveToolBuilder interface
-	virtual bool CanBuildTool(const FToolBuilderState& SceneState) const override;
 	virtual UInteractiveTool* BuildTool(const FToolBuilderState& SceneState) const override;
 	virtual void PostSetupTool(UInteractiveTool* Tool, const FToolBuilderState& SceneState) const override;
 	// ~End UInteractiveToolBuilder interface
diff --git a/Engine/Plugins/PCG/Source/PCGEditor/Private/EditorMode/Tools/PCGInteractiveToolSettings.cpp b/Engine/Plugins/PCG/Source/PCGEditor/Private/EditorMode/Tools/PCGInteractiveToolSettings.cpp
--- a/Engine/Plugins/PCG/Source/PCGEditor/Private/EditorMode/Tools/PCGInteractiveToolSettings.cpp
+++ b/Engine/Plugins/PCG/Source/PCGEditor/Private/EditorMode/Tools/PCGInteractiveToolSettings.cpp
@@ -83,7 +83,7 @@ TValueOrError<bool, FText> UPCGInteractiveToolSettings::Initialize(UInteractiveT
 	// - Not have a PCG Component
 	// - Have a PCG Component but no valid working data
 	// - Have a PCG Component with valid working data we want to edit
-	if(SelectedActors.Num() == 1)
+	if (SelectedActors.Num() == 1 && SelectedActors[0] && IsWorkingActorCompatible(SelectedActors[0]))
 	{
 		EditActor = SelectedActors[0];
 	}
diff --git a/Engine/Plugins/PCG/Source/PCGEditor/Private/EditorMode/Tools/PCGInteractiveToolSettings.h b/Engine/Plugins/PCG/Source/PCGEditor/Private/EditorMode/Tools/PCGInteractiveToolSettings.h
--- a/Engine/Plugins/PCG/Source/PCGEditor/Private/EditorMode/Tools/PCGInteractiveToolSettings.h
+++ b/Engine/Plugins/PCG/Source/PCGEditor/Private/EditorMode/Tools/PCGInteractiveToolSettings.h
@@ -57,8 +57,6 @@ public:
 
 	void SetDataInstance(FName DataInstanceIdentifier);
 
-	virtual void SetActorClassToSpawn(TSubclassOf<AActor> Class);
-	
 	/** A soft object path to the graph asset to load by default for this tool. */
 	virtual FSoftObjectPath GetDefaultGraph() const;
 	
@@ -72,6 +70,11 @@ public:
 	 * If your target actor has any special requirements, create an actor for that purpose and return it in your settings. */
 	virtual TSubclassOf<AActor> GetWorkingActorClass() const;
 
+	virtual void SetActorClassToSpawn(TSubclassOf<AActor> Class);
+
+	/** Tools can have a requirement on the working actor */
+	virtual bool IsWorkingActorCompatible(const AActor* InActor) const { return InActor != nullptr; }
+
 	/** Filter details view graph selection. */
 	UFUNCTION()
 	virtual bool GraphAssetFilter(const FAssetData& AssetData) const;
diff --git a/Engine/Plugins/PCG/Source/PCGEditor/Private/EditorMode/Tools/Volume/PCGVolumeTool.cpp b/Engine/Plugins/PCG/Source/PCGEditor/Private/EditorMode/Tools/Volume/PCGVolumeTool.cpp
--- a/Engine/Plugins/PCG/Source/PCGEditor/Private/EditorMode/Tools/Volume/PCGVolumeTool.cpp
+++ b/Engine/Plugins/PCG/Source/PCGEditor/Private/EditorMode/Tools/Volume/PCGVolumeTool.cpp
@@ -52,6 +52,12 @@ const UScriptStruct* UPCGInteractiveToolSettings_Volume::GetWorkingDataType() co
 	return FPCGInteractiveToolWorkingData_Volume::StaticStruct();
 }
 
+bool UPCGInteractiveToolSettings_Volume::IsWorkingActorCompatible(const AActor* InActor) const
+{
+	// Working actor needs to be a volume
+	return Super::IsWorkingActorCompatible(InActor) && InActor && InActor->IsA<AVolume>();
+}
+
 AVolume* UPCGInteractiveToolSettings_Volume::GetWorkingVolume(FName DataInstanceIdentifier) const
 {
 	if (const FPCGInteractiveToolWorkingData_Volume* ToolData = GetTypedWorkingData<FPCGInteractiveToolWorkingData_Volume>(DataInstanceIdentifier))
@@ -321,29 +327,11 @@ bool UPCGVolumeTool::OnUpdateHover(const FInputDeviceRay& DevicePos)
 
 bool UPCGVolumeToolBuilder::CanBuildTool(const FToolBuilderState& SceneState) const
 {
-	FToolTargetTypeRequirements AcceptAllReqs{};
-
-	// Accepts either no selection, or selection from an AVolume actor.
-	const int32 NumTargets = SceneState.TargetManager->CountSelectedAndTargetable(SceneState, AcceptAllReqs);
-	if (NumTargets > 0)
-	{
-		bool bHasValidTargets = false;
-		SceneState.TargetManager->EnumerateSelectedAndTargetableComponents(SceneState, AcceptAllReqs,
-			[&bHasValidTargets](UActorComponent* Component)
-			{
-				if (Component && Component->GetOwner() && Component->GetOwner()->IsA<AVolume>())
-				{
-					bHasValidTargets = true;
-				}
-			});
-
-		if (!bHasValidTargets)
-		{
-			return false;
-		}
-	}
-
-	return true;
+	// @todo_pcg: To be factorized in a base class
+	const UPCGInteractiveToolSettings_Volume* SettingsCDO = GetDefault<UPCGInteractiveToolSettings_Volume>();
+	check(SettingsCDO)
+	return SceneState.SelectedActors.IsEmpty()
+		|| SceneState.SelectedActors.ContainsByPredicate([SettingsCDO](const AActor* InActor) { return SettingsCDO->IsWorkingActorCompatible(InActor); });
 }
 
 UInteractiveTool* UPCGVolumeToolBuilder::BuildTool(const FToolBuilderState& SceneState) const
diff --git a/Engine/Plugins/PCG/Source/PCGEditor/Private/EditorMode/Tools/Volume/PCGVolumeTool.h b/Engine/Plugins/PCG/Source/PCGEditor/Private/EditorMode/Tools/Volume/PCGVolumeTool.h
--- a/Engine/Plugins/PCG/Source/PCGEditor/Private/EditorMode/Tools/Volume/PCGVolumeTool.h
+++ b/Engine/Plugins/PCG/Source/PCGEditor/Private/EditorMode/Tools/Volume/PCGVolumeTool.h
@@ -28,6 +28,11 @@ public:
 	virtual const UScriptStruct* GetWorkingDataType() const override;
 	virtual void SetActorClassToSpawn(TSubclassOf<AActor> Class) override;
 	virtual bool CanResetToolData(FName DataInstanceIdentifier) override { return false; }
+
+	/** This function is called by the builder using UPCGInteractiveToolSettings_Volume default class which prohibits to have override
+	 * implementation of this function by child classes. Mark it final for that reason, to be revisited if we change the tool classes hierarchy. */
+	virtual bool IsWorkingActorCompatible(const AActor* InActor) const override final;
+
 	// ~End UPCGInteractiveToolSettings interface
 
 	AVolume* GetWorkingVolume(FName DataInstanceIdentifier) const;
