[PCG Editor Mode] Cache component instance, graph, and params to be restored on tool cancellation.

#rnx

diff --git a/Engine/Plugins/PCG/Source/PCGEditor/Private/EditorMode/Tools/PCGInteractiveToolBuilder.cpp b/Engine/Plugins/PCG/Source/PCGEditor/Private/EditorMode/Tools/PCGInteractiveToolBuilder.cpp
--- a/Engine/Plugins/PCG/Source/PCGEditor/Private/EditorMode/Tools/PCGInteractiveToolBuilder.cpp
+++ b/Engine/Plugins/PCG/Source/PCGEditor/Private/EditorMode/Tools/PCGInteractiveToolBuilder.cpp
@@ -15,7 +15,12 @@ bool UE::PCG::EditorMode::Tool::BuildTool(UInteractiveTool* InTool)
 	// This allows BP to only create a settings property for it to work, in theory
 	
 	using namespace PCG::EditorMode;
-	
+
+	if (!ensure(InTool))
+	{
+		return false;
+	}
+
 	const FObjectProperty* SettingsProperty = nullptr;
 	const void* PropertyAddress = nullptr;
 	UClass* SettingsClass = nullptr;
@@ -72,12 +77,17 @@ UInteractiveTool* UPCGInteractiveToolBuilder::BuildTool(const FToolBuilderState&
 	return Super::BuildTool(SceneState);
 }
 
-void UPCGInteractiveToolBuilder::PostSetupTool(UInteractiveTool* Tool, const FToolBuilderState& SceneState) const
+void UPCGInteractiveToolBuilder::PostSetupTool(UInteractiveTool* InTool, const FToolBuilderState& SceneState) const
 {
-	Super::PostSetupTool(Tool, SceneState);
-	SceneState.ToolManager->DisplayMessage(Tool->GetToolInfo().ToolDisplayMessage, EToolMessageLevel::UserNotification);
+	if (!ensure(InTool))
+	{
+		return;
+	}
+
+	Super::PostSetupTool(InTool, SceneState);
+	SceneState.ToolManager->DisplayMessage(InTool->GetToolInfo().ToolDisplayMessage, EToolMessageLevel::UserNotification);
 
-	TArray<UObject*> ToolProperties = Tool->GetToolProperties();
+	TArray<UObject*> ToolProperties = InTool->GetToolProperties();
 
 	if(ToolGraph)
 	{
@@ -91,12 +101,17 @@ void UPCGInteractiveToolBuilder::PostSetupTool(UInteractiveTool* Tool, const FTo
 	}
 }
 
-void UPCGInteractiveToolWithToolTargetsBuilder::PostSetupTool(UInteractiveTool* Tool, const FToolBuilderState& SceneState) const
+void UPCGInteractiveToolWithToolTargetsBuilder::PostSetupTool(UInteractiveTool* InTool, const FToolBuilderState& SceneState) const
 {
-	Super::PostSetupTool(Tool, SceneState);
-	SceneState.ToolManager->DisplayMessage(Tool->GetToolInfo().ToolDisplayMessage, EToolMessageLevel::UserNotification);
+	if (!ensure(InTool))
+	{
+		return;
+	}
+
+	Super::PostSetupTool(InTool, SceneState);
+	SceneState.ToolManager->DisplayMessage(InTool->GetToolInfo().ToolDisplayMessage, EToolMessageLevel::UserNotification);
 
-	TArray<UObject*> ToolProperties = Tool->GetToolProperties();
+	TArray<UObject*> ToolProperties = InTool->GetToolProperties();
 
 	if(ToolGraph)
 	{
diff --git a/Engine/Plugins/PCG/Source/PCGEditor/Private/EditorMode/Tools/PCGInteractiveToolBuilder.h b/Engine/Plugins/PCG/Source/PCGEditor/Private/EditorMode/Tools/PCGInteractiveToolBuilder.h
--- a/Engine/Plugins/PCG/Source/PCGEditor/Private/EditorMode/Tools/PCGInteractiveToolBuilder.h
+++ b/Engine/Plugins/PCG/Source/PCGEditor/Private/EditorMode/Tools/PCGInteractiveToolBuilder.h
@@ -20,7 +20,7 @@ class PCGEDITOR_API UPCGInteractiveToolBuilder : public UInteractiveToolBuilder
 public:
 	// ~Begin UInteractiveToolBuilder interface
 	virtual UInteractiveTool* BuildTool(const FToolBuilderState& SceneState) const override;
-	virtual void PostSetupTool(UInteractiveTool* Tool, const FToolBuilderState& SceneState) const override;
+	virtual void PostSetupTool(UInteractiveTool* InTool, const FToolBuilderState& SceneState) const override;
 	// ~End UInteractiveToolBuilder interface
 
 	void SetToolGraph(UPCGGraph* InGraph) { ToolGraph = InGraph; }
@@ -37,7 +37,7 @@ class PCGEDITOR_API UPCGInteractiveToolWithToolTargetsBuilder : public UInteract
 
 public:
 	// ~Begin UInteractiveToolBuilder interface
-	virtual void PostSetupTool(UInteractiveTool* Tool, const FToolBuilderState& SceneState) const override;
+	virtual void PostSetupTool(UInteractiveTool* InTool, const FToolBuilderState& SceneState) const override;
 	// ~End UInteractiveToolBuilder interface
 
 	void SetToolGraph(UPCGGraph* InGraph) { ToolGraph = InGraph; }
@@ -45,4 +45,4 @@ public:
 private:
 	UPROPERTY()
 	TObjectPtr<UPCGGraph> ToolGraph;
-};
\ No newline at end of file
+};
diff --git a/Engine/Plugins/PCG/Source/PCGEditor/Private/EditorMode/Tools/PCGInteractiveToolSettings.cpp b/Engine/Plugins/PCG/Source/PCGEditor/Private/EditorMode/Tools/PCGInteractiveToolSettings.cpp
--- a/Engine/Plugins/PCG/Source/PCGEditor/Private/EditorMode/Tools/PCGInteractiveToolSettings.cpp
+++ b/Engine/Plugins/PCG/Source/PCGEditor/Private/EditorMode/Tools/PCGInteractiveToolSettings.cpp
@@ -23,34 +23,28 @@ namespace UE::PCG::EditorMode::Tool
 {
 	TAutoConsoleVariable<bool> CVarCommitPropertyChangedEventOnTimer(TEXT("pcg.Editor.Tools.CommitPropertyChangedEventOnTimer"), true,
 		TEXT("If true, will queue up and automatically commit property changed events on a timer instead of committing directly."));
-}
 
-void UE::PCG::EditorMode::Tool::Shutdown(UInteractiveTool* InTool, EToolShutdownType ShutdownType)
-{
-	for(UObject* ToolPropertyObject : InTool->GetToolProperties(false))
+	void Shutdown(UInteractiveTool* InTool, EToolShutdownType ShutdownType)
 	{
-		if(UPCGInteractiveToolSettings* PCGSettings = Cast<UPCGInteractiveToolSettings>(ToolPropertyObject))
+		for (UObject* ToolPropertyObject : InTool->GetToolProperties(false))
 		{
-			PCGSettings->SaveProperties(InTool);
-			
-			if(ShutdownType == EToolShutdownType::Cancel)
+			if (UPCGInteractiveToolSettings* PCGToolSettings = Cast<UPCGInteractiveToolSettings>(ToolPropertyObject))
 			{
-				PCGSettings->Cancel(InTool);
-			}
-			else if(ShutdownType == EToolShutdownType::Accept)
-			{
-				PCGSettings->Apply(InTool);
+				PCGToolSettings->SaveProperties(InTool);
+
+				if (ShutdownType == EToolShutdownType::Cancel)
+				{
+					PCGToolSettings->Cancel(InTool);
+				}
+				else if (ShutdownType == EToolShutdownType::Accept)
+				{
+					PCGToolSettings->Apply(InTool);
+				}
+
+				PCGToolSettings->Shutdown();
 			}
 		}
 	}
-
-	for(UObject* ToolPropertyObject : InTool->GetToolProperties(false))
-	{
-		if(UPCGInteractiveToolSettings* PCGSettings = Cast<UPCGInteractiveToolSettings>(ToolPropertyObject))
-		{
-			PCGSettings->Shutdown();
-		}
-	}	
 }
 
 UPCGInteractiveToolSettings::UPCGInteractiveToolSettings()
@@ -76,8 +70,6 @@ TValueOrError<bool, FText> UPCGInteractiveToolSettings::Initialize(UInteractiveT
 		return Result;
 	}
 
-	RegisterPropertyWatchers();
-	
 	// If we have selected one actor, that's the actor we want to edit
 	// This actor might:
 	// - Not have a PCG Component
@@ -98,15 +90,15 @@ TValueOrError<bool, FText> UPCGInteractiveToolSettings::Initialize(UInteractiveT
 		GeneratedResources.SpawnedActor = SpawnedActor;
 	}
 
-	if(!ensureMsgf(EditActor.IsValid(), TEXT("The Edit Actor has to be valid at this point. Either as a previously existing actor, or a new spawned actor.")))
+	if(!ensure(EditActor.IsValid()))
 	{
-		return MakeValue(false);
+		return MakeError(LOCTEXT("InvalidEditActorDuringInit", "Invalid selected actor - either as a previously existing actor or a newly spawned actor."));
 	}
-	
+
 	TWeakObjectPtr<UPCGComponent> PCGComponent = FindOrGeneratePCGComponent(*EditActor.Get());
-	if(!ensureMsgf(PCGComponent.IsValid(), TEXT("The PCG Component has to be valid at this point. Either as a previously existing PCG Component, or a transient generated PCG Component.")))
+	if(!ensure(PCGComponent.IsValid()))
 	{
-		return MakeValue(false);
+		return MakeError(LOCTEXT("InvalidPCGComponentDuringInit", "Invalid PCG Component - either as a previously existing or transient generated component."));
 	}
 
 	OnModified.AddUObject(this, &UPCGInteractiveToolSettings::OnPropertyModified);
@@ -114,9 +106,9 @@ TValueOrError<bool, FText> UPCGInteractiveToolSettings::Initialize(UInteractiveT
 	// If we don't have a valid graph already, we try to assign the default graph to the newly created pcg component.
 	if (PCGComponent->GetGraph() == nullptr)
 	{
-		if (!ensureMsgf(ToolGraph, TEXT("The PCG component has no graph and no default graph for this tool was found in the editor preferences. Shutting down the tool.")))
+		if (!ensure(ToolGraph))
 		{
-			return MakeValue(false);
+			return MakeError(LOCTEXT("ToolHasNoDefaultGraph", "The PCG component has no graph and no default graph for this tool was found in the editor preferences. Shutting down the tool."));
 		}
 
 		PCGComponent->SetGraph(ToolGraph);
@@ -124,10 +116,13 @@ TValueOrError<bool, FText> UPCGInteractiveToolSettings::Initialize(UInteractiveT
 		bTryToGenerate = true;
 	}
 
-	// Finally, update the tool using the component graph (either the newly set one or a previously existing one).
-	// This will also update all tool data.
-	SetToolGraph(PCGComponent->GetGraph());
-	
+	UPCGGraphInstance* GraphInstance = PCGComponent->GetGraphInstance();
+	CacheGraphInstance(GraphInstance); // Guaranteed to be an instance.
+
+	// Finally, update the tool using the component graph. This will also update all tool data.
+	// Note: ToolGraph just tracks the graph and the graph parameters will be populated directly from the component's instance
+	SetToolGraph(GraphInstance->GetGraph());
+
 	// Keyboard setup for switching between data instance identifiers
 	UKeyInputBehavior* KeyInputBehavior = NewObject<UKeyInputBehavior>();
 	KeyInputBehavior->Initialize(this,
@@ -138,6 +133,8 @@ TValueOrError<bool, FText> UPCGInteractiveToolSettings::Initialize(UInteractiveT
 	KeyInputBehavior->bRequireAllKeys = false;
 	InTool->AddInputBehavior(KeyInputBehavior);
 
+	RegisterPropertyWatchers();
+
 	return MakeValue(true);
 }
 
@@ -171,11 +168,14 @@ void UPCGInteractiveToolSettings::Apply(UInteractiveTool* OwningTool)
 		WorkingDataMap.Value->OnToolApply(Context);
 	}
 
-   OwningTool->GetToolManager()->EndUndoTransaction();
+	OwningTool->GetToolManager()->EndUndoTransaction();
 }
 
 void UPCGInteractiveToolSettings::Cancel(UInteractiveTool* OwningTool)
 {
+	// Note: Okay for graph instance to be null here. Might need to restore a previous instance.
+	RestoreGraphInstanceInitialState(GetGraphInstance());
+
 	/** Each working data struct is responsible for wiping or restoring data to the previous state. */
 	for(const auto& WorkingDataMap : GetMutableTypedWorkingDataMap<FPCGInteractiveToolWorkingData>())
 	{
@@ -491,6 +491,39 @@ bool UPCGInteractiveToolSettings::HasGeneratedPCGComponent() const
 	return GeneratedResources.GeneratedPCGComponent.IsValid();
 }
 
+void UPCGInteractiveToolSettings::CacheGraphInstance(UPCGGraphInstance* InGraphInstance)
+{
+	if (InGraphInstance)
+	{
+		GraphInstanceCache = DuplicateObject(InGraphInstance, this);
+	}
+}
+
+void UPCGInteractiveToolSettings::RestoreGraphInstanceInitialState(UPCGGraphInstance* OutGraphInstance)
+{
+	// If the instance itself exists and was not changed, ONLY clobber the graph parameters...
+	if (OutGraphInstance && GraphInstanceCache && OutGraphInstance->GetGraph() == GraphInstanceCache->GetGraph())
+	{
+		// Double-check the layout is correct, before overwriting the params.
+		if (ensure(OutGraphInstance->ParametersOverrides.Parameters.HasSameLayout(GraphInstanceCache->ParametersOverrides.Parameters)))
+		{
+			OutGraphInstance->ParametersOverrides = GraphInstanceCache->ParametersOverrides;
+		}
+	}
+	else // Otherwise, initial instance state was null, or the final state became null, or the graph was changed. In those cases, revert it all.
+	{
+		if (UPCGComponent* PCGComponent = GetWorkingPCGComponent())
+		{
+			PCGComponent->SetGraph(GraphInstanceCache->GetGraph());
+			SetToolGraph(PCGComponent->GetGraph());
+			if (UPCGGraphInstance* Instance = PCGComponent->GetGraphInstance())
+			{
+				Instance->ParametersOverrides = GraphInstanceCache->ParametersOverrides;
+			}
+		}
+	}
+}
+
 void UPCGInteractiveToolSettings::CommitPropertyChangedEvents()
 {
 	for(auto& Pair : PropertyChangedEventQueue)
@@ -569,7 +602,10 @@ void UPCGInteractiveToolSettings::OnPropertyModified(UObject* Object, FProperty*
 	{
 		if (UPCGComponent* WorkingComponent = GetWorkingPCGComponent())
 		{
-			WorkingComponent->SetGraph(ToolGraph);
+			if (ToolGraph != WorkingComponent->GetGraph())
+			{
+				WorkingComponent->SetGraph(ToolGraph);
+			}
 		}
 
 		if(ToolGraph)
@@ -614,8 +650,6 @@ FSoftObjectPath UPCGInteractiveToolSettings::GetDefaultGraph() const
 		{
 			return Settings->DefaultGraph.ToSoftObjectPath();
 		}
-		
-		return FSoftObjectPath();
 	}
 	
 	return FSoftObjectPath();
diff --git a/Engine/Plugins/PCG/Source/PCGEditor/Private/EditorMode/Tools/PCGInteractiveToolSettings.h b/Engine/Plugins/PCG/Source/PCGEditor/Private/EditorMode/Tools/PCGInteractiveToolSettings.h
--- a/Engine/Plugins/PCG/Source/PCGEditor/Private/EditorMode/Tools/PCGInteractiveToolSettings.h
+++ b/Engine/Plugins/PCG/Source/PCGEditor/Private/EditorMode/Tools/PCGInteractiveToolSettings.h
@@ -200,10 +200,13 @@ private:
 	void DeleteGeneratedResources() const;
 	
 	void OnRefresh();
-	
+
+	void CacheGraphInstance(UPCGGraphInstance* InGraphInstance);
+	void RestoreGraphInstanceInitialState(UPCGGraphInstance* OutGraphInstance);
+
 public:
-	UPROPERTY(EditAnywhere, BlueprintReadOnly, Category = "PCG", meta = (NoResetToDefault, GetAssetFilter = "GraphAssetFilter"))
-	TObjectPtr<UPCGGraphInterface> ToolGraph = nullptr;
+	UPROPERTY(EditAnywhere, BlueprintReadOnly, Category = "PCG", meta = (TransientToolProperty, NoResetToDefault, GetAssetFilter = "GraphAssetFilter"))
+	TObjectPtr<UPCGGraphInterface> ToolGraph = nullptr; // Marked transient because the tool system will either fill this with a default or take the graph from the selected actor.
 
 	UPROPERTY(EditAnywhere, BlueprintReadOnly, Category = "PCG", meta = (NoResetToDefault, EditCondition="ToolGraph != nullptr", EditConditionHides, GetOptions="GetValidDataInstanceNamesForGraph"))
 	FName DataInstance = NAME_None;
@@ -251,4 +254,7 @@ private:
 
 	UFUNCTION()
 	bool HasGeneratedPCGComponent() const;
+
+	UPROPERTY(meta = (TransientToolProperty))
+	TObjectPtr<UPCGGraphInstance> GraphInstanceCache;
 };
