[PCG] Remove Duplicates on point was not safe to empty partition

#rnx

diff --git a/Engine/Plugins/PCG/Source/PCG/Private/Metadata/PCGMetadataPartitionCommon.cpp b/Engine/Plugins/PCG/Source/PCG/Private/Metadata/PCGMetadataPartitionCommon.cpp
--- a/Engine/Plugins/PCG/Source/PCG/Private/Metadata/PCGMetadataPartitionCommon.cpp
+++ b/Engine/Plugins/PCG/Source/PCG/Private/Metadata/PCGMetadataPartitionCommon.cpp
@@ -466,6 +466,11 @@ namespace PCGMetadataPartitionCommon
 
 		for (TArray<int32>& Indices : Partition)
 		{
+			if (Indices.IsEmpty())
+			{
+				continue;
+			}
+
 			IndicesToCopy.Add(Indices[0]);
 		}
 
