[PCG Editor Mode] Fixed null pointer dereference potential

#rnx

diff --git a/Engine/Plugins/PCG/Source/PCGEditor/Private/EditorMode/Tools/PCGInteractiveToolSettings.cpp b/Engine/Plugins/PCG/Source/PCGEditor/Private/EditorMode/Tools/PCGInteractiveToolSettings.cpp
--- a/Engine/Plugins/PCG/Source/PCGEditor/Private/EditorMode/Tools/PCGInteractiveToolSettings.cpp
+++ b/Engine/Plugins/PCG/Source/PCGEditor/Private/EditorMode/Tools/PCGInteractiveToolSettings.cpp
@@ -514,11 +514,19 @@ void UPCGInteractiveToolSettings::RestoreGraphInstanceInitialState(UPCGGraphInst
 	{
 		if (UPCGComponent* PCGComponent = GetWorkingPCGComponent())
 		{
-			PCGComponent->SetGraph(GraphInstanceCache->GetGraph());
-			SetToolGraph(PCGComponent->GetGraph());
-			if (UPCGGraphInstance* Instance = PCGComponent->GetGraphInstance())
+			if (GraphInstanceCache)
 			{
-				Instance->ParametersOverrides = GraphInstanceCache->ParametersOverrides;
+				PCGComponent->SetGraph(GraphInstanceCache->GetGraph());
+				SetToolGraph(PCGComponent->GetGraph());
+				if (UPCGGraphInstance* Instance = PCGComponent->GetGraphInstance())
+				{
+					Instance->ParametersOverrides = GraphInstanceCache->ParametersOverrides;
+				}
+			}
+			else // If the original graph was null, reset back to null.
+			{
+				PCGComponent->SetGraph(nullptr);
+				SetToolGraph(nullptr);
 			}
 		}
 	}
