[PCG] Fix multiple bugs in CreateGridXD shader functions.

* CreateGrid2D had a bug in X & Y computation and only worked for square grids
* CreateGrid3D had X & Y counts swapped
* Moved away from "rows" and "cols" to unambiguous axes X&Y

[RN] [PCG] Fixed bugs with CreateGrid2D and CreateGrid3D where the NumRows/NumCols arguments were not used correctly
#virtualized

diff --git a/Engine/Plugins/PCG/Source/PCG/Private/Compute/Elements/PCGCustomHLSL.cpp b/Engine/Plugins/PCG/Source/PCG/Private/Compute/Elements/PCGCustomHLSL.cpp
--- a/Engine/Plugins/PCG/Source/PCG/Private/Compute/Elements/PCGCustomHLSL.cpp
+++ b/Engine/Plugins/PCG/Source/PCG/Private/Compute/Elements/PCGCustomHLSL.cpp
@@ -853,9 +853,9 @@ void UPCGCustomHLSLSettings::UpdateHelperDeclarations()
 			"\n"
 			"// Returns the position of the Nth point in a 2D or 3D grid with the given constraints.\n"
 			"float3 CreateGrid2D(uint ElementIndex, uint NumPoints, float3 Min, float3 Max);\n"
-			"float3 CreateGrid2D(uint ElementIndex, uint NumPoints, uint NumRows, float3 Min, float3 Max);\n"
+			"float3 CreateGrid2D(uint ElementIndex, uint NumPoints, uint NumX, float3 Min, float3 Max);\n"
 			"float3 CreateGrid3D(uint ElementIndex, uint NumPoints, float3 Min, float3 Max);\n"
-			"float3 CreateGrid3D(uint ElementIndex, uint NumPoints, uint NumRows, uint NumCols, float3 Min, float3 Max);\n");
+			"float3 CreateGrid3D(uint ElementIndex, uint NumPoints, uint NumX, uint NumY, float3 Min, float3 Max);\n");
 
 		if (bPrintShaderDebugValues)
 		{
