Refactored Tool Data Identifiers to be contained within the working data. Working Data Container is now an array of working data structs.

- Changed the WorkingData pattern from underscore to period. {ToolTag}.{DataInstanceName}.

diff --git a/Engine/Plugins/PCG/Source/PCG/Internal/Data/Tool/PCGToolBaseData.h b/Engine/Plugins/PCG/Source/PCG/Internal/Data/Tool/PCGToolBaseData.h
--- a/Engine/Plugins/PCG/Source/PCG/Internal/Data/Tool/PCGToolBaseData.h
+++ b/Engine/Plugins/PCG/Source/PCG/Internal/Data/Tool/PCGToolBaseData.h
@@ -41,16 +41,19 @@ struct FPCGInteractiveToolWorkingDataContext
 	virtual ~FPCGInteractiveToolWorkingDataContext() = default;
 
 #if WITH_EDITORONLY_DATA
-	PCG_API FPCGInteractiveToolWorkingDataContext(UPCGComponent* OwningComponent, UInteractiveTool* InteractiveTool, UInteractiveToolPropertySet* PropertySet, FName DataInstanceIdentifier);
-
 	/** The owning PCG Component. Should always be valid. Can be used as an Outer for new objects. */
 	UPROPERTY(Transient)
 	TWeakObjectPtr<UPCGComponent> OwningPCGComponent;
 
-	/** An optional DataInstanceIdentifier. Can be used to find or create a specific object, such as a Spline Component with a certain tag. */
+	/** An optional DataInstanceIdentifier. Can be used to find or create a specific object, such as a Spline Component with a certain tag.
+	 *  Represents a Data Instance of a tool within a PCG graph.*/
 	UPROPERTY(Transient)
 	FName DataInstanceIdentifier;
 
+	/** An optional WorkingDataIdentifier. Takes shape of {ToolTag}.{DataInstanceIdentifier}. Needs to be unique within component storage. */
+	UPROPERTY(Transient)
+	FName WorkingDataIdentifier;
+
 	/** The owning actor. The owning PCG component's outer. Should always be valid. Convenience accessor. */
 	UPROPERTY(Transient)
 	TWeakObjectPtr<AActor> OwningActor;
@@ -73,7 +76,7 @@ struct FPCGInteractiveToolWorkingData : public FPCGInteractiveToolWorkingBaseDat
 	GENERATED_BODY()
 
 	PCG_API bool IsInitialized() const;
-
+	
 	PCG_API virtual bool IsValid() const;
 
 	/** Construct and add pcg data to the runtime context.
@@ -83,8 +86,6 @@ struct FPCGInteractiveToolWorkingData : public FPCGInteractiveToolWorkingBaseDat
 
 #if WITH_EDITORONLY_DATA
 public:
-	DECLARE_DELEGATE_OneParam(FOnInitialized, FPCGInteractiveToolWorkingData*)
-
 	/** Initializes data used by the tool working data. Initialize is called once when the working data is created. */
 	PCG_API void Initialize(const FPCGInteractiveToolWorkingDataContext& Context);
 
@@ -101,17 +102,14 @@ public:
 	/** Called when reset data was specifically requested, for example by the user. */
 	virtual void OnResetToolDataRequested(const FPCGInteractiveToolWorkingDataContext& Context) {}
 
-	PCG_API FOnInitialized& OnInitialized() { return OnInitializedDelegate; }
 private:
 	/** Implement your initialization logic here. Called by Initialize. */
 	PCG_API virtual void InitializeInternal(const FPCGInteractiveToolWorkingDataContext& Context);
+	
 protected:
 	/** Each working data struct keeps track of its generated resources. They have to be reverted during OnToolCancel. */
 	UPROPERTY(Transient)
 	FPCGInteractiveToolWorkingDataGeneratedResources GeneratedResources;
-
-	/** OnInitialized delegate to allow the settings class to do more work. */
-	FOnInitialized OnInitializedDelegate;
 #endif
 
 private:
diff --git a/Engine/Plugins/PCG/Source/PCG/Private/Data/PCGToolData.cpp b/Engine/Plugins/PCG/Source/PCG/Private/Data/PCGToolData.cpp
--- a/Engine/Plugins/PCG/Source/PCG/Private/Data/PCGToolData.cpp
+++ b/Engine/Plugins/PCG/Source/PCG/Private/Data/PCGToolData.cpp
@@ -4,8 +4,54 @@
 
 #include "GameFramework/Actor.h"
 
+FName FPCGInteractiveToolWorkingBaseData::GetWorkingDataIdentifier() const
+{
+	return WorkingDataIdentifier;
+}
+
 FPCGGraphToolData::FPCGGraphToolData()
 	: InitialActorClassToSpawn(AActor::StaticClass())
 {
 
-}
\ No newline at end of file
+}
+
+const TInstancedStruct<FPCGInteractiveToolWorkingBaseData>* FPCGInteractiveToolDataContainer::FindToolDataStruct(FName WorkingDataIdentifier) const
+{
+	for (const TInstancedStruct<FPCGInteractiveToolWorkingBaseData>& Data : ToolData)
+	{
+		if (const FPCGInteractiveToolWorkingBaseData* DataPtr = Data.GetPtr<FPCGInteractiveToolWorkingBaseData>(); DataPtr && DataPtr->GetWorkingDataIdentifier().IsEqual(WorkingDataIdentifier))
+		{
+			return &Data;
+		}
+	}
+
+	return nullptr;
+}
+
+TInstancedStruct<FPCGInteractiveToolWorkingBaseData>* FPCGInteractiveToolDataContainer::FindMutableToolDataStruct(FName WorkingDataIdentifier)
+{
+	for (TInstancedStruct<FPCGInteractiveToolWorkingBaseData>& Data : ToolData)
+	{
+		if (const FPCGInteractiveToolWorkingBaseData* DataPtr = Data.GetPtr<FPCGInteractiveToolWorkingBaseData>(); DataPtr && DataPtr->GetWorkingDataIdentifier().IsEqual(WorkingDataIdentifier))
+		{
+			return &Data;
+		}
+	}
+
+	return nullptr;
+}
+
+int32 FPCGInteractiveToolDataContainer::RemoveToolData(FName WorkingDataIdentifier)
+{
+	int32 Result = 0;
+	for (auto It = ToolData.CreateIterator(); It; ++It)
+	{
+		if (const FPCGInteractiveToolWorkingBaseData* DataPtr = It->GetPtr<FPCGInteractiveToolWorkingBaseData>(); DataPtr && DataPtr->GetWorkingDataIdentifier().IsEqual(WorkingDataIdentifier))
+		{
+			It.RemoveCurrent();
+			Result++;
+		}
+	}
+
+	return Result;
+}
diff --git a/Engine/Plugins/PCG/Source/PCG/Private/Data/Tool/PCGToolBaseData.cpp b/Engine/Plugins/PCG/Source/PCG/Private/Data/Tool/PCGToolBaseData.cpp
--- a/Engine/Plugins/PCG/Source/PCG/Private/Data/Tool/PCGToolBaseData.cpp
+++ b/Engine/Plugins/PCG/Source/PCG/Private/Data/Tool/PCGToolBaseData.cpp
@@ -14,7 +14,7 @@
 
 bool FPCGInteractiveToolWorkingData::IsInitialized() const
 {
-	return bInitialized;
+	return bInitialized && WorkingDataIdentifier.IsValid() && !WorkingDataIdentifier.IsNone();
 }
 
 bool FPCGInteractiveToolWorkingData::IsValid() const
@@ -30,11 +30,11 @@ void FPCGInteractiveToolWorkingData::InitializeRuntimeElementData(FPCGContext* I
 void FPCGInteractiveToolWorkingData::Initialize(const FPCGInteractiveToolWorkingDataContext& Context)
 {
 	ensureMsgf(Context.OwningPCGComponent.IsValid() && Context.OwningActor.IsValid(), TEXT("Owning PCG Component or Owning Actor invalid! This should never happen so an assumption was hurt."));
-
+	ensureMsgf(Context.WorkingDataIdentifier.IsNone() == false && Context.WorkingDataIdentifier.IsValid(), TEXT("WorkingDataIdentifier needs to be set when initializing tool data."));
+	WorkingDataIdentifier = Context.WorkingDataIdentifier;
+	
 	InitializeInternal(Context);
 
-	OnInitializedDelegate.ExecuteIfBound(this);
-
 	bInitialized = true;
 }
 
@@ -87,13 +87,4 @@ void FPCGInteractiveToolWorkingDataGeneratedResources::Cleanup()
 	GeneratedComponents.Empty();
 }
 
-FPCGInteractiveToolWorkingDataContext::FPCGInteractiveToolWorkingDataContext(UPCGComponent* InOwningComponent, UInteractiveTool* InInteractiveTool, UInteractiveToolPropertySet* InPropertySet, FName InDataInstanceIdentifier)
-{
-	OwningPCGComponent = InOwningComponent;
-	OwningActor = OwningPCGComponent->GetOwner();
-	DataInstanceIdentifier = InDataInstanceIdentifier;
-	InteractiveTool = InInteractiveTool;
-	PCGSettings = InPropertySet;
-}
 #endif // WITH_EDITOR
-
diff --git a/Engine/Plugins/PCG/Source/PCG/Private/Elements/PCGDataFromTool.cpp b/Engine/Plugins/PCG/Source/PCG/Private/Elements/PCGDataFromTool.cpp
--- a/Engine/Plugins/PCG/Source/PCG/Private/Elements/PCGDataFromTool.cpp
+++ b/Engine/Plugins/PCG/Source/PCG/Private/Elements/PCGDataFromTool.cpp
@@ -84,22 +84,22 @@ bool FPCGDataFromToolElement::ExecuteInternal(FPCGContext* Context) const
 	
 	UPCGComponent* PCGComponent = Cast<UPCGComponent>(Context->ExecutionSource.Get());
 
-	if (Settings->ToolTag.IsNone() == false)
+	if (PCGComponent && Settings->ToolTag.IsNone() == false)
 	{
 		FNameBuilder Builder(Settings->ToolTag);
 
 		// This mirrors UPCGInteractiveToolSettings::GetWorkingDataIdentifier
 		if (Settings->DataInstance.IsNone() == false)
 		{
-			Builder.AppendChar(TEXT('_'));
+			Builder.AppendChar(TEXT('.'));
 			Builder.Append(Settings->DataInstance.ToString());
 		}
 
-		FName ToolDataIdentifier = Builder.ToString();
-		
-		if (PCGComponent->ToolDataContainer.ToolData.Contains(ToolDataIdentifier) && PCGComponent->ToolDataContainer.ToolData[ToolDataIdentifier].IsValid())
+		FName WorkingDataIdentifier = Builder.ToString();
+
+		const FPCGInteractiveToolWorkingData* Data = PCGComponent->ToolDataContainer.GetTypedWorkingData<FPCGInteractiveToolWorkingData>(WorkingDataIdentifier);
+		if (Data && Data->IsValid())
 		{
-			const FPCGInteractiveToolWorkingData* Data = PCGComponent->ToolDataContainer.ToolData[ToolDataIdentifier].GetPtr<FPCGInteractiveToolWorkingData>();
 			Data->InitializeRuntimeElementData(Context);
 		}
 		// @todo_pcg
diff --git a/Engine/Plugins/PCG/Source/PCG/Public/Data/PCGToolData.h b/Engine/Plugins/PCG/Source/PCG/Public/Data/PCGToolData.h
--- a/Engine/Plugins/PCG/Source/PCG/Public/Data/PCGToolData.h
+++ b/Engine/Plugins/PCG/Source/PCG/Public/Data/PCGToolData.h
@@ -21,6 +21,14 @@ struct FPCGInteractiveToolWorkingBaseData
 
 	FPCGInteractiveToolWorkingBaseData() = default;
 	virtual ~FPCGInteractiveToolWorkingBaseData() = default;
+
+	PCG_API FName GetWorkingDataIdentifier() const;
+
+protected:
+	/** Used to identify this working data struct. Pattern: {ToolTag}.{DataInstanceIdentifier}.
+	 *  If there is no DataInstanceIdentifier, will just use {ToolTag}. Changing this will remap data. */
+	UPROPERTY(EditAnywhere, Category = "Tool")
+	FName WorkingDataIdentifier = NAME_None;
 };
 
 /** The tool data stored on the graph itself. Informs the UI. */
@@ -43,13 +51,52 @@ struct FPCGGraphToolData
 };
 
 USTRUCT(BlueprintType)
-struct FPCGInteractiveToolWorkingDataContainer
+struct FPCGInteractiveToolDataContainer
 {
 	GENERATED_BODY()
 
-	FPCGInteractiveToolWorkingDataContainer() = default;
-	virtual ~FPCGInteractiveToolWorkingDataContainer() = default;
+	FPCGInteractiveToolDataContainer() = default;
+	virtual ~FPCGInteractiveToolDataContainer() = default;
+
+	PCG_API const TInstancedStruct<FPCGInteractiveToolWorkingBaseData>* FindToolDataStruct(FName WorkingDataIdentifier) const;
+	PCG_API TInstancedStruct<FPCGInteractiveToolWorkingBaseData>* FindMutableToolDataStruct(FName WorkingDataIdentifier);
+
+	/** Removes all Tool Data with a given Working Data Identifier.
+	 *  @returns number of elements removed. Should generally not be more than 1. */
+	PCG_API int32 RemoveToolData(FName WorkingDataIdentifier);
+	
+	template<typename T = FPCGInteractiveToolWorkingBaseData>
+	const T* GetTypedWorkingData(FName WorkingDataIdentifier = NAME_None) const
+	{
+		static_assert(TIsDerivedFrom<T, FPCGInteractiveToolWorkingBaseData>::IsDerived, "T needs to be derived from FPCGInteractiveToolWorkingBaseData.");
+
+		if (const TInstancedStruct<FPCGInteractiveToolWorkingBaseData>* Data = FindToolDataStruct(WorkingDataIdentifier))
+		{
+			if (Data->GetScriptStruct()->IsChildOf(T::StaticStruct()) && Data->IsValid())
+			{
+				return Data->GetPtr<T>();
+			}
+		}
+
+		return nullptr;
+	}
+	
+	template<typename T = FPCGInteractiveToolWorkingBaseData>
+	T* GetMutableTypedWorkingData(FName WorkingDataIdentifier = NAME_None)
+	{
+		static_assert(TIsDerivedFrom<T, FPCGInteractiveToolWorkingBaseData>::IsDerived, "T needs to be derived from FPCGInteractiveToolWorkingBaseData.");
+
+		if (TInstancedStruct<FPCGInteractiveToolWorkingBaseData>* Data = FindMutableToolDataStruct(WorkingDataIdentifier))
+		{
+			if (Data->GetScriptStruct()->IsChildOf(T::StaticStruct()) && Data->IsValid())
+			{
+				return Data->GetMutablePtr<T>();
+			}
+		}
+
+		return nullptr;
+	}
 	
 	UPROPERTY(BlueprintReadOnly, VisibleAnywhere, Category = "Tool Data")
-	TMap<FName, TInstancedStruct<FPCGInteractiveToolWorkingBaseData>> ToolData;
+	TArray<TInstancedStruct<FPCGInteractiveToolWorkingBaseData>> ToolData;
 };
\ No newline at end of file
diff --git a/Engine/Plugins/PCG/Source/PCG/Public/PCGComponent.h b/Engine/Plugins/PCG/Source/PCG/Public/PCGComponent.h
--- a/Engine/Plugins/PCG/Source/PCG/Public/PCGComponent.h
+++ b/Engine/Plugins/PCG/Source/PCG/Public/PCGComponent.h
@@ -366,7 +366,7 @@ public:
 
 	/** This stores working data of a tool; runtime property as the PCG Component will rely on the tool's working data to generate properly. */
 	UPROPERTY(BlueprintReadOnly, EditAnywhere, Category = "Tool Data")
-	FPCGInteractiveToolWorkingDataContainer ToolDataContainer;
+	FPCGInteractiveToolDataContainer ToolDataContainer;
 	
 #if WITH_EDITORONLY_DATA
 	UPROPERTY(BlueprintReadWrite, EditAnywhere, Category = "Editing Settings", meta = (DisplayName = "Regenerate PCG Volume In Editor", DisplayPriority = 400))
diff --git a/Engine/Plugins/PCG/Source/PCGEditor/Private/EditorMode/Tools/Line/PCGDrawSplineTool.cpp b/Engine/Plugins/PCG/Source/PCGEditor/Private/EditorMode/Tools/Line/PCGDrawSplineTool.cpp
--- a/Engine/Plugins/PCG/Source/PCGEditor/Private/EditorMode/Tools/Line/PCGDrawSplineTool.cpp
+++ b/Engine/Plugins/PCG/Source/PCGEditor/Private/EditorMode/Tools/Line/PCGDrawSplineTool.cpp
@@ -258,7 +258,7 @@ FName UPCGInteractiveToolSettings_Spline::GetToolTag() const
 	return StaticGetToolTag();
 }
 
-void UPCGInteractiveToolSettings_Spline::OnWorkingDataInitialized(FPCGInteractiveToolWorkingData* WorkingData) const
+void UPCGInteractiveToolSettings_Spline::PostWorkingDataInitialized(FPCGInteractiveToolWorkingData* WorkingData) const
 {
 	if(USplineComponent* SplineComponent = GetWorkingSplineComponent(DataInstance))
 	{
diff --git a/Engine/Plugins/PCG/Source/PCGEditor/Private/EditorMode/Tools/Line/PCGDrawSplineTool.h b/Engine/Plugins/PCG/Source/PCGEditor/Private/EditorMode/Tools/Line/PCGDrawSplineTool.h
--- a/Engine/Plugins/PCG/Source/PCGEditor/Private/EditorMode/Tools/Line/PCGDrawSplineTool.h
+++ b/Engine/Plugins/PCG/Source/PCGEditor/Private/EditorMode/Tools/Line/PCGDrawSplineTool.h
@@ -74,7 +74,7 @@ class PCGEDITOR_API UPCGInteractiveToolSettings_Spline : public UPCGInteractiveT
 protected:
 	virtual const UScriptStruct* GetWorkingDataType() const override;
 	virtual FName GetToolTag() const override;
-	virtual void OnWorkingDataInitialized(FPCGInteractiveToolWorkingData* WorkingData) const override;
+	virtual void PostWorkingDataInitialized(FPCGInteractiveToolWorkingData* WorkingData) const override;
 	virtual void OnPropertyModified(UObject* Object, FProperty* Property) override;
 public:
 	UPROPERTY(EditAnywhere, Category = "Spline")
diff --git a/Engine/Plugins/PCG/Source/PCGEditor/Private/EditorMode/Tools/PCGInteractiveToolSettings.cpp b/Engine/Plugins/PCG/Source/PCGEditor/Private/EditorMode/Tools/PCGInteractiveToolSettings.cpp
--- a/Engine/Plugins/PCG/Source/PCGEditor/Private/EditorMode/Tools/PCGInteractiveToolSettings.cpp
+++ b/Engine/Plugins/PCG/Source/PCGEditor/Private/EditorMode/Tools/PCGInteractiveToolSettings.cpp
@@ -279,35 +279,29 @@ void UPCGInteractiveToolSettings::InitializeWorkingData(FName DataInstanceIdenti
 	Context.OwningActor = (Context.OwningPCGComponent.IsValid() ? Context.OwningPCGComponent->GetOwner() : nullptr);
 	Context.InteractiveTool = GetTypedOuter<UInteractiveTool>();
 	Context.PCGSettings = this;
-	
-	FName WorkingDataIdentifier = GetWorkingDataIdentifier(DataInstanceIdentifier);
+	Context.WorkingDataIdentifier = GetWorkingDataIdentifier(DataInstanceIdentifier);
+
+	FPCGInteractiveToolWorkingData* WorkingData = PCGComponent->ToolDataContainer.GetMutableTypedWorkingData<FPCGInteractiveToolWorkingData>(Context.WorkingDataIdentifier);
 	
 	if (EditActor.IsValid() && PCGComponent)
 	{
-		// If we don't find the requested working data, we create it
-		if(PCGComponent->ToolDataContainer.ToolData.Contains(WorkingDataIdentifier) == false)
+		// If we don't find the requested working data, we create and initialize it
+		if(WorkingData == nullptr)
 		{
 			TInstancedStruct<FPCGInteractiveToolWorkingData> NewToolWorkingData;
 			NewToolWorkingData.InitializeAsScriptStruct(Struct);
 			
-			PCGComponent->ToolDataContainer.ToolData.Add(WorkingDataIdentifier, MoveTemp(NewToolWorkingData));
-			GeneratedResources.AddedWorkingDataIdentifiers.Add(WorkingDataIdentifier);
+			NewToolWorkingData.GetMutablePtr<>()->Initialize(Context);
+			PostWorkingDataInitialized(NewToolWorkingData.GetMutablePtr<>());
+			
+			PCGComponent->ToolDataContainer.ToolData.Add(MoveTemp(NewToolWorkingData));
+			GeneratedResources.AddedWorkingDataIdentifiers.Add(Context.WorkingDataIdentifier);
 		}
 	}
 
-	FPCGInteractiveToolWorkingData* WorkingData = GetMutableTypedWorkingData(DataInstanceIdentifier);
+	WorkingData = GetMutableTypedWorkingData(DataInstanceIdentifier);
 	
-	// If it was freshly constructed, we initialize it. This will allow the structs to create or find components to work with.
-	if(ensureMsgf(WorkingData, TEXT("Working Data has to be allocated at this point!.")))
-	{
-		if(WorkingData->IsInitialized() == false)
-		{
-			WorkingData->OnInitialized().BindUObject(this, &UPCGInteractiveToolSettings::OnWorkingDataInitialized);
-			WorkingData->Initialize(Context);
-			WorkingData->OnInitialized().Unbind();
-		}
-	}
-	else
+	if(!ensureMsgf(WorkingData, TEXT("Working Data has to be allocated at this point!.")))
 	{
 		return;
 	}
@@ -328,7 +322,7 @@ void UPCGInteractiveToolSettings::InitializeWorkingData(FName DataInstanceIdenti
 		WorkingData->OnToolCancel(Context);
 		if (PCGComponent)
 		{
-			PCGComponent->ToolDataContainer.ToolData.Remove(WorkingDataIdentifier);
+			PCGComponent->ToolDataContainer.RemoveToolData(Context.WorkingDataIdentifier);
 		}
 
 		// In which case, requesting new working data will recreate a new spline
@@ -412,7 +406,7 @@ void UPCGInteractiveToolSettings::DeleteGeneratedResources() const
 	{
 		for(const FName& AddedWorkingDataIdentifier : GeneratedResources.AddedWorkingDataIdentifiers)
 		{
-			PCGComponent->ToolDataContainer.ToolData.Remove(AddedWorkingDataIdentifier);
+			PCGComponent->ToolDataContainer.RemoveToolData(AddedWorkingDataIdentifier);
 		}
 	}
 }
@@ -600,7 +594,7 @@ void UPCGInteractiveToolSettings::OnKeyPressed(const FKey& InKeyID)
 	}
 }
 
-void UPCGInteractiveToolSettings::OnWorkingDataInitialized(FPCGInteractiveToolWorkingData* WorkingData) const
+void UPCGInteractiveToolSettings::PostWorkingDataInitialized(FPCGInteractiveToolWorkingData* WorkingData) const
 {
 }
 
@@ -626,7 +620,7 @@ void UPCGInteractiveToolSettings::OnPropertyModified(UObject* Object, FProperty*
 		FName DefaultDataInstanceName = GetDefaultDataInstanceName();
 		SetDataInstance(DefaultDataInstanceName);
 		
-		// Update the working data to the available data instances ones from the graph
+		// Update the working data to the available data instances ones for the current tool from the graph
 		TArray<FName> ValidDataInstanceNames = GetDataInstanceNamesForGraph();
 
 		for (const FName& DataInstanceIdentifier : ValidDataInstanceNames)
@@ -766,7 +760,7 @@ FName UPCGInteractiveToolSettings::GetWorkingDataIdentifier(FName DataInstanceId
 
 	if (DataInstanceIdentifier.IsNone() == false)
 	{
-		DataIdentifierBuilder.AppendChar('_');
+		DataIdentifierBuilder.AppendChar('.');
 		DataIdentifierBuilder.Append(DataInstanceIdentifier.ToString());
 	}
 
diff --git a/Engine/Plugins/PCG/Source/PCGEditor/Private/EditorMode/Tools/PCGInteractiveToolSettings.h b/Engine/Plugins/PCG/Source/PCGEditor/Private/EditorMode/Tools/PCGInteractiveToolSettings.h
--- a/Engine/Plugins/PCG/Source/PCGEditor/Private/EditorMode/Tools/PCGInteractiveToolSettings.h
+++ b/Engine/Plugins/PCG/Source/PCGEditor/Private/EditorMode/Tools/PCGInteractiveToolSettings.h
@@ -104,12 +104,11 @@ public:
 	{
 		static_assert(TIsDerivedFrom<T, FPCGInteractiveToolWorkingData>::IsDerived, "T needs to be derived from FPCGInteractiveToolWorkingData.");
 
-		FName DataIdentifier = GetWorkingDataIdentifier(DataInstanceIdentifier);
+		FName WorkingDataIdentifier = GetWorkingDataIdentifier(DataInstanceIdentifier);
 		
-		UPCGComponent* PCGComponent = GetWorkingPCGComponent();
-		if (PCGComponent && PCGComponent->ToolDataContainer.ToolData.Contains(DataIdentifier))
+		if (UPCGComponent* PCGComponent = GetWorkingPCGComponent())
 		{
-			return PCGComponent->ToolDataContainer.ToolData[DataIdentifier].GetMutablePtr<T>(); 
+			return PCGComponent->ToolDataContainer.GetMutableTypedWorkingData<T>(WorkingDataIdentifier);
 		}
 
 		return nullptr;
@@ -122,10 +121,9 @@ public:
 
 		FName WorkingDataIdentifier = GetWorkingDataIdentifier(DataInstanceIdentifier);
 
-		UPCGComponent* PCGComponent = GetWorkingPCGComponent();
-		if(PCGComponent && PCGComponent->ToolDataContainer.ToolData.Contains(WorkingDataIdentifier))
+		if (UPCGComponent* PCGComponent = GetWorkingPCGComponent())
 		{
-			return PCGComponent->ToolDataContainer.ToolData[WorkingDataIdentifier].GetPtr<T>(); 
+			return PCGComponent->ToolDataContainer.GetTypedWorkingData<T>(WorkingDataIdentifier);
 		}
 		
 		return nullptr;
@@ -144,9 +142,9 @@ public:
 			{
 				FName WorkingDataIdentifier = GetWorkingDataIdentifier(DataInstanceName);
 			
-				if(PCGComponent->ToolDataContainer.ToolData.Contains(WorkingDataIdentifier))
+				if (T* ToolData = PCGComponent->ToolDataContainer.GetMutableTypedWorkingData<T>(WorkingDataIdentifier))
 				{
-					Result.Add(WorkingDataIdentifier, PCGComponent->ToolDataContainer.ToolData[WorkingDataIdentifier].GetMutablePtr<T>()); 
+					Result.Add(WorkingDataIdentifier, ToolData);
 				}
 			}
 		}
@@ -167,9 +165,9 @@ public:
 			{
 				FName WorkingDataIdentifier = GetWorkingDataIdentifier(DataInstanceName);
 			
-				if(PCGComponent->ToolDataContainer.ToolData.Contains(WorkingDataIdentifier))
+				if (const T* ToolData = PCGComponent->ToolDataContainer.GetTypedWorkingData<T>(WorkingDataIdentifier))
 				{
-					Result.Add(WorkingDataIdentifier, PCGComponent->ToolDataContainer.ToolData[WorkingDataIdentifier].GetPtr<T>()); 
+					Result.Add(WorkingDataIdentifier, ToolData);
 				}
 			}
 		}
@@ -190,7 +188,7 @@ protected:
 	virtual void OnKeyPressed(const FKey& KeyID) override;
 
 	/** Allows to modify the just initialized working data from the editor. This can be used if the initialization of working data requires editor data access. */
-	virtual void OnWorkingDataInitialized(FPCGInteractiveToolWorkingData* WorkingData) const;
+	virtual void PostWorkingDataInitialized(FPCGInteractiveToolWorkingData* WorkingData) const;
 	virtual void OnPropertyModified(UObject* Object, FProperty* Property);
 	virtual void OnRefreshInternal();
 	
