Minor fixes + removed bInitialized property in favor of WorkingDataIdentifier validity

diff --git a/Engine/Plugins/PCG/Source/PCG/Internal/Data/Tool/PCGToolBaseData.h b/Engine/Plugins/PCG/Source/PCG/Internal/Data/Tool/PCGToolBaseData.h
--- a/Engine/Plugins/PCG/Source/PCG/Internal/Data/Tool/PCGToolBaseData.h
+++ b/Engine/Plugins/PCG/Source/PCG/Internal/Data/Tool/PCGToolBaseData.h
@@ -111,8 +111,4 @@ protected:
 	UPROPERTY(Transient)
 	FPCGInteractiveToolWorkingDataGeneratedResources GeneratedResources;
 #endif
-
-private:
-	UPROPERTY()
-	bool bInitialized = false;
 };
\ No newline at end of file
diff --git a/Engine/Plugins/PCG/Source/PCG/Private/Data/Tool/PCGToolBaseData.cpp b/Engine/Plugins/PCG/Source/PCG/Private/Data/Tool/PCGToolBaseData.cpp
--- a/Engine/Plugins/PCG/Source/PCG/Private/Data/Tool/PCGToolBaseData.cpp
+++ b/Engine/Plugins/PCG/Source/PCG/Private/Data/Tool/PCGToolBaseData.cpp
@@ -14,7 +14,7 @@
 
 bool FPCGInteractiveToolWorkingData::IsInitialized() const
 {
-	return bInitialized && WorkingDataIdentifier.IsValid() && !WorkingDataIdentifier.IsNone();
+	return WorkingDataIdentifier.IsValid() && !WorkingDataIdentifier.IsNone();
 }
 
 bool FPCGInteractiveToolWorkingData::IsValid() const
@@ -34,8 +34,6 @@ void FPCGInteractiveToolWorkingData::Initialize(const FPCGInteractiveToolWorking
 	WorkingDataIdentifier = Context.WorkingDataIdentifier;
 	
 	InitializeInternal(Context);
-
-	bInitialized = true;
 }
 
 void FPCGInteractiveToolWorkingData::InitializeInternal(const FPCGInteractiveToolWorkingDataContext& Context)
diff --git a/Engine/Plugins/PCG/Source/PCGEditor/Private/EditorMode/Tools/PCGInteractiveToolSettings.cpp b/Engine/Plugins/PCG/Source/PCGEditor/Private/EditorMode/Tools/PCGInteractiveToolSettings.cpp
--- a/Engine/Plugins/PCG/Source/PCGEditor/Private/EditorMode/Tools/PCGInteractiveToolSettings.cpp
+++ b/Engine/Plugins/PCG/Source/PCGEditor/Private/EditorMode/Tools/PCGInteractiveToolSettings.cpp
@@ -280,13 +280,11 @@ void UPCGInteractiveToolSettings::InitializeWorkingData(FName DataInstanceIdenti
 	Context.InteractiveTool = GetTypedOuter<UInteractiveTool>();
 	Context.PCGSettings = this;
 	Context.WorkingDataIdentifier = GetWorkingDataIdentifier(DataInstanceIdentifier);
-
-	FPCGInteractiveToolWorkingData* WorkingData = PCGComponent->ToolDataContainer.GetMutableTypedWorkingData<FPCGInteractiveToolWorkingData>(Context.WorkingDataIdentifier);
 	
 	if (EditActor.IsValid() && PCGComponent)
 	{
 		// If we don't find the requested working data, we create and initialize it
-		if(WorkingData == nullptr)
+		if (GetMutableTypedWorkingData(DataInstanceIdentifier) == nullptr)
 		{
 			TInstancedStruct<FPCGInteractiveToolWorkingData> NewToolWorkingData;
 			NewToolWorkingData.InitializeAsScriptStruct(Struct);
@@ -299,7 +297,7 @@ void UPCGInteractiveToolSettings::InitializeWorkingData(FName DataInstanceIdenti
 		}
 	}
 
-	WorkingData = GetMutableTypedWorkingData(DataInstanceIdentifier);
+	FPCGInteractiveToolWorkingData* WorkingData = GetMutableTypedWorkingData(DataInstanceIdentifier);
 	
 	if(!ensureMsgf(WorkingData, TEXT("Working Data has to be allocated at this point!.")))
 	{
