[PCG] Fixed get tool data title not updating, added recursive tool name gathering, renamed draw surface for draw spline surface, improved the way we prepare the points in the paint tool data.

#rnx

diff --git a/Engine/Plugins/PCG/Source/PCG/Private/Data/Tool/PCGToolPointData.cpp b/Engine/Plugins/PCG/Source/PCG/Private/Data/Tool/PCGToolPointData.cpp
--- a/Engine/Plugins/PCG/Source/PCG/Private/Data/Tool/PCGToolPointData.cpp
+++ b/Engine/Plugins/PCG/Source/PCG/Private/Data/Tool/PCGToolPointData.cpp
@@ -55,18 +55,21 @@ void FPCGInteractiveToolWorkingData_PointArrayData::InitializeRuntimeElementData
 
 	FPCGTaggedData& Output = Outputs.Emplace_GetRef();
 
+	if (GeneratedPointData)
+	{
 #if WITH_EDITOR
-	// Implementation note - we duplicate here to make sure we are not preserving the parent, because it might still change, which would cause issues.
-	check(IsInGameThread());
-	Output.Data = Cast<UPCGData>(StaticDuplicateObject(GeneratedPointData, GetTransientPackage()));
+		// Implementation note - we duplicate here to make sure we are not preserving the parent, because it might still change, which would cause issues.
+		check(IsInGameThread());
+		Output.Data = Cast<UPCGData>(StaticDuplicateObject(GeneratedPointData, GetTransientPackage()));
 #else
-	Output.Data = GeneratedPointData->DuplicateData(Context);
+		Output.Data = GeneratedPointData->DuplicateData(Context);
 #endif
 
 #if WITH_EDITOR
-	FPCGSelectionKey Key = FPCGSelectionKey::CreateFromObjectPtr(GeneratedPointData);
-	FPCGDynamicTrackingHelper::AddSingleDynamicTrackingKey(Context, MoveTemp(Key), false);
+		FPCGSelectionKey Key = FPCGSelectionKey::CreateFromObjectPtr(GeneratedPointData);
+		FPCGDynamicTrackingHelper::AddSingleDynamicTrackingKey(Context, MoveTemp(Key), false);
 #endif
+	}
 }
 
 bool FPCGInteractiveToolWorkingData_PointArrayData::IsValid() const
diff --git a/Engine/Plugins/PCG/Source/PCG/Private/Data/Tool/PCGToolVolumeData.cpp b/Engine/Plugins/PCG/Source/PCG/Private/Data/Tool/PCGToolVolumeData.cpp
--- a/Engine/Plugins/PCG/Source/PCG/Private/Data/Tool/PCGToolVolumeData.cpp
+++ b/Engine/Plugins/PCG/Source/PCG/Private/Data/Tool/PCGToolVolumeData.cpp
@@ -44,7 +44,7 @@ void FPCGInteractiveToolWorkingData_Volume::OnToolStart(const FPCGInteractiveToo
 
 void FPCGInteractiveToolWorkingData_Volume::OnToolCancel(const FPCGInteractiveToolWorkingDataContext& Context)
 {
-	if (BrushComponent.IsValid() && GeneratedResources.GeneratedComponents.Contains(BrushComponent) == false && OnToolStartBrushComponent)
+	if (Volume.IsValid() && BrushComponent.IsValid() && GeneratedResources.GeneratedComponents.Contains(BrushComponent) == false && OnToolStartBrushComponent)
 	{
 		Volume->SetActorTransform(OnToolStartVolumeTransform);
 
diff --git a/Engine/Plugins/PCG/Source/PCG/Private/Elements/PCGDataFromTool.cpp b/Engine/Plugins/PCG/Source/PCG/Private/Elements/PCGDataFromTool.cpp
--- a/Engine/Plugins/PCG/Source/PCG/Private/Elements/PCGDataFromTool.cpp
+++ b/Engine/Plugins/PCG/Source/PCG/Private/Elements/PCGDataFromTool.cpp
@@ -14,12 +14,10 @@
 
 #include "Algo/AnyOf.h"
 
-
 #include UE_INLINE_GENERATED_CPP_BY_NAME(PCGDataFromTool)
 
 #define LOCTEXT_NAMESPACE "PCGDataFromTool"
 
-#if WITH_EDITOR
 FString UPCGDataFromTool::GetAdditionalTitleInformation() const
 {
 	FStringBuilderBase Builder;
@@ -47,10 +45,24 @@ FString UPCGDataFromTool::GetAdditionalTitleInformation() const
 	return Builder.ToString();
 }
 
+#if WITH_EDITOR
 FText UPCGDataFromTool::GetNodeTooltipText() const
 {
 	return LOCTEXT("DataFromToolTooltip", "Builds a collection of PCG-compatible data from the selected tools.");
 }
+
+EPCGChangeType UPCGDataFromTool::GetChangeTypeForProperty(const FName& InPropertyName) const
+{
+	EPCGChangeType ChangeType = Super::GetChangeTypeForProperty(InPropertyName);
+
+	if (InPropertyName == GET_MEMBER_NAME_CHECKED(UPCGDataFromTool, ToolTag) ||
+		InPropertyName == GET_MEMBER_NAME_CHECKED(UPCGDataFromTool, DataInstance))
+	{
+		ChangeType |= EPCGChangeType::Cosmetic;
+	}
+
+	return ChangeType;
+}
 #endif
 
 FPCGElementPtr UPCGDataFromTool::CreateElement() const
diff --git a/Engine/Plugins/PCG/Source/PCG/Public/Elements/PCGDataFromTool.h b/Engine/Plugins/PCG/Source/PCG/Public/Elements/PCGDataFromTool.h
--- a/Engine/Plugins/PCG/Source/PCG/Public/Elements/PCGDataFromTool.h
+++ b/Engine/Plugins/PCG/Source/PCG/Public/Elements/PCGDataFromTool.h
@@ -22,16 +22,18 @@ public:
 #if WITH_EDITOR
 	virtual FName GetDefaultNodeName() const override { return FName(TEXT("GetToolData")); }
 	virtual FText GetDefaultNodeTitle() const override { return NSLOCTEXT("PCGDataFromTool", "NodeTitle", "Get Tool Data"); }
-	PCG_API virtual FString GetAdditionalTitleInformation() const override;
-	PCG_API virtual FText GetNodeTooltipText() const override;
+	virtual FText GetNodeTooltipText() const override;
 	virtual EPCGSettingsType GetType() const override { return EPCGSettingsType::Generic; }
 	virtual bool CanDynamicallyTrackKeys() const override { return true; }
+	virtual EPCGChangeType GetChangeTypeForProperty(const FName& InPropertyName) const override;
 #endif
 
+	virtual FString GetAdditionalTitleInformation() const override;
+
 protected:
-	PCG_API virtual TArray<FPCGPinProperties> InputPinProperties() const override;
-	PCG_API virtual TArray<FPCGPinProperties> OutputPinProperties() const override;
-	PCG_API virtual FPCGElementPtr CreateElement() const override;
+	virtual TArray<FPCGPinProperties> InputPinProperties() const override;
+	virtual TArray<FPCGPinProperties> OutputPinProperties() const override;
+	virtual FPCGElementPtr CreateElement() const override;
 	//~End UPCGSettings interface
 
 public:
@@ -53,6 +55,6 @@ public:
 	virtual bool IsCacheable(const UPCGSettings* InSettings) const override { return false; }
 protected:
 	virtual bool CanExecuteOnlyOnMainThread(FPCGContext* Context) const override { return true; }
-	PCG_API virtual bool ExecuteInternal(FPCGContext* Context) const override;
+	virtual bool ExecuteInternal(FPCGContext* Context) const override;
 	virtual bool SupportsBasePointDataInputs(FPCGContext* InContext) const override { return true; }
 };
diff --git a/Engine/Plugins/PCG/Source/PCGEditor/Private/EditorMode/PCGEdModeCommands.cpp b/Engine/Plugins/PCG/Source/PCGEditor/Private/EditorMode/PCGEdModeCommands.cpp
--- a/Engine/Plugins/PCG/Source/PCGEditor/Private/EditorMode/PCGEdModeCommands.cpp
+++ b/Engine/Plugins/PCG/Source/PCGEditor/Private/EditorMode/PCGEdModeCommands.cpp
@@ -87,7 +87,7 @@ void FPCGEditorModeToolCommands::RegisterCommands()
 	/*** PCG EDITOR MODE TOOLS **/
 
 	PCG_TOOL_COMMAND_INFO(EnableDrawSplineTool, "DrawSpline", "Draw Spline", "Use the '{ToolTag}' tool to build a spline for PCG.", UPCGInteractiveToolSettings_Spline::StaticGetToolTag(), "PCGEditorMode.Tools.DrawSpline");
-	PCG_TOOL_COMMAND_INFO(EnableDrawSurfaceTool, "DrawSurface", "Draw Surface", "Use the '{ToolTag}' tool to build a closed spline representing an area for PCG.", UPCGInteractiveToolSettings_SplineSurface::StaticGetToolTag(), "PCGEditorMode.Tools.DrawSurface");
+	PCG_TOOL_COMMAND_INFO(EnableDrawSurfaceTool, "DrawSurface", "Draw Spline Surface", "Use the '{ToolTag}' tool to build a closed spline representing an area for PCG.", UPCGInteractiveToolSettings_SplineSurface::StaticGetToolTag(), "PCGEditorMode.Tools.DrawSurface");
 	PCG_TOOL_COMMAND_INFO(EnablePaintTool, "Paint", "Paint", "Use the '{ToolTag}' tool to apply PCG points with a brush.", UPCGInteractiveToolSettings_PaintTool::StaticGetToolTag(), "PCGEditorMode.Context.Paint");
 	PCG_TOOL_COMMAND_INFO(EnableVolumeTool, "Volume", "Volume", "Create or updates a volume in the scene using the '{ToolTag}' tool.", UPCGInteractiveToolSettings_Volume::StaticGetToolTag(), "PCGEditorMode.Tools.Volume");
 }
diff --git a/Engine/Plugins/PCG/Source/PCGEditor/Private/EditorMode/Tools/PCGInteractiveToolSettings.cpp b/Engine/Plugins/PCG/Source/PCGEditor/Private/EditorMode/Tools/PCGInteractiveToolSettings.cpp
--- a/Engine/Plugins/PCG/Source/PCGEditor/Private/EditorMode/Tools/PCGInteractiveToolSettings.cpp
+++ b/Engine/Plugins/PCG/Source/PCGEditor/Private/EditorMode/Tools/PCGInteractiveToolSettings.cpp
@@ -713,7 +713,7 @@ TArray<FName> UPCGInteractiveToolSettings::GetValidDataInstanceNamesForGraph() c
 		return Result;
 	}
 
-	TArray<UPCGNode*> Nodes = ToolGraph->GetGraph()->FindNodesWithSettings(UPCGDataFromTool::StaticClass(), false);
+	TArray<UPCGNode*> Nodes = ToolGraph->GetGraph()->FindNodesWithSettings(UPCGDataFromTool::StaticClass(), /*bRecursive=*/true);
 
 	FName ToolTag = GetToolTag();
 	for (UPCGNode* Node : Nodes)
diff --git a/Engine/Plugins/PCG/Source/PCGEditor/Private/EditorMode/Tools/Paint/PCGPaintTool.cpp b/Engine/Plugins/PCG/Source/PCGEditor/Private/EditorMode/Tools/Paint/PCGPaintTool.cpp
--- a/Engine/Plugins/PCG/Source/PCGEditor/Private/EditorMode/Tools/Paint/PCGPaintTool.cpp
+++ b/Engine/Plugins/PCG/Source/PCGEditor/Private/EditorMode/Tools/Paint/PCGPaintTool.cpp
@@ -196,7 +196,7 @@ void UPCGPaintTool::OnUpdateDrag(const FRay& Ray)
 			if (FVector::DistSquared(LastBrushStamp.WorldPosition, ToolSettings->LastPointTransform.GetValue().GetLocation()) >= ToolSettings->CachedMinPointSpacingSquared)
 			{
 				AddSinglePointFromBrush();
-			}		
+			}
 		}
 		else
 		{
@@ -276,7 +276,7 @@ void UPCGPaintTool::UpdateBrushStampIndicator()
 void UPCGPaintTool::ShutdownBrushStampIndicator()
 {
 	if (BrushStampIndicator)
-	{		
+	{
 		GetToolManager()->GetPairedGizmoManager()->DestroyGizmo(BrushStampIndicator);
 		BrushStampIndicator = nullptr;
 		GetToolManager()->GetPairedGizmoManager()->DeregisterGizmoType(UE::PCG::EditorMode::Gizmos::Constants::BrushIdentifier_VolumetricBox);
@@ -306,89 +306,30 @@ void UPCGPaintTool::AddSinglePointFromBrush()
 		return;
 	}
 
+	const int32 PointIndex = PointData->GetNumPoints();
+	PointData->SetNumPoints(PointIndex + 1);
+
 	TOptional<FRotator> BrushRotator;
 	if (UPCGVolumetricBrushStampIndicator* TypedBrushStampIndicator = Cast<UPCGVolumetricBrushStampIndicator>(BrushStampIndicator))
 	{
 		BrushRotator = TypedBrushStampIndicator->GetPreviewMesh()->GetActor()->GetActorRotation();
 	}
-	
+
 	const FTransform PointTransform = FTransform(BrushRotator.Get(FRotator::ZeroRotator), LastBrushStamp.WorldPosition);
-	
-	const int32 PointIndex = PointData->GetNumPoints();
-	PointData->SetNumPoints(PointIndex + 1);
-	
-	// Values to assign
-	const float PointSteepness = 1.0f;
-	const FVector& Position = PointTransform.GetLocation();
-	const int32 PointSeed = PCGHelpers::ComputeSeed((int)Position.X, (int)Position.Y, (int)Position.Z);
-	const FBox LocalBounds = FBox(FVector(-CurrentBrushRadius), FVector(CurrentBrushRadius));
-	const FVector PointBoundsMin = LocalBounds.Min;
-	const FVector PointBoundsMax = LocalBounds.Max;
-	const int64 PointMetadataEntry = PointData->Metadata->AddEntry();
-
-	// SteepnessRange - Initialize and Allocate if needed 
-	const TConstPCGValueRange<float> ConstSteepnessRange = PointData->GetConstSteepnessValueRange();
-	TOptional<const float> SteepnessSingleValue = ConstSteepnessRange.GetSingleValue();
-	const bool bAllocateSteepness = SteepnessSingleValue.IsSet() && SteepnessSingleValue.GetValue() != PointSteepness;
-	TPCGValueRange<float> SteepnessRange = PointData->GetSteepnessValueRange(bAllocateSteepness);
-	
-	// TransformRange - Initialize and Allocate if needed 
-	const TConstPCGValueRange<FTransform> ConstTransformRange = PointData->GetConstTransformValueRange();
-	TOptional<const FTransform> TransformSingleValue = ConstTransformRange.GetSingleValue();
-	const bool bAllocateTransform = TransformSingleValue.IsSet() && !TransformSingleValue.GetValue().Equals(PointTransform);
-	TPCGValueRange<FTransform> TransformRange = PointData->GetTransformValueRange(bAllocateTransform);
-
-	// SeedRange - Initialize and Allocate if needed 
-	const TConstPCGValueRange<int32> ConstSeedRange = PointData->GetConstSeedValueRange();
-	TOptional<const int32> SeedSingleValue = ConstSeedRange.GetSingleValue();
-	const bool bAllocateSeed = SeedSingleValue.IsSet() && SeedSingleValue.GetValue() != PointSeed;
-	TPCGValueRange<int32> SeedRange = PointData->GetSeedValueRange(bAllocateSeed);
-
-	// BoundsMinRange - Initialize and Allocate if needed 
-	const TConstPCGValueRange<FVector> ConstBoundsMinRange = PointData->GetConstBoundsMinValueRange();
-	TOptional<const FVector> BoundsMinSingleValue = ConstBoundsMinRange.GetSingleValue();
-	const bool bAllocateBoundsMin = BoundsMinSingleValue.IsSet() && BoundsMinSingleValue.GetValue() != PointBoundsMin;
-	TPCGValueRange<FVector> BoundsMinRange = PointData->GetBoundsMinValueRange(bAllocateBoundsMin);
-		
-	// BoundsMaxRange - Initialize and Allocate if needed 
-	const TConstPCGValueRange<FVector> ConstBoundsMaxRange = PointData->GetConstBoundsMaxValueRange();
-	TOptional<const FVector> BoundsMaxSingleValue = ConstBoundsMaxRange.GetSingleValue();
-	const bool bAllocateBoundsMax = BoundsMaxSingleValue.IsSet() && BoundsMaxSingleValue.GetValue() != PointBoundsMax;
-	TPCGValueRange<FVector> BoundsMaxRange = PointData->GetBoundsMaxValueRange(bAllocateBoundsMax);
-
-	// MetadataEntryRange - Initialize and Allocate if needed 
-	const TConstPCGValueRange<int64> ConstMetadataEntryRange = PointData->GetConstMetadataEntryValueRange();
-	TOptional<const int64> MetadataEntrySingleValue = ConstMetadataEntryRange.GetSingleValue();
-	const bool bAllocateMetadataEntry = MetadataEntrySingleValue.IsSet() && MetadataEntrySingleValue.GetValue() != PointMetadataEntry;
-	TPCGValueRange<int64> MetadataEntryRange = PointData->GetMetadataEntryValueRange(bAllocateMetadataEntry);
-
-	// Assign values
-	SteepnessRange[PointIndex] = PointSteepness;
-	TransformRange[PointIndex] = PointTransform;
-	SeedRange[PointIndex] = PointSeed;
-	BoundsMinRange[PointIndex] = PointBoundsMin;
-	BoundsMaxRange[PointIndex] = PointBoundsMax;
-	MetadataEntryRange[PointIndex] = PointMetadataEntry;
-
-	FPCGAttributeIdentifier BrushSizeAttributeIdentifier;
-	BrushSizeAttributeIdentifier.Name = "BrushSize";
-	BrushSizeAttributeIdentifier.MetadataDomain = PointData->Metadata->GetConstDefaultMetadataDomain()->GetDomainID();
-	
-	FPCGMetadataAttribute<float>* BrushSizeAttribute = PointData->Metadata->FindOrCreateAttribute(BrushSizeAttributeIdentifier, 100.f, /*bAllowsInterpolation=*/false, /*bOverrideParent=*/false, /*bOverwriteIfTypeMismatch=*/false);
-	if (BrushSizeAttribute)
-	{
-		BrushSizeAttribute->SetValue(PointMetadataEntry, CurrentBrushRadius);
-	}
+	FPCGPoint NewPoint(PointTransform, BrushProperties->BrushStrength, PCGHelpers::ComputeSeedFromPosition(PointTransform.GetLocation()));
+	// Affect bounds based on brush size & falloff.
+	const float PointSteepness = (BrushProperties->bShowFalloff ? 1.0f - BrushProperties->BrushFalloffAmount : 1.0f);
+	NewPoint.Steepness = PointSteepness;
 
-	FPCGAttributeIdentifier FalloffAttributeIdentifier;
-	FalloffAttributeIdentifier.Name = "BrushFalloff";
-	FalloffAttributeIdentifier.MetadataDomain = PointData->Metadata->GetConstDefaultMetadataDomain()->GetDomainID();
-	FPCGMetadataAttribute<float>* FalloffAttribute = PointData->Metadata->FindOrCreateAttribute(FalloffAttributeIdentifier, 100.f, /*bAllowsInterpolation=*/false, /*bOverrideParent=*/false, /*bOverwriteIfTypeMismatch=*/false);
-	if (FalloffAttribute)
-	{
-		FalloffAttribute->SetValue(PointMetadataEntry, BrushProperties->BrushFalloffAmount);
-	}
+	const float BoundsAdjustment = 0.5f + 0.5f * PointSteepness;
+
+	NewPoint.BoundsMin = FVector(-CurrentBrushRadius * BoundsAdjustment);
+	NewPoint.BoundsMax = FVector(+CurrentBrushRadius * BoundsAdjustment);
 	
+	// We use everything but the metadata entry & color
+	PointData->AllocateProperties(PointData->GetAllocatedProperties() | (EPCGPointNativeProperties::All & ~EPCGPointNativeProperties::Color & ~EPCGPointNativeProperties::MetadataEntry));
+	FPCGPointValueRanges PointDataRanges(PointData, /*bAllocate=*/false);
+	PointDataRanges.SetFromPoint(PointIndex, NewPoint);
 
 	// @todo_pcg
 	// Instead of using a queue on timer, can we put this on another thread? Direct broadcasting takes too much perf on tick, timed queue results in delays
