[PCG] Added experimental tags & removed api exposure

#rnx

diff --git a/Engine/Plugins/PCG/Source/PCG/Internal/Data/Tool/PCGToolBaseData.h b/Engine/Plugins/PCG/Source/PCG/Internal/Data/Tool/PCGToolBaseData.h
--- a/Engine/Plugins/PCG/Source/PCG/Internal/Data/Tool/PCGToolBaseData.h
+++ b/Engine/Plugins/PCG/Source/PCG/Internal/Data/Tool/PCGToolBaseData.h
@@ -10,7 +10,9 @@ class UInteractiveTool;
 class UInteractiveToolPropertySet;
 struct FPCGContext;
 
-/** A transient struct that contains resources generated during tool usage for a specific tool data struct. Used to finalize or clean up when a tool is applied or cancelled. */
+/** [EXPERIMENTAL]
+* A transient struct that contains resources generated during tool usage for a specific tool data struct. Used to finalize or clean up when a tool is applied or cancelled.
+*/
 USTRUCT()
 struct FPCGInteractiveToolWorkingDataGeneratedResources
 {
@@ -21,17 +23,19 @@ struct FPCGInteractiveToolWorkingDataGeneratedResources
 
 #if WITH_EDITORONLY_DATA
 	/** Finalizes generated resources (e.g. remove Transient flag). */
-	PCG_API void Finalize();
+	void Finalize();
 
 	/** Remove and destroy generated resources. */
-	PCG_API void Cleanup();
+	void Cleanup();
 
 	UPROPERTY(Transient)
 	TArray<TWeakObjectPtr<UActorComponent>> GeneratedComponents;
 #endif
 };
 
-/** A transient context struct that is passed in so the tool data can manage its data properly during tool usage. Used at editor-time only. */
+/** [EXPERIMENTAL]
+* A transient context struct that is passed in so the tool data can manage its data properly during tool usage. Used at editor-time only.
+*/
 USTRUCT()
 struct FPCGInteractiveToolWorkingDataContext
 {
@@ -68,8 +72,10 @@ struct FPCGInteractiveToolWorkingDataContext
 #endif
 };
 
-/** The base struct for working data for a tool. This data is created by the tool, stored on the PCG Component,
- * then read back by the tool on subsequent tool invocations. */
+/** [EXPERIMENTAL]
+* The base struct for working data for a tool. This data is created by the tool, stored on the PCG Component,
+* then read back by the tool on subsequent tool invocations.
+*/
 USTRUCT()
 struct FPCGInteractiveToolWorkingData : public FPCGInteractiveToolWorkingBaseData
 {
@@ -82,7 +88,7 @@ struct FPCGInteractiveToolWorkingData : public FPCGInteractiveToolWorkingBaseDat
 	/** Construct and add pcg data to the runtime context.
 	 * This will translate data from this struct into UPCGData to be processed by the graph.
 	 * Add dynamic tracking keys in here if you want to refresh when a change to your source data happens. */
-	PCG_API virtual void InitializeRuntimeElementData(FPCGContext* InContext) const;
+	virtual void InitializeRuntimeElementData(FPCGContext* InContext) const;
 
 #if WITH_EDITORONLY_DATA
 public:
@@ -104,8 +110,7 @@ public:
 
 private:
 	/** Implement your initialization logic here. Called by Initialize. */
-	PCG_API virtual void InitializeInternal(const FPCGInteractiveToolWorkingDataContext& Context);
-	
+	virtual void InitializeInternal(const FPCGInteractiveToolWorkingDataContext& Context);
 protected:
 	/** Each working data struct keeps track of its generated resources. They have to be reverted during OnToolCancel. */
 	UPROPERTY(Transient)
diff --git a/Engine/Plugins/PCG/Source/PCG/Internal/Data/Tool/PCGToolPointData.h b/Engine/Plugins/PCG/Source/PCG/Internal/Data/Tool/PCGToolPointData.h
--- a/Engine/Plugins/PCG/Source/PCG/Internal/Data/Tool/PCGToolPointData.h
+++ b/Engine/Plugins/PCG/Source/PCG/Internal/Data/Tool/PCGToolPointData.h
@@ -6,8 +6,10 @@
 
 #include "PCGToolPointData.generated.h"
 
-/** Tool Data that stores points. */
-USTRUCT(BlueprintType, DisplayName = "Point Array Tool Data")
+/** [EXPERIMENTAL]
+* Tool Data that stores points.
+*/
+USTRUCT(DisplayName = "Point Array Tool Data")
 struct FPCGInteractiveToolWorkingData_PointArrayData : public FPCGInteractiveToolWorkingData
 {
 	GENERATED_BODY()
@@ -15,24 +17,21 @@ struct FPCGInteractiveToolWorkingData_PointArrayData : public FPCGInteractiveToo
 public:
 	FPCGInteractiveToolWorkingData_PointArrayData() = default;
 
-	PCG_API virtual void InitializeRuntimeElementData(FPCGContext* InContext) const override;
-
 	PCG_API virtual bool IsValid() const override;
+	virtual void InitializeRuntimeElementData(FPCGContext* InContext) const override;
 
 	PCG_API UPCGPointArrayData* GetPointArrayData() const;
 
 #if WITH_EDITORONLY_DATA
 	PCG_API virtual void InitializeInternal(const FPCGInteractiveToolWorkingDataContext& Context) override;
-
 	PCG_API virtual void OnToolStart(const FPCGInteractiveToolWorkingDataContext& Context) override;
-
 	PCG_API virtual void OnToolCancel(const FPCGInteractiveToolWorkingDataContext& Context) override;
-
 	PCG_API virtual void OnResetToolDataRequested(const FPCGInteractiveToolWorkingDataContext& Context) override;
 #endif
 
 protected:
-	UPROPERTY(VisibleAnywhere, BlueprintReadOnly, Category = "Working Data")
+	/** Contents of the tool data. */
+	UPROPERTY(VisibleAnywhere, Category = "Working Data")
 	TObjectPtr<UPCGPointArrayData> GeneratedPointData;
 
 	/** A transient copy of the point array data, created on tool start/data access. We revert back to these points on Cancel. */
diff --git a/Engine/Plugins/PCG/Source/PCG/Internal/Data/Tool/PCGToolSplineData.h b/Engine/Plugins/PCG/Source/PCG/Internal/Data/Tool/PCGToolSplineData.h
--- a/Engine/Plugins/PCG/Source/PCG/Internal/Data/Tool/PCGToolSplineData.h
+++ b/Engine/Plugins/PCG/Source/PCG/Internal/Data/Tool/PCGToolSplineData.h
@@ -7,8 +7,10 @@
 
 #include "PCGToolSplineData.generated.h"
 
-/** The working data for splines. */
-USTRUCT(BlueprintType, DisplayName = "Spline Tool Data")
+/** [EXPERIMENTAL]
+* Tool-created working data for spline interaction.
+*/
+USTRUCT(DisplayName = "Spline Tool Data")
 struct FPCGInteractiveToolWorkingData_Spline : public FPCGInteractiveToolWorkingData
 {
 	GENERATED_BODY()
@@ -16,17 +18,16 @@ public:
 	FPCGInteractiveToolWorkingData_Spline() = default;
 	virtual ~FPCGInteractiveToolWorkingData_Spline() override = default;
 
-	PCG_API virtual void InitializeRuntimeElementData(FPCGContext* InContext) const override;
-
 	PCG_API virtual bool IsValid() const override;
+	virtual void InitializeRuntimeElementData(FPCGContext* InContext) const override;
 
 	PCG_API USplineComponent* GetSplineComponent() const;
 
 #if WITH_EDITOR
 public:
-	virtual void OnToolStart(const FPCGInteractiveToolWorkingDataContext& Context) override;
-	virtual void OnToolCancel(const FPCGInteractiveToolWorkingDataContext& Context) override;
-	virtual void OnResetToolDataRequested(const FPCGInteractiveToolWorkingDataContext& Context) override;
+	PCG_API virtual void OnToolStart(const FPCGInteractiveToolWorkingDataContext& Context) override;
+	PCG_API virtual void OnToolCancel(const FPCGInteractiveToolWorkingDataContext& Context) override;
+	PCG_API virtual void OnResetToolDataRequested(const FPCGInteractiveToolWorkingDataContext& Context) override;
 private:
 	virtual void InitializeInternal(const FPCGInteractiveToolWorkingDataContext& Context) override;
 
@@ -38,14 +39,16 @@ private:
 #endif
 
 protected:
-	UPROPERTY(VisibleAnywhere, BlueprintReadOnly, Category = "Working Data")
+	/** Spline this working data interacts with. */
+	UPROPERTY(VisibleAnywhere, Category = "Working Data")
 	TWeakObjectPtr<USplineComponent> SplineComponent;
 
+	/** Copy of the spline at the beginning of the interaction, for rollback purposes. */
 	UPROPERTY(Transient)
 	TObjectPtr<USplineComponent> OnToolStartSplineComponent;
 };
 
-USTRUCT(BlueprintType, DisplayName = "Spline Surface Tool Data")
+USTRUCT(DisplayName = "Spline Surface Tool Data")
 struct FPCGInteractiveToolWorkingData_SplineSurface : public FPCGInteractiveToolWorkingData_Spline
 {
 	GENERATED_BODY()
diff --git a/Engine/Plugins/PCG/Source/PCG/Internal/Data/Tool/PCGToolVolumeData.h b/Engine/Plugins/PCG/Source/PCG/Internal/Data/Tool/PCGToolVolumeData.h
--- a/Engine/Plugins/PCG/Source/PCG/Internal/Data/Tool/PCGToolVolumeData.h
+++ b/Engine/Plugins/PCG/Source/PCG/Internal/Data/Tool/PCGToolVolumeData.h
@@ -10,21 +10,26 @@ class AVolume;
 class UBrushComponent;
 class UBrushBuilder;
 
-USTRUCT(BlueprintType, DisplayName = "Volume Tool Data")
+/** [EXPERIMENTAL] 
+* Data to hold interaction data with volumes.
+*/
+USTRUCT(DisplayName = "Volume Tool Data")
 struct FPCGInteractiveToolWorkingData_Volume : public FPCGInteractiveToolWorkingData
 {
 	GENERATED_BODY()
 public:
-	UPROPERTY(VisibleAnywhere, BlueprintReadOnly, Category = "Working Data")
+	/** Volume actor this working data is affecting. */
+	UPROPERTY(VisibleAnywhere, Category = "Working Data")
 	TWeakObjectPtr<AVolume> Volume = nullptr;
 
-	UPROPERTY(VisibleAnywhere, BlueprintReadOnly, Category = "Working Data")
+	/** Brush component this working data is affecting. */
+	UPROPERTY(VisibleAnywhere, Category = "Working Data")
 	TWeakObjectPtr<UBrushComponent> BrushComponent = nullptr;
 
 #if WITH_EDITOR
 public:
-	virtual void OnToolStart(const FPCGInteractiveToolWorkingDataContext& Context) override;
-	virtual void OnToolCancel(const FPCGInteractiveToolWorkingDataContext& Context) override;
+	PCG_API virtual void OnToolStart(const FPCGInteractiveToolWorkingDataContext& Context) override;
+	PCG_API virtual void OnToolCancel(const FPCGInteractiveToolWorkingDataContext& Context) override;
 
 private:
 	virtual void InitializeInternal(const FPCGInteractiveToolWorkingDataContext& Context) override;
@@ -36,12 +41,14 @@ private:
 
 protected:
 #if WITH_EDITORONLY_DATA
+	/** Copy of the original brush component on tool initialization, used for rollback purposes. */
 	UPROPERTY(Transient)
 	TObjectPtr<UBrushComponent> OnToolStartBrushComponent = nullptr;
 
 	UPROPERTY(Transient)
 	TObjectPtr<UBrushBuilder> OnToolStartBrushBuilder = nullptr;
 
+	/** Original transform of the volume actor, used for rollback. */
 	UPROPERTY(Transient)
 	FTransform OnToolStartVolumeTransform;
 #endif
diff --git a/Engine/Plugins/PCG/Source/PCGEditor/Private/EditorMode/PCGEdMode.h b/Engine/Plugins/PCG/Source/PCGEditor/Private/EditorMode/PCGEdMode.h
--- a/Engine/Plugins/PCG/Source/PCGEditor/Private/EditorMode/PCGEdMode.h
+++ b/Engine/Plugins/PCG/Source/PCGEditor/Private/EditorMode/PCGEdMode.h
@@ -17,7 +17,7 @@ class UPCGEditorMode : public UEdMode
 
 public:
 	/** Unique name for PCG Editor Mode */
-	PCGEDITOR_API const static FEditorModeID EM_PCGEditorModeId;
+	const static FEditorModeID EM_PCGEditorModeId;
 
 	UPCGEditorMode();
 
diff --git a/Engine/Plugins/PCG/Source/PCGEditor/Private/EditorMode/PCGEdModeCommands.h b/Engine/Plugins/PCG/Source/PCGEditor/Private/EditorMode/PCGEdModeCommands.h
--- a/Engine/Plugins/PCG/Source/PCGEditor/Private/EditorMode/PCGEdModeCommands.h
+++ b/Engine/Plugins/PCG/Source/PCGEditor/Private/EditorMode/PCGEdModeCommands.h
@@ -6,7 +6,7 @@
 #include "Tools/InteractiveToolsCommands.h"
 
 /** General actions related to the PCG Editor Mode. */
-class PCGEDITOR_API FPCGEditorModeCommands : public TCommands<FPCGEditorModeCommands>
+class FPCGEditorModeCommands : public TCommands<FPCGEditorModeCommands>
 {
 public:
 	FPCGEditorModeCommands();
@@ -17,7 +17,7 @@ public:
 };
 
 /** Actions related to the PCG Editor Mode Toolkit Palette. */
-class PCGEDITOR_API FPCGEditorModePaletteCommands : public TCommands<FPCGEditorModePaletteCommands>
+class FPCGEditorModePaletteCommands : public TCommands<FPCGEditorModePaletteCommands>
 {
 public:
 	FPCGEditorModePaletteCommands();
diff --git a/Engine/Plugins/PCG/Source/PCGEditor/Private/EditorMode/Tools/Helpers/PCGEdModeSceneQueryHelpers.h b/Engine/Plugins/PCG/Source/PCGEditor/Private/EditorMode/Tools/Helpers/PCGEdModeSceneQueryHelpers.h
--- a/Engine/Plugins/PCG/Source/PCGEditor/Private/EditorMode/Tools/Helpers/PCGEdModeSceneQueryHelpers.h
+++ b/Engine/Plugins/PCG/Source/PCGEditor/Private/EditorMode/Tools/Helpers/PCGEdModeSceneQueryHelpers.h
@@ -19,8 +19,8 @@ struct FPCGRaycastFilterRule
 {
 	GENERATED_BODY()
 
-	PCGEDITOR_API FPCGRaycastFilterRule();
-	PCGEDITOR_API virtual ~FPCGRaycastFilterRule();
+	FPCGRaycastFilterRule();
+	virtual ~FPCGRaycastFilterRule();
 	
 	/** If this returns false, the raycast hit is filtered out.
 	 *  If it returns true, this rule has determined is valid under this given rule.
diff --git a/Engine/Plugins/PCG/Source/PCGEditor/Private/EditorMode/Tools/Line/PCGDrawSplineTool.h b/Engine/Plugins/PCG/Source/PCGEditor/Private/EditorMode/Tools/Line/PCGDrawSplineTool.h
--- a/Engine/Plugins/PCG/Source/PCGEditor/Private/EditorMode/Tools/Line/PCGDrawSplineTool.h
+++ b/Engine/Plugins/PCG/Source/PCGEditor/Private/EditorMode/Tools/Line/PCGDrawSplineTool.h
@@ -15,7 +15,7 @@
 class USplineComponent;
 class USingleClickOrDragInputBehavior;
 
-UENUM(BlueprintType)
+UENUM()
 enum class EPCGDrawSplineDrawMode : uint8
 {
 	TangentDrag UMETA(ToolTip = "Click to place a point and then drag to set its tangent. Clicking without dragging will create sharp corners."),
@@ -23,8 +23,8 @@ enum class EPCGDrawSplineDrawMode : uint8
 	FreeDraw UMETA(ToolTip = "Drag to place multiple points, with spacing controlled by Min Point Spacing.")
 };
 
-UCLASS(Abstract, BlueprintType)
-class PCGEDITOR_API UPCGInteractiveToolSettings_SplineBase : public UPCGInteractiveToolSettings
+UCLASS(Abstract)
+class UPCGInteractiveToolSettings_SplineBase : public UPCGInteractiveToolSettings
 {
 	GENERATED_BODY()
 
@@ -66,8 +66,8 @@ public:
 	FPCGToolRaycastSettings Settings{};
 };
 
-UCLASS(BlueprintType)
-class PCGEDITOR_API UPCGInteractiveToolSettings_Spline : public UPCGInteractiveToolSettings_SplineBase
+UCLASS()
+class UPCGInteractiveToolSettings_Spline : public UPCGInteractiveToolSettings_SplineBase
 {
 	GENERATED_BODY()
 
@@ -84,8 +84,8 @@ public:
 };
 
 
-UCLASS(BlueprintType)
-class PCGEDITOR_API UPCGInteractiveToolSettings_SplineSurface : public UPCGInteractiveToolSettings_SplineBase
+UCLASS()
+class UPCGInteractiveToolSettings_SplineSurface : public UPCGInteractiveToolSettings_SplineBase
 {
 	GENERATED_BODY()
 
@@ -98,7 +98,7 @@ public:
 };
 
 UCLASS(Abstract, Experimental)
-class PCGEDITOR_API UPCGDrawSplineToolBase : public UInteractiveTool, public IClickBehaviorTarget, public IClickDragBehaviorTarget
+class UPCGDrawSplineToolBase : public UInteractiveTool, public IClickBehaviorTarget, public IClickDragBehaviorTarget
 {
 	GENERATED_BODY()
 
@@ -145,7 +145,7 @@ protected:
 };
 
 UCLASS(Experimental)
-class PCGEDITOR_API UPCGDrawSplineTool : public UPCGDrawSplineToolBase
+class UPCGDrawSplineTool : public UPCGDrawSplineToolBase
 {
 	GENERATED_BODY()
 
@@ -157,7 +157,7 @@ public:
 };
 
 UCLASS(Experimental)
-class PCGEDITOR_API UPCGDrawSplineSurfaceTool : public UPCGDrawSplineToolBase
+class UPCGDrawSplineSurfaceTool : public UPCGDrawSplineToolBase
 {
 	GENERATED_BODY()
 
@@ -169,7 +169,7 @@ public:
 };
 
 UCLASS(Transient)
-class PCGEDITOR_API UPCGDrawSplineToolBuilder : public UPCGInteractiveToolBuilder
+class UPCGDrawSplineToolBuilder : public UPCGInteractiveToolBuilder
 {
 	GENERATED_BODY()
 
diff --git a/Engine/Plugins/PCG/Source/PCGEditor/Private/EditorMode/Tools/PCGInteractiveToolBuilder.h b/Engine/Plugins/PCG/Source/PCGEditor/Private/EditorMode/Tools/PCGInteractiveToolBuilder.h
--- a/Engine/Plugins/PCG/Source/PCGEditor/Private/EditorMode/Tools/PCGInteractiveToolBuilder.h
+++ b/Engine/Plugins/PCG/Source/PCGEditor/Private/EditorMode/Tools/PCGInteractiveToolBuilder.h
@@ -13,7 +13,7 @@ namespace UE::PCG::EditorMode::Tool
 }
 
 UCLASS()
-class PCGEDITOR_API UPCGInteractiveToolBuilder : public UInteractiveToolBuilder
+class UPCGInteractiveToolBuilder : public UInteractiveToolBuilder
 {
 	GENERATED_BODY()
 
@@ -31,7 +31,7 @@ private:
 };
 
 UCLASS()
-class PCGEDITOR_API UPCGInteractiveToolWithToolTargetsBuilder : public UInteractiveToolWithToolTargetsBuilder
+class UPCGInteractiveToolWithToolTargetsBuilder : public UInteractiveToolWithToolTargetsBuilder
 {
 	GENERATED_BODY()
 
diff --git a/Engine/Plugins/PCG/Source/PCGEditor/Private/EditorMode/Tools/PCGInteractiveToolCommon.h b/Engine/Plugins/PCG/Source/PCGEditor/Private/EditorMode/Tools/PCGInteractiveToolCommon.h
--- a/Engine/Plugins/PCG/Source/PCGEditor/Private/EditorMode/Tools/PCGInteractiveToolCommon.h
+++ b/Engine/Plugins/PCG/Source/PCGEditor/Private/EditorMode/Tools/PCGInteractiveToolCommon.h
@@ -31,7 +31,7 @@ enum class EPCGToolDrawTargetNormal : uint8
 };
 
 USTRUCT()
-struct PCGEDITOR_API FPCGToolRaycastSettings
+struct FPCGToolRaycastSettings
 {
 	GENERATED_BODY()
 
diff --git a/Engine/Plugins/PCG/Source/PCGEditor/Private/EditorMode/Tools/PCGInteractiveToolSettings.h b/Engine/Plugins/PCG/Source/PCGEditor/Private/EditorMode/Tools/PCGInteractiveToolSettings.h
--- a/Engine/Plugins/PCG/Source/PCGEditor/Private/EditorMode/Tools/PCGInteractiveToolSettings.h
+++ b/Engine/Plugins/PCG/Source/PCGEditor/Private/EditorMode/Tools/PCGInteractiveToolSettings.h
@@ -16,13 +16,13 @@
 
 namespace UE::PCG::EditorMode::Tool
 {
-	PCGEDITOR_API void Shutdown(UInteractiveTool* InTool, EToolShutdownType ShutdownType);
+	void Shutdown(UInteractiveTool* InTool, EToolShutdownType ShutdownType);
 	
-	PCGEDITOR_API extern TAutoConsoleVariable<bool> CVarCommitPropertyChangedEventOnTimer;
+	extern TAutoConsoleVariable<bool> CVarCommitPropertyChangedEventOnTimer;
 }
 
-UCLASS(Abstract, Blueprintable)
-class PCGEDITOR_API UPCGInteractiveToolSettings : public UInteractiveToolPropertySet, public FTickableEditorObject, public IKeyInputBehaviorTarget
+UCLASS(Abstract)
+class UPCGInteractiveToolSettings : public UInteractiveToolPropertySet, public FTickableEditorObject, public IKeyInputBehaviorTarget
 {
 	GENERATED_BODY()
 
diff --git a/Engine/Plugins/PCG/Source/PCGEditor/Private/EditorMode/Tools/Paint/PCGEdModeBrushBase.h b/Engine/Plugins/PCG/Source/PCGEditor/Private/EditorMode/Tools/Paint/PCGEdModeBrushBase.h
--- a/Engine/Plugins/PCG/Source/PCGEditor/Private/EditorMode/Tools/Paint/PCGEdModeBrushBase.h
+++ b/Engine/Plugins/PCG/Source/PCGEditor/Private/EditorMode/Tools/Paint/PCGEdModeBrushBase.h
@@ -7,7 +7,7 @@
 #include "PCGEdModeBrushBase.generated.h"
 
 UCLASS()
-class PCGEDITOR_API UPCGEdModeBrushBase : public UBrushBaseProperties
+class UPCGEdModeBrushBase : public UBrushBaseProperties
 {
 	GENERATED_BODY()
 
diff --git a/Engine/Plugins/PCG/Source/PCGEditor/Private/EditorMode/Tools/Paint/PCGPaintTool.h b/Engine/Plugins/PCG/Source/PCGEditor/Private/EditorMode/Tools/Paint/PCGPaintTool.h
--- a/Engine/Plugins/PCG/Source/PCGEditor/Private/EditorMode/Tools/Paint/PCGPaintTool.h
+++ b/Engine/Plugins/PCG/Source/PCGEditor/Private/EditorMode/Tools/Paint/PCGPaintTool.h
@@ -13,7 +13,7 @@
 class UMouseWheelInputBehavior;
 
 UCLASS()
-class PCGEDITOR_API UPCGInteractiveToolSettings_PaintTool : public UPCGInteractiveToolSettings
+class UPCGInteractiveToolSettings_PaintTool : public UPCGInteractiveToolSettings
 {
 	GENERATED_BODY()
 
diff --git a/Engine/Plugins/PCG/Source/PCGEditor/Private/EditorMode/Tools/Volume/PCGVolumeTool.h b/Engine/Plugins/PCG/Source/PCGEditor/Private/EditorMode/Tools/Volume/PCGVolumeTool.h
--- a/Engine/Plugins/PCG/Source/PCGEditor/Private/EditorMode/Tools/Volume/PCGVolumeTool.h
+++ b/Engine/Plugins/PCG/Source/PCGEditor/Private/EditorMode/Tools/Volume/PCGVolumeTool.h
@@ -14,7 +14,7 @@ class USingleClickOrDragInputBehavior;
 class UMouseHoverBehavior;
 
 UCLASS()
-class PCGEDITOR_API UPCGInteractiveToolSettings_Volume : public UPCGInteractiveToolSettings
+class UPCGInteractiveToolSettings_Volume : public UPCGInteractiveToolSettings
 {
 	GENERATED_BODY()
 
