[PCG] Fixed deadlock when deleting landscape w/ landscape cache

#rnx

diff --git a/Engine/Plugins/PCG/Source/PCG/Private/Grid/PCGLandscapeCache.cpp b/Engine/Plugins/PCG/Source/PCG/Private/Grid/PCGLandscapeCache.cpp
--- a/Engine/Plugins/PCG/Source/PCG/Private/Grid/PCGLandscapeCache.cpp
+++ b/Engine/Plugins/PCG/Source/PCG/Private/Grid/PCGLandscapeCache.cpp
@@ -1166,9 +1166,6 @@ void UPCGLandscapeCache::CacheLayerNames()
 void UPCGLandscapeCache::RemoveComponentFromCache(const ALandscapeProxy* LandscapeProxy)
 {
 	check(LandscapeProxy);
-
-	FWriteScopeLock Lock(CacheLock);
-
 	LandscapeProxy->ForEachComponent<ULandscapeComponent>(/*bIncludeFromChildActors=*/false, [this, LandscapeProxy](const ULandscapeComponent* LandscapeComponent)
 	{
 		if (LandscapeComponent)
