[PCG] Fix invalid compute graph warnings by always creating a valid resource for the raw buffer data interface.

#rnx

diff --git a/Engine/Plugins/PCG/Source/PCG/Private/Compute/DataInterfaces/PCGRawBufferDataInterface.cpp b/Engine/Plugins/PCG/Source/PCG/Private/Compute/DataInterfaces/PCGRawBufferDataInterface.cpp
--- a/Engine/Plugins/PCG/Source/PCG/Private/Compute/DataInterfaces/PCGRawBufferDataInterface.cpp
+++ b/Engine/Plugins/PCG/Source/PCG/Private/Compute/DataInterfaces/PCGRawBufferDataInterface.cpp
@@ -244,14 +244,16 @@ void UPCGRawBufferDataProvider::Reset()
 }
 
 FPCGRawBufferDataProviderProxy::FPCGRawBufferDataProviderProxy(const FParams& InParams)
-	: SizeBytes(InParams.SizeBytes)
-	, bZeroInitialize(InParams.bZeroInitialize)
+	: bZeroInitialize(InParams.bZeroInitialize)
 	, ExportMode(InParams.ExportMode)
 	, OutputPinLabel(InParams.OutputPinLabel)
 	, OutputPinLabelAlias(InParams.OutputPinLabelAlias)
 	, DataProviderWeakPtr(InParams.DataProviderWeakPtr)
 	, AsyncReadbackCallback_RenderThread(InParams.AsyncReadbackCallback_RenderThread)
 {
+	// Ensure we will create at least a minimal valid resource (4 bytes).
+	SizeBytes = FMath::Max(4, InParams.SizeBytes);
+
 	if (InParams.DataToUpload)
 	{
 		DataToUpload.Append(InParams.DataToUpload->GetConstData());
@@ -260,7 +262,7 @@ FPCGRawBufferDataProviderProxy::FPCGRawBufferDataProviderProxy(const FParams& In
 
 bool FPCGRawBufferDataProviderProxy::IsValid(FValidationData const& InValidationData) const
 {
-	if (SizeBytes <= 0 || (SizeBytes % sizeof(uint32)) != 0)
+	if ((SizeBytes % sizeof(uint32)) != 0)
 	{
 		UE_LOG(LogPCG, Warning, TEXT("FPCGRawBufferDataProviderProxy not valid due to invalid size: %d bytes."), SizeBytes);
 		return false;
@@ -271,8 +273,6 @@ bool FPCGRawBufferDataProviderProxy::IsValid(FValidationData const& InValidation
 
 void FPCGRawBufferDataProviderProxy::AllocateResources(FRDGBuilder& GraphBuilder, FAllocationData const& InAllocationData)
 {
-	check(SizeBytes > 0);
-
 	FRDGBufferDesc Desc = FRDGBufferDesc::CreateByteAddressDesc(SizeBytes);
 
 	if (ExportMode != EPCGExportMode::NoExport)
