[PCG] Fix crash when validating determinism on BiomeCore

#rnx

diff --git a/Engine/Plugins/PCG/Source/PCG/Private/Tests/Determinism/PCGDeterminismTestsCommon.cpp b/Engine/Plugins/PCG/Source/PCG/Private/Tests/Determinism/PCGDeterminismTestsCommon.cpp
--- a/Engine/Plugins/PCG/Source/PCG/Private/Tests/Determinism/PCGDeterminismTestsCommon.cpp
+++ b/Engine/Plugins/PCG/Source/PCG/Private/Tests/Determinism/PCGDeterminismTestsCommon.cpp
@@ -1551,8 +1551,7 @@ namespace PCGDeterminismTests
 	FPCGDataTypeIdentifier GetPermutation(int32 PermutationIteration, int32 PinIndex, const TArray<TArray<FPCGDataTypeIdentifier>>& BaseOptionsPerPin)
 	{
 		check(PermutationIteration >= 0);
-		check(PinIndex >= 0);
-		check(PinIndex < BaseOptionsPerPin[PinIndex].Num());
+		check(BaseOptionsPerPin.IsValidIndex(PinIndex));
 
 		int32 CardinalityProduct = 1;
 		for (int I = 0; I < PinIndex; ++I)
