[PCG] Fix name changing and connection issue when converting Change Grid Size node and undo/redo

diff --git a/Engine/Plugins/PCG/Source/PCGEditor/Private/Nodes/PCGEditorGraphNodeBase.cpp b/Engine/Plugins/PCG/Source/PCGEditor/Private/Nodes/PCGEditorGraphNodeBase.cpp
--- a/Engine/Plugins/PCG/Source/PCGEditor/Private/Nodes/PCGEditorGraphNodeBase.cpp
+++ b/Engine/Plugins/PCG/Source/PCGEditor/Private/Nodes/PCGEditorGraphNodeBase.cpp
@@ -1471,6 +1471,8 @@ void UPCGEditorGraphNodeBase::OnConvertNode(const FPCGPreconfiguredInfo& Convers
 		{
 			FText ConversionTransactionName = FText::Format(INVTEXT("{0}: {1}"), LOCTEXT("PCGEditorConvertNode", "Convert Node"), PCGNode->GetDefaultTitle());
 			FScopedTransaction Transaction(*FPCGEditorCommon::ContextIdentifier, std::move(ConversionTransactionName), PCGGraph);
+			PCGNode->Modify();
+			PCGNode->GetSettings()->Modify();
 
 			if (!PCGNode->GetSettings()->ConvertNode(ConversionInfo))
 			{
