- Added generation when tool graph changes.
This fixes an issue where, while the tool data is generated, the graph isn't actually run, which will result in tracking keys not being registered so tool usage wouldn't refresh the generation either

- Added actor selection when a new actor is spawned. This helps with going from one tool to the next quickly (previously you'd need to accept, select the actor, then choose a different tool. Now you can choose tool -> tool and it will automatically accept and change to the next tool)

diff --git a/Engine/Plugins/PCG/Source/PCGEditor/Private/EditorMode/Tools/PCGInteractiveToolSettings.cpp b/Engine/Plugins/PCG/Source/PCGEditor/Private/EditorMode/Tools/PCGInteractiveToolSettings.cpp
--- a/Engine/Plugins/PCG/Source/PCGEditor/Private/EditorMode/Tools/PCGInteractiveToolSettings.cpp
+++ b/Engine/Plugins/PCG/Source/PCGEditor/Private/EditorMode/Tools/PCGInteractiveToolSettings.cpp
@@ -83,9 +83,16 @@ TValueOrError<bool, FText> UPCGInteractiveToolSettings::Initialize(UInteractiveT
 	else
 	{
 		AActor* SpawnedActor = Actor::SpawnWorking(World, FTransform::Identity, GetWorkingActorClass());
-		SpawnedActor->SetActorLabel(NewActorName.ToString(), false);
 		// Our spawned actor is also the actor we are editing
 		EditActor = SpawnedActor;
+		
+		SpawnedActor->SetActorLabel(NewActorName.ToString(), false);
+
+		FSelectedObjectsChangeList ChangeList;
+		ChangeList.ModificationType = ESelectedObjectsModificationType::Replace;
+		ChangeList.Actors = {SpawnedActor};
+		
+		InTool->GetToolManager()->RequestSelectionChange(ChangeList);
 
 		GeneratedResources.SpawnedActor = SpawnedActor;
 	}
@@ -267,8 +274,8 @@ void UPCGInteractiveToolSettings::RegisterPropertyWatchers()
 	});
 }
 
-void UPCGInteractiveToolSettings::InitializeWorkingData(FName DataInstanceIdentifier)
-{
+bool UPCGInteractiveToolSettings::InitializeWorkingData(FName DataInstanceIdentifier)
+{	
 	UPCGComponent* PCGComponent = GetWorkingPCGComponent();
 	const UScriptStruct* Struct = GetWorkingDataType();
 	ensureMsgf(Struct->IsChildOf<FPCGInteractiveToolWorkingData>(), TEXT("Working Data Type must inherit from FPCGInteractiveToolWorkingData."));
@@ -301,7 +308,7 @@ void UPCGInteractiveToolSettings::InitializeWorkingData(FName DataInstanceIdenti
 	
 	if(!ensureMsgf(WorkingData, TEXT("Working Data has to be allocated at this point!.")))
 	{
-		return;
+		return false;
 	}
 	
 	// The working data needs to be initialized at this point
@@ -324,8 +331,10 @@ void UPCGInteractiveToolSettings::InitializeWorkingData(FName DataInstanceIdenti
 		}
 
 		// In which case, requesting new working data will recreate a new spline
-		InitializeWorkingData(DataInstanceIdentifier);
-	}	
+		return InitializeWorkingData(DataInstanceIdentifier);
+	}
+
+	return GeneratedResources.AddedWorkingDataIdentifiers.Contains(Context.WorkingDataIdentifier);
 }
 
 void UPCGInteractiveToolSettings::SetToolGraph(UPCGGraphInterface* InGraph)
@@ -416,7 +425,7 @@ TValueOrError<bool, FText> UPCGInteractiveToolSettings::TryLoadDefaults()
 	{
 		if (DefaultGraphInterface->ToolData.CompatibleToolTags.Contains(GetToolTag()) == false)
 		{
-			return MakeError(FText::FormatOrdered(LOCTEXT("DefaultGraph_IncompatibleToolTag", "Graph {Graph} does not have compatible tool tag {ToolTag}"), FText::FromString(GetPathNameSafe(DefaultGraphInterface)), FText::FromName(GetToolTag())));
+			return MakeError(FText::FormatOrdered(LOCTEXT("DefaultGraph_IncompatibleToolTag", "Graph {0} does not have compatible tool tag {1}"), FText::FromString(GetPathNameSafe(DefaultGraphInterface)), FText::FromName(GetToolTag())));
 		}
 		
 		// At this point, the delegates (OnPropertyModified) aren't setup, so we need to set things up manually.
@@ -427,7 +436,7 @@ TValueOrError<bool, FText> UPCGInteractiveToolSettings::TryLoadDefaults()
 		return MakeValue(true);
 	}
 
-	FText Message = LOCTEXT("DefaultGraph_NotFound", "No default graph found for tool settings {ToolSettings}. Please specify a graph with tool tag {ToolTag} under Editor Preferences -> PCG Editor Mode Settings");
+	FText Message = LOCTEXT("DefaultGraph_NotFound", "No default graph found for tool settings {0}. Please specify a graph with tool tag {1} under Editor Preferences -> PCG Editor Mode Settings");
 	return MakeError(FText::FormatOrdered(Message, GetClass()->GetDisplayNameText(), FText::FromName(GetToolTag())));
 }
 
@@ -621,9 +630,19 @@ void UPCGInteractiveToolSettings::OnPropertyModified(UObject* Object, FProperty*
 		// Update the working data to the available data instances ones for the current tool from the graph
 		TArray<FName> ValidDataInstanceNames = GetDataInstanceNamesForGraph();
 
+		bool bAnyDataCreated = false;
 		for (const FName& DataInstanceIdentifier : ValidDataInstanceNames)
 		{
-			InitializeWorkingData(DataInstanceIdentifier);
+			bAnyDataCreated |= InitializeWorkingData(DataInstanceIdentifier);
+		}
+
+		if(bAnyDataCreated)
+		{
+			// We generate again to make sure the initialized working data has a chance to register its tracking keys.
+			if (UPCGComponent* PCGComponent = GetWorkingPCGComponent())
+			{
+				PCGComponent->GenerateLocal(/*bForce=*/true);
+			}
 		}
 	}
 	if(Object == this && Property->GetFName() == GET_MEMBER_NAME_CHECKED(UPCGInteractiveToolSettings, ActorClassToSpawn))
diff --git a/Engine/Plugins/PCG/Source/PCGEditor/Private/EditorMode/Tools/PCGInteractiveToolSettings.h b/Engine/Plugins/PCG/Source/PCGEditor/Private/EditorMode/Tools/PCGInteractiveToolSettings.h
--- a/Engine/Plugins/PCG/Source/PCGEditor/Private/EditorMode/Tools/PCGInteractiveToolSettings.h
+++ b/Engine/Plugins/PCG/Source/PCGEditor/Private/EditorMode/Tools/PCGInteractiveToolSettings.h
@@ -49,8 +49,10 @@ public:
 	virtual void RegisterPropertyWatchers();
 	
 	/** Initializes or finds the working data associated with the data instance identifier.
-	 * If NAME_None, the ToolTag is used as identifier. */
-	void InitializeWorkingData(FName DataInstanceIdentifier = NAME_None);
+	 * If NAME_None, the ToolTag is used as identifier.
+	 * @returns Whether new working data was created and initialized.
+	 * False does not imply working data is invalid; just that it didn't get created. */
+	bool InitializeWorkingData(FName DataInstanceIdentifier = NAME_None);
 	
 	UPCGGraphInterface* GetToolGraph() const { return ToolGraph; }
 	void SetToolGraph(UPCGGraphInterface* InGraph);
