[PCG] Followup on Landscape cache deadlock removal

* Indicate clearly that the function is unsafe

#rnx

diff --git a/Engine/Plugins/PCG/Source/PCG/Private/Grid/PCGLandscapeCache.cpp b/Engine/Plugins/PCG/Source/PCG/Private/Grid/PCGLandscapeCache.cpp
--- a/Engine/Plugins/PCG/Source/PCG/Private/Grid/PCGLandscapeCache.cpp
+++ b/Engine/Plugins/PCG/Source/PCG/Private/Grid/PCGLandscapeCache.cpp
@@ -1111,7 +1111,7 @@ void UPCGLandscapeCache::OnLandscapeMoved(AActor* InActor)
 	CacheLock.WriteLock();
 
 	// Just remove these from the cache, they'll be added back on demand
-	RemoveComponentFromCache(Landscape);
+	RemoveComponentFromCache_Unsafe(Landscape);
 
 	CacheLayerNames(Landscape);
 
@@ -1128,7 +1128,7 @@ void UPCGLandscapeCache::OnLandscapeDeleted(AActor* Actor)
 
 	CacheLock.WriteLock();
 
-	RemoveComponentFromCache(LandscapeProxy);
+	RemoveComponentFromCache_Unsafe(LandscapeProxy);
 
 	Landscapes.Remove(LandscapeProxy);
 
@@ -1163,7 +1163,7 @@ void UPCGLandscapeCache::CacheLayerNames()
 	}
 }
 
-void UPCGLandscapeCache::RemoveComponentFromCache(const ALandscapeProxy* LandscapeProxy)
+void UPCGLandscapeCache::RemoveComponentFromCache_Unsafe(const ALandscapeProxy* LandscapeProxy)
 {
 	check(LandscapeProxy);
 	LandscapeProxy->ForEachComponent<ULandscapeComponent>(/*bIncludeFromChildActors=*/false, [this, LandscapeProxy](const ULandscapeComponent* LandscapeComponent)
diff --git a/Engine/Plugins/PCG/Source/PCG/Public/Grid/PCGLandscapeCache.h b/Engine/Plugins/PCG/Source/PCG/Public/Grid/PCGLandscapeCache.h
--- a/Engine/Plugins/PCG/Source/PCG/Public/Grid/PCGLandscapeCache.h
+++ b/Engine/Plugins/PCG/Source/PCG/Public/Grid/PCGLandscapeCache.h
@@ -177,7 +177,7 @@ private:
 	void OnLandscapeLoaded(AActor& Actor);
 	void CacheLayerNames(ALandscapeProxy* InLandscape);
 	void CacheLayerNames();
-	void RemoveComponentFromCache(const ALandscapeProxy* LandscapeProxy);
+	void RemoveComponentFromCache_Unsafe(const ALandscapeProxy* LandscapeProxy);
 
 	/** Gets (and creates if needed) the cache entry - available only in Editor */
 	TSharedPtr<const FPCGLandscapeCacheEntry> GetOrCreateCacheEntryInternal(const ULandscapeInfo* InLandscapeInfo, const FIntPoint& InComponentKey, bool bInAllowLoadOrCreate, bool& bOutValidEntry);
