Added Up & Forward arrows to PCG Gizmo

diff --git a/Engine/Plugins/PCG/Source/PCGEditor/Private/EditorMode/Tools/Gizmos/PCGGizmos.cpp b/Engine/Plugins/PCG/Source/PCGEditor/Private/EditorMode/Tools/Gizmos/PCGGizmos.cpp
--- a/Engine/Plugins/PCG/Source/PCGEditor/Private/EditorMode/Tools/Gizmos/PCGGizmos.cpp
+++ b/Engine/Plugins/PCG/Source/PCGEditor/Private/EditorMode/Tools/Gizmos/PCGGizmos.cpp
@@ -48,39 +48,54 @@ void UPCGVolumetricBrushStampIndicator::Render(IToolsContextRenderAPI* RenderAPI
 	
 	if (BrushMode == EPCGBrushMode::Sphere)
 	{
-		Super::Render(RenderAPI);		
+		Super::Render(RenderAPI);
 	}
 	else if (BrushMode == EPCGBrushMode::Box)
 	{
 		if (bDrawIndicatorLines)
 		{
 			FToolDataVisualizer Draw;
+			Draw.SetTransform(GetPreviewMesh()->GetTransform());
 			Draw.BeginFrame(RenderAPI);
 
-			FRotator ActorRotation = GetPreviewMesh()->GetActor()->GetActorRotation();
 			if (bDrawSecondaryLines)
 			{
 				const bool bCirclesCoincident = bDrawRadiusCircle && FMath::IsNearlyEqual(BrushFalloff, 1.0f);
 				if (!bCirclesCoincident)
 				{
-					float SideLength = BrushRadius * 2 * BrushFalloff;
-					FVector SideA = ActorRotation.RotateVector(FVector(SideLength, 0, 0));
-					FVector SideB = ActorRotation.RotateVector(FVector(0, SideLength, 0));
-					Draw.DrawSquare(BrushPosition, SideA, SideB, LineColor, LineThickness, bDepthTested);
+					// SideA and SideB are 2 units long due to extents being twice the brush radius, which is the scale in the transform.
+					// Then scaled down by the BrushFalloff
+					Draw.DrawSquare(FVector::ZeroVector, FVector(2 * BrushFalloff, 0, 0), FVector(0, 2 * BrushFalloff, 0), LineColor, LineThickness, bDepthTested);
 				}
 			}
 		
 			if (bDrawExtents)
 			{
-				float SideLength = BrushRadius * 2;
-				FVector SideA = ActorRotation.RotateVector(FVector(SideLength, 0, 0));
-				FVector SideB = ActorRotation.RotateVector(FVector(0, SideLength, 0));
-				Draw.DrawSquare(BrushPosition, SideA, SideB, LineColor, LineThickness, bDepthTested);
+				Draw.DrawSquare(FVector::ZeroVector, FVector(2, 0, 0), FVector(0, 2, 0), LineColor, LineThickness, bDepthTested);
 			}
 		
 			Draw.EndFrame();
 		}
-	}	
+	}
+
+	if (bDrawDirections)
+	{
+		FToolDataVisualizer Draw;
+		
+		FTransform Transform = GetPreviewMesh()->GetTransform();
+		Draw.SetTransform(Transform);
+		
+		Draw.BeginFrame(RenderAPI);
+
+		// Unit vectors are fine because scale takes care of appropriate sizing.
+		FVector DirectionUp = FVector(0, 0, 1);
+		FVector DirectionForward = FVector(1, 0, 0);
+		
+		Draw.DrawDirectionalArrow(FVector::ZeroVector, DirectionUp, DirectionUp, FLinearColor::Blue, 10.f, 2.f);
+		Draw.DrawDirectionalArrow(FVector::ZeroVector, DirectionForward, DirectionForward, FLinearColor::Red, 10.f, 2.f);
+
+		Draw.EndFrame();
+	}
 }
 
 void UPCGVolumetricBrushStampIndicator::Shutdown()
diff --git a/Engine/Plugins/PCG/Source/PCGEditor/Private/EditorMode/Tools/Gizmos/PCGGizmos.h b/Engine/Plugins/PCG/Source/PCGEditor/Private/EditorMode/Tools/Gizmos/PCGGizmos.h
--- a/Engine/Plugins/PCG/Source/PCGEditor/Private/EditorMode/Tools/Gizmos/PCGGizmos.h
+++ b/Engine/Plugins/PCG/Source/PCGEditor/Private/EditorMode/Tools/Gizmos/PCGGizmos.h
@@ -65,6 +65,9 @@ private:
 	UPROPERTY()
 	bool bDrawExtents = true;
 
+	UPROPERTY()
+	bool bDrawDirections = true;
+	
 	UPROPERTY()
 	EPCGBrushMode BrushMode = EPCGBrushMode::Sphere;
 };
\ No newline at end of file
