[PCG] Move GPU from Experimental to Beta

diff --git a/Engine/Plugins/PCG/Source/PCG/Internal/Compute/IPCGNodeSourceTextProvider.h b/Engine/Plugins/PCG/Source/PCG/Internal/Compute/IPCGNodeSourceTextProvider.h
--- a/Engine/Plugins/PCG/Source/PCG/Internal/Compute/IPCGNodeSourceTextProvider.h
+++ b/Engine/Plugins/PCG/Source/PCG/Internal/Compute/IPCGNodeSourceTextProvider.h
@@ -12,8 +12,7 @@ class UPCGNodeSourceTextProvider : public UInterface
 };
 
 /**
-* [EXPERIMENTAL] Interface for node source text editing widgets.
-* Note: This struct is subject to change without deprecation.
+* Interface for node source text editing widgets.
 */
 class IPCGNodeSourceTextProvider
 {
diff --git a/Engine/Plugins/PCG/Source/PCG/Private/Components/PCGProceduralISMComponent.h b/Engine/Plugins/PCG/Source/PCG/Private/Components/PCGProceduralISMComponent.h
--- a/Engine/Plugins/PCG/Source/PCG/Private/Components/PCGProceduralISMComponent.h
+++ b/Engine/Plugins/PCG/Source/PCG/Private/Components/PCGProceduralISMComponent.h
@@ -28,7 +28,7 @@ struct FPCGContext;
 struct FPrimitiveSceneDesc;
 struct FPrimitiveSceneProxyDesc;
 
-/** Experimental instanced primitive that is not backed by CPU instance data. Instance data must be written direct to the GPU scene. */
+/** Instanced primitive that is not backed by CPU instance data. Instance data must be written direct to the GPU scene. */
 UCLASS(ClassGroup = Rendering, meta = (BlueprintSpawnableComponent), Blueprintable, MinimalAPI, hideCategories = (Collision, Cooking, HLOD, MeshPainting, Navigation, Physics, RayTracing))
 class UPCGProceduralISMComponent : public UStaticMeshComponent
 {
diff --git a/Engine/Plugins/PCG/Source/PCG/Private/Compute/BuiltInKernels/PCGCountUniqueAttributeValuesKernel.cpp b/Engine/Plugins/PCG/Source/PCG/Private/Compute/BuiltInKernels/PCGCountUniqueAttributeValuesKernel.cpp
--- a/Engine/Plugins/PCG/Source/PCG/Private/Compute/BuiltInKernels/PCGCountUniqueAttributeValuesKernel.cpp
+++ b/Engine/Plugins/PCG/Source/PCG/Private/Compute/BuiltInKernels/PCGCountUniqueAttributeValuesKernel.cpp
@@ -45,7 +45,7 @@ bool UPCGCountUniqueAttributeValuesKernel::IsKernelDataValid(FPCGContext* InCont
 			return false;
 		}
 
-		if (!InputDataDesc->DataDescs.IsEmpty())
+		if (!InputDataDesc->GetDataDescriptions().IsEmpty())
 		{
 			FPCGKernelAttributeDesc AttributeDesc;
 			bool bConflictingTypesFound = false;
@@ -54,7 +54,7 @@ bool UPCGCountUniqueAttributeValuesKernel::IsKernelDataValid(FPCGContext* InCont
 
 			if (!bPresentOnAllData)
 			{
-				if (!InputDataDesc->DataDescs.IsEmpty())
+				if (!InputDataDesc->GetDataDescriptions().IsEmpty())
 				{
 					PCG_KERNEL_VALIDATION_ERR(InContext, GetSettings(), FText::Format(
 						LOCTEXT("AttributeMissing", "Count attribute '{0}' not found, this attribute must be present on all input data, and be of type String Key."),
@@ -73,7 +73,7 @@ bool UPCGCountUniqueAttributeValuesKernel::IsKernelDataValid(FPCGContext* InCont
 				return false;
 			}
 
-			if (AttributeDesc.AttributeKey.Type != EPCGKernelAttributeType::StringKey)
+			if (AttributeDesc.GetAttributeKey().GetType() != EPCGKernelAttributeType::StringKey)
 			{
 				// Attribute value counting only currently supported for attributes of type StringKey.
 				PCG_KERNEL_VALIDATION_ERR(Context, GetSettings(), FText::Format(
@@ -107,6 +107,7 @@ TSharedPtr<const FPCGDataCollectionDesc> UPCGCountUniqueAttributeValuesKernel::C
 	}
 
 	TSharedPtr<FPCGDataCollectionDesc> OutDataDesc = FPCGDataCollectionDesc::MakeShared();
+	TArray<FPCGDataDesc>& OutDataDescs = OutDataDesc->GetDataDescriptionsMutable();
 
 	const int32 AttributeId = InBinding->GetAttributeId(AttributeName, EPCGKernelAttributeType::StringKey);
 
@@ -139,25 +140,27 @@ TSharedPtr<const FPCGDataCollectionDesc> UPCGCountUniqueAttributeValuesKernel::C
 
 		if (bEmitPerDataCounts)
 		{
+			const TConstArrayView<FPCGDataDesc> InputDataDescs = InputDesc->GetDataDescriptions();
+
 			// Output data is array of (value, count) pairs for each data.
-			for (int InputIndex = 0; InputIndex < InputDesc->DataDescs.Num(); ++InputIndex)
+			for (int InputIndex = 0; InputIndex < InputDataDescs.Num(); ++InputIndex)
 			{
-				const FPCGDataDesc& DataDesc = InputDesc->DataDescs[InputIndex];
+				const FPCGDataDesc& DataDesc = InputDataDescs[InputIndex];
 
 				// Data descriptions have multiple collection objects so worth to move last element for efficiency, otherwise copy.
-				if (InputIndex < InputDesc->DataDescs.Num() - 1)
+				if (InputIndex < InputDataDescs.Num() - 1)
 				{
-					OutDataDesc->DataDescs.Add(CountDataDesc);
+					OutDataDescs.Add(CountDataDesc);
 				}
 				else
 				{
-					OutDataDesc->DataDescs.Add(MoveTemp(CountDataDesc));
+					OutDataDescs.Add(MoveTemp(CountDataDesc));
 				}
 			}
 		}
 		else
 		{
-			OutDataDesc->DataDescs.Add(MoveTemp(CountDataDesc));
+			OutDataDescs.Add(MoveTemp(CountDataDesc));
 		}
 	}
 	else
@@ -166,12 +169,12 @@ TSharedPtr<const FPCGDataCollectionDesc> UPCGCountUniqueAttributeValuesKernel::C
 		if (bOutputRawBuffer)
 		{
 			// Null output: One pair (string key 0 (invalid), value count 0).
-			OutDataDesc->DataDescs.Emplace(FPCGDataTypeInfoRawBuffer::AsId(), 2);
+			OutDataDescs.Emplace(FPCGDataTypeInfoRawBuffer::AsId(), 2);
 		}
 		else
 		{
 			// Null output: One entry: string key 0, value count 0.
-			FPCGDataDesc& CountDataDesc = OutDataDesc->DataDescs.Emplace_GetRef(EPCGDataType::Param, 1);
+			FPCGDataDesc& CountDataDesc = OutDataDescs.Emplace_GetRef(EPCGDataType::Param, 1);
 			CountDataDesc.AddAttribute(FPCGKernelAttributeKey(PCGCountUniqueAttributeValuesConstants::ValueAttributeName, EPCGKernelAttributeType::Int), InBinding);
 			CountDataDesc.AddAttribute(FPCGKernelAttributeKey(PCGCountUniqueAttributeValuesConstants::ValueCountAttributeName, EPCGKernelAttributeType::Int), InBinding);
 		}
diff --git a/Engine/Plugins/PCG/Source/PCG/Private/Compute/BuiltInKernels/PCGDataNumKernel.cpp b/Engine/Plugins/PCG/Source/PCG/Private/Compute/BuiltInKernels/PCGDataNumKernel.cpp
--- a/Engine/Plugins/PCG/Source/PCG/Private/Compute/BuiltInKernels/PCGDataNumKernel.cpp
+++ b/Engine/Plugins/PCG/Source/PCG/Private/Compute/BuiltInKernels/PCGDataNumKernel.cpp
@@ -26,7 +26,7 @@ TSharedPtr<const FPCGDataCollectionDesc> UPCGDataNumKernel::ComputeOutputBinding
 
 	TSharedPtr<FPCGDataCollectionDesc> OutDataDesc = FPCGDataCollectionDesc::MakeShared();
 
-	FPCGDataDesc& OutputDataDesc = OutDataDesc->DataDescs.Emplace_GetRef(EPCGDataType::Param, 1);
+	FPCGDataDesc& OutputDataDesc = OutDataDesc->GetDataDescriptionsMutable().Emplace_GetRef(EPCGDataType::Param, 1);
 	OutputDataDesc.AddAttribute(OutCountAttributeKey, InBinding);
 
 	return OutDataDesc;
diff --git a/Engine/Plugins/PCG/Source/PCG/Private/Compute/Data/PCGProxyForGPUData.cpp b/Engine/Plugins/PCG/Source/PCG/Private/Compute/Data/PCGProxyForGPUData.cpp
--- a/Engine/Plugins/PCG/Source/PCG/Private/Compute/Data/PCGProxyForGPUData.cpp
+++ b/Engine/Plugins/PCG/Source/PCG/Private/Compute/Data/PCGProxyForGPUData.cpp
@@ -56,7 +56,7 @@ FPCGDataTypeBaseId UPCGProxyForGPUData::GetUnderlyingDataTypeId() const
 	{
 		if (const TSharedPtr<const FPCGDataCollectionDesc> DataCollectionDesc = DataCollectionGPU->GetDescription())
 		{
-			return DataCollectionDesc->DataDescs[DataIndexInCollection].GetType().GetId();
+			return DataCollectionDesc->GetDataDescriptions()[DataIndexInCollection].GetType().GetId();
 		}
 	}
 
@@ -79,7 +79,7 @@ FIntVector4 UPCGProxyForGPUData::GetElementCount() const
 	{
 		if (const TSharedPtr<const FPCGDataCollectionDesc> DataCollectionDesc = DataCollectionGPU->GetDescription())
 		{
-			return DataCollectionDesc->DataDescs[DataIndexInCollection].GetElementCount();
+			return DataCollectionDesc->GetDataDescriptions()[DataIndexInCollection].GetElementCount();
 		}
 	}
 
@@ -92,7 +92,7 @@ bool UPCGProxyForGPUData::GetDescription(FPCGDataDesc& OutDescription) const
 	{
 		if (const TSharedPtr<const FPCGDataCollectionDesc> DataCollectionDesc = DataCollectionGPU->GetDescription())
 		{
-			OutDescription = DataCollectionDesc->DataDescs[DataIndexInCollection];
+			OutDescription = DataCollectionDesc->GetDataDescriptions()[DataIndexInCollection];
 			return true;
 		}
 	}
@@ -131,9 +131,9 @@ TSharedPtr<FPCGProxyForGPUDataCollection> UPCGProxyForGPUData::GetGPUInfoMutable
 		ensureMsgf(false, TEXT("Data %s: GPU data collection lost. Enabling the define 'PCG_DATA_USAGE_LOGGING' may help to identify when resource was released."), *GetName());
 		return nullptr;
 	}
-	else if (!DataCollectionOnGPU->GetDescription()->DataDescs.IsValidIndex(DataIndexInCollection))
+	else if (!DataCollectionOnGPU->GetDescription()->GetDataDescriptions().IsValidIndex(DataIndexInCollection))
 	{
-		ensureMsgf(false, TEXT("Data %s: DataIndexInCollection (%d) was out of range [0, %d)."), *GetName(), DataIndexInCollection, DataCollectionOnGPU->GetDescription()->DataDescs.Num());
+		ensureMsgf(false, TEXT("Data %s: DataIndexInCollection (%d) was out of range [0, %d)."), *GetName(), DataIndexInCollection, DataCollectionOnGPU->GetDescription()->GetDataDescriptions().Num());
 		return nullptr;
 	}
 
diff --git a/Engine/Plugins/PCG/Source/PCG/Private/Compute/DataInterfaces/BuiltInKernels/PCGCountUniqueAttributeValuesDataInterface.cpp b/Engine/Plugins/PCG/Source/PCG/Private/Compute/DataInterfaces/BuiltInKernels/PCGCountUniqueAttributeValuesDataInterface.cpp
--- a/Engine/Plugins/PCG/Source/PCG/Private/Compute/DataInterfaces/BuiltInKernels/PCGCountUniqueAttributeValuesDataInterface.cpp
+++ b/Engine/Plugins/PCG/Source/PCG/Private/Compute/DataInterfaces/BuiltInKernels/PCGCountUniqueAttributeValuesDataInterface.cpp
@@ -138,15 +138,13 @@ bool UPCGCountUniqueAttributeValuesDataProvider::PrepareForExecute_GameThread(UP
 		if (bEmitPerDataCounts)
 		{
 			// Gather unique string key values for each data.
-			for (int InputIndex = 0; InputIndex < InputDataDesc->DataDescs.Num(); ++InputIndex)
+			for (const FPCGDataDesc& DataDesc : InputDataDesc->GetDataDescriptions())
 			{
-				const FPCGDataDesc& DataDesc = InputDataDesc->DataDescs[InputIndex];
-
 				TArray<int32> UniqueStringKeyValues;
 
-				for (const FPCGKernelAttributeDesc& AttributeDesc : DataDesc.AttributeDescs)
+				for (const FPCGKernelAttributeDesc& AttributeDesc : DataDesc.GetAttributeDescriptions())
 				{
-					if (AttributeDesc.AttributeId == AttributeToCountId)
+					if (AttributeDesc.GetAttributeId() == AttributeToCountId)
 					{
 						UniqueStringKeyValues = AttributeDesc.GetUniqueStringKeys();
 						break;
diff --git a/Engine/Plugins/PCG/Source/PCG/Private/Compute/DataInterfaces/Elements/PCGCopyPointsDataInterface.cpp b/Engine/Plugins/PCG/Source/PCG/Private/Compute/DataInterfaces/Elements/PCGCopyPointsDataInterface.cpp
--- a/Engine/Plugins/PCG/Source/PCG/Private/Compute/DataInterfaces/Elements/PCGCopyPointsDataInterface.cpp
+++ b/Engine/Plugins/PCG/Source/PCG/Private/Compute/DataInterfaces/Elements/PCGCopyPointsDataInterface.cpp
@@ -182,9 +182,12 @@ bool UPCGCopyPointsDataProvider::PrepareForExecute_GameThread(UPCGDataBinding* I
 			return true;
 		}
 
+		const TConstArrayView<FPCGDataDesc> SourceDataDescs = SourcePinDesc->GetDataDescriptions();
+		const TConstArrayView<FPCGDataDesc> TargetDataDescs = TargetPinDesc->GetDataDescriptions();
+
 		const bool bCopyEachSourceOnEveryTarget = KernelParams->GetValueBool(GET_MEMBER_NAME_CHECKED(UPCGCopyPointsSettings, bCopyEachSourceOnEveryTarget));
 
-		const int32 NumOutputs = bCopyEachSourceOnEveryTarget ? SourcePinDesc->DataDescs.Num() * TargetPinDesc->DataDescs.Num() : FMath::Max(SourcePinDesc->DataDescs.Num(), TargetPinDesc->DataDescs.Num());
+		const int32 NumOutputs = bCopyEachSourceOnEveryTarget ? SourceDataDescs.Num() * TargetDataDescs.Num() : FMath::Max(SourceDataDescs.Num(), TargetDataDescs.Num());
 		SourceAndTargetDataIndices.Reserve(NumOutputs);
 
 		const UPCGParamData* AnalysisResultsData = Cast<UPCGParamData>(InBinding->GetInputDataCollection().TaggedData[AnalysisDataIndex].Data);
@@ -197,11 +200,11 @@ bool UPCGCopyPointsDataProvider::PrepareForExecute_GameThread(UPCGDataBinding* I
 		{
 			if (bCopyEachSourceOnEveryTarget)
 			{
-				for (int32 S = 0; S < SourcePinDesc->DataDescs.Num(); ++S)
+				for (int32 S = 0; S < SourceDataDescs.Num(); ++S)
 				{
-					for (int32 T = 0; T < TargetPinDesc->DataDescs.Num(); ++T)
+					for (int32 T = 0; T < TargetDataDescs.Num(); ++T)
 					{
-						if (SelectedFlagsAttribute->GetValue(S * TargetPinDesc->DataDescs.Num() + T))
+						if (SelectedFlagsAttribute->GetValue(S * TargetDataDescs.Num() + T))
 						{
 							SourceAndTargetDataIndices.Emplace(S, T);
 						}
@@ -212,8 +215,8 @@ bool UPCGCopyPointsDataProvider::PrepareForExecute_GameThread(UPCGDataBinding* I
 			{
 				for (int32 OutputIndex = 0; OutputIndex < NumOutputs; ++OutputIndex)
 				{
-					const int32 S = FMath::Clamp(OutputIndex, 0, SourcePinDesc->DataDescs.Num() - 1);
-					const int32 T = FMath::Clamp(OutputIndex, 0, TargetPinDesc->DataDescs.Num() - 1);
+					const int32 S = FMath::Clamp(OutputIndex, 0, SourceDataDescs.Num() - 1);
+					const int32 T = FMath::Clamp(OutputIndex, 0, TargetDataDescs.Num() - 1);
 
 					if (SelectedFlagsAttribute->GetValue(OutputIndex))
 					{
diff --git a/Engine/Plugins/PCG/Source/PCG/Private/Compute/DataInterfaces/Elements/PCGMetadataPartitionDataInterface.cpp b/Engine/Plugins/PCG/Source/PCG/Private/Compute/DataInterfaces/Elements/PCGMetadataPartitionDataInterface.cpp
--- a/Engine/Plugins/PCG/Source/PCG/Private/Compute/DataInterfaces/Elements/PCGMetadataPartitionDataInterface.cpp
+++ b/Engine/Plugins/PCG/Source/PCG/Private/Compute/DataInterfaces/Elements/PCGMetadataPartitionDataInterface.cpp
@@ -145,10 +145,10 @@ bool UPCGMetaDataPartitionDataProvider::PrepareForExecute_GameThread(UPCGDataBin
 		return true;
 	}
 
-	if (AnalysisDataIndices.Num() != InputDesc->DataDescs.Num())
+	if (AnalysisDataIndices.Num() != InputDesc->GetDataDescriptions().Num())
 	{
 		UE_LOG(LogPCG, Error, TEXT("UPCGMetaDataPartitionDataProvider: Number of received analysis data items (%d) does not match input data count (%d)."),
-			AnalysisDataIndices.Num(), InputDesc->DataDescs.Num());
+			AnalysisDataIndices.Num(), InputDesc->GetDataDescriptions().Num());
 		return true;
 	}
 
@@ -219,7 +219,7 @@ bool UPCGMetaDataPartitionDataProvider::PrepareForExecute_GameThread(UPCGDataBin
 
 		PartitionAttributeId = InBinding->GetAttributeId(Settings->PartitionAttributeSelectors[0].GetAttributeName(), EPCGKernelAttributeType::StringKey);
 
-		NumInputData = InputDataDesc->DataDescs.Num();
+		NumInputData = InputDataDesc->GetDataDescriptions().Num();
 	}
 
 	return true;
diff --git a/Engine/Plugins/PCG/Source/PCG/Private/Compute/DataInterfaces/Elements/PCGSkinnedMeshSpawnerDataInterface.cpp b/Engine/Plugins/PCG/Source/PCG/Private/Compute/DataInterfaces/Elements/PCGSkinnedMeshSpawnerDataInterface.cpp
--- a/Engine/Plugins/PCG/Source/PCG/Private/Compute/DataInterfaces/Elements/PCGSkinnedMeshSpawnerDataInterface.cpp
+++ b/Engine/Plugins/PCG/Source/PCG/Private/Compute/DataInterfaces/Elements/PCGSkinnedMeshSpawnerDataInterface.cpp
@@ -320,7 +320,7 @@ bool UPCGSkinnedMeshSpawnerDataProvider::PrepareForExecute_GameThread(UPCGDataBi
 
 		bool bAnyPointsPresent = false;
 
-		for (const FPCGDataDesc& Desc : InputDataDesc->DataDescs)
+		for (const FPCGDataDesc& Desc : InputDataDesc->GetDataDescriptions())
 		{
 			if (Desc.GetElementCount().X <= 0)
 			{
@@ -329,11 +329,11 @@ bool UPCGSkinnedMeshSpawnerDataProvider::PrepareForExecute_GameThread(UPCGDataBi
 
 			bAnyPointsPresent = true;
 
-			for (const FPCGKernelAttributeDesc& AttributeDesc : Desc.AttributeDescs)
+			for (const FPCGKernelAttributeDesc& AttributeDesc : Desc.GetAttributeDescriptions())
 			{
-				if (AttributeDesc.AttributeKey.Identifier.Name == SelectorName && AttributeDesc.AttributeKey.Type == EPCGKernelAttributeType::StringKey)
+				if (AttributeDesc.GetAttributeKey().GetIdentifier().Name == SelectorName && AttributeDesc.GetAttributeKey().GetType() == EPCGKernelAttributeType::StringKey)
 				{
-					SelectorAttributeId = AttributeDesc.AttributeId;
+					SelectorAttributeId = AttributeDesc.GetAttributeId();
 					break;
 				}
 			}
@@ -531,18 +531,18 @@ void UPCGSkinnedMeshSpawnerDataProvider::CreatePrimitiveDescriptors(FPCGContext*
 		// Compute how many instances we expect for each mesh, if we do not already have the answer from analysis.
 		if (StringKeyToInstanceCount.IsEmpty())
 		{
-			for (const FPCGDataDesc& Desc : InputDataDesc->DataDescs)
+			for (const FPCGDataDesc& Desc : InputDataDesc->GetDataDescriptions())
 			{
 				if (Desc.GetElementCount().X <= 0)
 				{
 					continue;
 				}
 
-				for (const FPCGKernelAttributeDesc& AttributeDesc : Desc.AttributeDescs)
+				for (const FPCGKernelAttributeDesc& AttributeDesc : Desc.GetAttributeDescriptions())
 				{
-					if (AttributeDesc.AttributeKey.Identifier.Name == SelectorName)
+					if (AttributeDesc.GetAttributeKey().GetIdentifier().Name == SelectorName)
 					{
-						if (AttributeDesc.AttributeKey.Type == EPCGKernelAttributeType::StringKey)
+						if (AttributeDesc.GetAttributeKey().GetType() == EPCGKernelAttributeType::StringKey)
 						{
 							for (int32 AttributeStringKeyValue : AttributeDesc.GetUniqueStringKeys())
 							{
diff --git a/Engine/Plugins/PCG/Source/PCG/Private/Compute/DataInterfaces/Elements/PCGStaticMeshSpawnerDataInterface.cpp b/Engine/Plugins/PCG/Source/PCG/Private/Compute/DataInterfaces/Elements/PCGStaticMeshSpawnerDataInterface.cpp
--- a/Engine/Plugins/PCG/Source/PCG/Private/Compute/DataInterfaces/Elements/PCGStaticMeshSpawnerDataInterface.cpp
+++ b/Engine/Plugins/PCG/Source/PCG/Private/Compute/DataInterfaces/Elements/PCGStaticMeshSpawnerDataInterface.cpp
@@ -345,7 +345,7 @@ bool UPCGStaticMeshSpawnerDataProvider::PrepareForExecute_GameThread(UPCGDataBin
 
 		bool bAnyPointsPresent = false;
 
-		for (const FPCGDataDesc& Desc : InputDataDesc->DataDescs)
+		for (const FPCGDataDesc& Desc : InputDataDesc->GetDataDescriptions())
 		{
 			if (Desc.GetElementCount().X <= 0)
 			{
@@ -354,11 +354,11 @@ bool UPCGStaticMeshSpawnerDataProvider::PrepareForExecute_GameThread(UPCGDataBin
 
 			bAnyPointsPresent = true;
 
-			for (const FPCGKernelAttributeDesc& AttributeDesc : Desc.AttributeDescs)
+			for (const FPCGKernelAttributeDesc& AttributeDesc : Desc.GetAttributeDescriptions())
 			{
-				if (AttributeDesc.AttributeKey.Identifier.Name == SelectorName && AttributeDesc.AttributeKey.Type == EPCGKernelAttributeType::StringKey)
+				if (AttributeDesc.GetAttributeKey().GetIdentifier().Name == SelectorName && AttributeDesc.GetAttributeKey().GetType() == EPCGKernelAttributeType::StringKey)
 				{
-					SelectorAttributeId = AttributeDesc.AttributeId;
+					SelectorAttributeId = AttributeDesc.GetAttributeId();
 					break;
 				}
 			}
@@ -372,7 +372,7 @@ bool UPCGStaticMeshSpawnerDataProvider::PrepareForExecute_GameThread(UPCGDataBin
 		if (SelectorAttributeId == INDEX_NONE)
 		{
 			// Mute this error if the point data is empty.
-			if (!InputDataDesc->DataDescs.IsEmpty() && bAnyPointsPresent)
+			if (!InputDataDesc->GetDataDescriptions().IsEmpty() && bAnyPointsPresent)
 			{
 				PCG_KERNEL_VALIDATION_ERR(Context, Settings, FText::Format(
 					LOCTEXT("MeshSelectorAttributeNotFound", "Mesh selector attribute '{0}' not found."),
@@ -577,18 +577,18 @@ void UPCGStaticMeshSpawnerDataProvider::CreatePrimitiveDescriptors(FPCGContext*
 		// Compute how many instances we expect for each mesh, if we do not already have the answer from analysis.
 		if (StringKeyToInstanceCount.IsEmpty())
 		{
-			for (const FPCGDataDesc& Desc : InputDataDesc->DataDescs)
+			for (const FPCGDataDesc& Desc : InputDataDesc->GetDataDescriptions())
 			{
 				if (Desc.GetElementCount().X <= 0)
 				{
 					continue;
 				}
 
-				for (const FPCGKernelAttributeDesc& AttributeDesc : Desc.AttributeDescs)
+				for (const FPCGKernelAttributeDesc& AttributeDesc : Desc.GetAttributeDescriptions())
 				{
-					if (AttributeDesc.AttributeKey.Identifier.Name == SelectorName)
+					if (AttributeDesc.GetAttributeKey().GetIdentifier().Name == SelectorName)
 					{
-						if (AttributeDesc.AttributeKey.Type == EPCGKernelAttributeType::StringKey)
+						if (AttributeDesc.GetAttributeKey().GetType() == EPCGKernelAttributeType::StringKey)
 						{
 							for (int32 AttributeStringKeyValue : AttributeDesc.GetUniqueStringKeys())
 							{
diff --git a/Engine/Plugins/PCG/Source/PCG/Private/Compute/DataInterfaces/Elements/PCGTransformPointsDataInterface.cpp b/Engine/Plugins/PCG/Source/PCG/Private/Compute/DataInterfaces/Elements/PCGTransformPointsDataInterface.cpp
--- a/Engine/Plugins/PCG/Source/PCG/Private/Compute/DataInterfaces/Elements/PCGTransformPointsDataInterface.cpp
+++ b/Engine/Plugins/PCG/Source/PCG/Private/Compute/DataInterfaces/Elements/PCGTransformPointsDataInterface.cpp
@@ -198,7 +198,7 @@ bool UPCGTransformPointsDataProvider::PrepareForExecute_GameThread(UPCGDataBindi
 
 		bool bAnyPointsPresent = false;
 
-		for (const FPCGDataDesc& Desc : InputDataDesc->DataDescs)
+		for (const FPCGDataDesc& Desc : InputDataDesc->GetDataDescriptions())
 		{
 			if (Desc.GetElementCount().X <= 0)
 			{
@@ -207,20 +207,20 @@ bool UPCGTransformPointsDataProvider::PrepareForExecute_GameThread(UPCGDataBindi
 
 			bAnyPointsPresent = true;
 			
-			const FPCGKernelAttributeDesc* It = Algo::FindByPredicate(Desc.AttributeDescs, [Settings](const FPCGKernelAttributeDesc& AttributeDesc)
+			const FPCGKernelAttributeDesc* It = Algo::FindByPredicate(Desc.GetAttributeDescriptions(), [Settings](const FPCGKernelAttributeDesc& AttributeDesc)
 			{
-				return AttributeDesc.AttributeKey.Identifier.Name == Settings->AttributeName && AttributeDesc.AttributeKey.Type == EPCGKernelAttributeType::Transform;
+				return AttributeDesc.GetAttributeKey().GetIdentifier().Name == Settings->AttributeName && AttributeDesc.GetAttributeKey().GetType() == EPCGKernelAttributeType::Transform;
 			});
 			
 			if (It)
 			{
-				AttributeId = It->AttributeId;
+				AttributeId = It->GetAttributeId();
 				break;
 			}
 		}
 
 		// Mute this error if the point data is empty.
-		if (AttributeId == INDEX_NONE && !InputDataDesc->DataDescs.IsEmpty() && bAnyPointsPresent)
+		if (AttributeId == INDEX_NONE && !InputDataDesc->GetDataDescriptions().IsEmpty() && bAnyPointsPresent)
 		{
 			PCG_KERNEL_VALIDATION_ERR(Context, Settings, FText::Format(
 				LOCTEXT("TransformPointsAttributeNotFound", "Transfrom points attribute '{0}' not found."),
diff --git a/Engine/Plugins/PCG/Source/PCG/Private/Compute/DataInterfaces/PCGDataCollectionDataInterface.cpp b/Engine/Plugins/PCG/Source/PCG/Private/Compute/DataInterfaces/PCGDataCollectionDataInterface.cpp
--- a/Engine/Plugins/PCG/Source/PCG/Private/Compute/DataInterfaces/PCGDataCollectionDataInterface.cpp
+++ b/Engine/Plugins/PCG/Source/PCG/Private/Compute/DataInterfaces/PCGDataCollectionDataInterface.cpp
@@ -678,11 +678,11 @@ FPCGDataCollectionDataProviderProxy::FPCGDataCollectionDataProviderProxy(const F
 	{
 		PinDesc = InParams.PinDesc;
 
-		for (const FPCGDataDesc& DataDesc : PinDesc->DataDescs)
+		for (const FPCGDataDesc& DataDesc : PinDesc->GetDataDescriptions())
 		{
-			for (const FPCGKernelAttributeDesc& AttrDesc : DataDesc.AttributeDescs)
+			for (const FPCGKernelAttributeDesc& AttrDesc : DataDesc.GetAttributeDescriptions())
 			{
-				LargestAttributeIndex = FMath::Max(LargestAttributeIndex, PCGComputeHelpers::GetMetadataAttributeIndexFromAttributeId(AttrDesc.AttributeId));
+				LargestAttributeIndex = FMath::Max(LargestAttributeIndex, PCGComputeHelpers::GetMetadataAttributeIndexFromAttributeId(AttrDesc.GetAttributeId()));
 			}
 		}
 	}
@@ -823,7 +823,7 @@ void FPCGDataCollectionDataProviderProxy::AllocateResources(FRDGBuilder& GraphBu
 				{
 					TSharedPtr<FPCGProxyForGPUDataCollection> DataCollectionOnGPU = MakeShared<FPCGProxyForGPUDataCollection>(ExportedPooledBuffer, SizeBytes, PinDesc, Binding->GetStringTable());
 
-					for (int DataIndex = 0; DataIndex < PinDesc->DataDescs.Num(); ++DataIndex)
+					for (int DataIndex = 0; DataIndex < PinDesc->GetDataDescriptions().Num(); ++DataIndex)
 					{
 						UPCGProxyForGPUData* Proxy = NewObject<UPCGProxyForGPUData>();
 						Proxy->Initialize(DataCollectionOnGPU, DataIndex);
diff --git a/Engine/Plugins/PCG/Source/PCG/Private/Compute/DataInterfaces/PCGDataLabelResolverDataInterface.cpp b/Engine/Plugins/PCG/Source/PCG/Private/Compute/DataInterfaces/PCGDataLabelResolverDataInterface.cpp
--- a/Engine/Plugins/PCG/Source/PCG/Private/Compute/DataInterfaces/PCGDataLabelResolverDataInterface.cpp
+++ b/Engine/Plugins/PCG/Source/PCG/Private/Compute/DataInterfaces/PCGDataLabelResolverDataInterface.cpp
@@ -117,9 +117,9 @@ bool UPCGDataLabelResolverDataProvider::PrepareForExecute_GameThread(UPCGDataBin
 			// Initialize with invalid indices. The data IDs are not ordered.
 			DataIdToDataIndexMap[LabelIndex] = INDEX_NONE;
 
-			for (int32 InputDataIndex = 0; InputDataIndex < PinDataDesc->DataDescs.Num(); ++InputDataIndex)
+			for (int32 InputDataIndex = 0; InputDataIndex < PinDataDesc->GetDataDescriptions().Num(); ++InputDataIndex)
 			{
-				for (const int32 TagStringKey : PinDataDesc->DataDescs[InputDataIndex].TagStringKeys)
+				for (const int32 TagStringKey : PinDataDesc->GetDataDescriptions()[InputDataIndex].GetTagStringKeys())
 				{
 					if (StringTable[TagStringKey] == PrefixedLabel)
 					{
diff --git a/Engine/Plugins/PCG/Source/PCG/Private/Compute/DataInterfaces/PCGRawBufferDataInterface.cpp b/Engine/Plugins/PCG/Source/PCG/Private/Compute/DataInterfaces/PCGRawBufferDataInterface.cpp
--- a/Engine/Plugins/PCG/Source/PCG/Private/Compute/DataInterfaces/PCGRawBufferDataInterface.cpp
+++ b/Engine/Plugins/PCG/Source/PCG/Private/Compute/DataInterfaces/PCGRawBufferDataInterface.cpp
@@ -34,9 +34,7 @@ namespace PCGRawBufferDataInterface
 
 		if (InDataCollectionDesc)
 		{
-			const TArray<FPCGDataDesc>& DataDescs = InDataCollectionDesc->DataDescs;
-
-			for (const FPCGDataDesc& DataDesc : DataDescs)
+			for (const FPCGDataDesc& DataDesc : InDataCollectionDesc->GetDataDescriptions())
 			{
 				TotalDataSizeBytes += ComputePackedSizeBytes(DataDesc);
 			}
diff --git a/Engine/Plugins/PCG/Source/PCG/Private/Compute/DataInterfaces/PCGTextureDataInterface.cpp b/Engine/Plugins/PCG/Source/PCG/Private/Compute/DataInterfaces/PCGTextureDataInterface.cpp
--- a/Engine/Plugins/PCG/Source/PCG/Private/Compute/DataInterfaces/PCGTextureDataInterface.cpp
+++ b/Engine/Plugins/PCG/Source/PCG/Private/Compute/DataInterfaces/PCGTextureDataInterface.cpp
@@ -354,10 +354,8 @@ void UPCGTextureDataProvider::BuildInfosFromDataDescription(UPCGDataBinding* InB
 	FTransform Transform = Context->ExecutionSource->GetExecutionState().GetOriginalTransform();
 	Transform.SetScale3D(Transform.GetScale3D() * 0.5 * (OriginalLocalBounds.Max - OriginalLocalBounds.Min));
 
-	for (int DataIndex = 0; DataIndex < PinDesc->DataDescs.Num(); ++DataIndex)
+	for (const FPCGDataDesc& DataDesc : PinDesc->GetDataDescriptions())
 	{
-		const FPCGDataDesc& DataDesc = PinDesc->DataDescs[DataIndex];
-
 		FPCGTextureBindingInfo BindingInfo(DataDesc, Transform);
 		int32 BindingIndex = BindingInfos.Find(BindingInfo);
 
@@ -632,7 +630,7 @@ void FPCGTextureDataProviderProxy::ExportTextureUAVs(const TArray<TRefCountPtr<I
 			{
 				const FPCGTextureInfo& TextureInfo = TextureInfos_GT[TextureDataIndex];
 				check(ExportedTextures.IsValidIndex(TextureDataIndex));
-				check(PinDesc->DataDescs.IsValidIndex(TextureDataIndex))
+				check(PinDesc->GetDataDescriptions().IsValidIndex(TextureDataIndex))
 				check(BindingInfos_GT.IsValidIndex(TextureInfo.BindingIndex));
 
 				const FPCGTextureBindingInfo& BindingInfo = BindingInfos_GT[TextureInfo.BindingIndex];
@@ -646,7 +644,7 @@ void FPCGTextureDataProviderProxy::ExportTextureUAVs(const TArray<TRefCountPtr<I
 				ExportedData->Initialize(ExportedTextures[TextureDataIndex], TextureInfo.SliceIndex, TextureTransform, /*bInSkipReadbackToCPU=*/true);
 
 				TSet<FString> Tags;
-				for (const int32 TagStringKey : PinDesc->DataDescs[TextureDataIndex].TagStringKeys)
+				for (const int32 TagStringKey : PinDesc->GetDataDescriptions()[TextureDataIndex].GetTagStringKeys())
 				{
 					if (ensure(StringTable.IsValidIndex(TagStringKey)))
 					{
diff --git a/Engine/Plugins/PCG/Source/PCG/Private/Compute/Elements/PCGCustomHLSL.h b/Engine/Plugins/PCG/Source/PCG/Private/Compute/Elements/PCGCustomHLSL.h
--- a/Engine/Plugins/PCG/Source/PCG/Private/Compute/Elements/PCGCustomHLSL.h
+++ b/Engine/Plugins/PCG/Source/PCG/Private/Compute/Elements/PCGCustomHLSL.h
@@ -63,7 +63,7 @@ public:
 #if WITH_EDITOR
 	virtual FName GetDefaultNodeName() const override { return FName(TEXT("CustomHLSL")); }
 	virtual FText GetDefaultNodeTitle() const override { return NSLOCTEXT("PCGCustomHLSLElement", "NodeTitle", "Custom HLSL"); }
-	virtual FText GetNodeTooltipText() const override { return NSLOCTEXT("PCGCustomHLSLElement", "NodeTooltip", "[EXPERIMENTAL] Produces a HLSL compute shader which will be executed on the GPU."); }
+	virtual FText GetNodeTooltipText() const override { return NSLOCTEXT("PCGCustomHLSLElement", "NodeTooltip", "Produces a HLSL compute shader which will be executed on the GPU."); }
 	virtual EPCGSettingsType GetType() const override { return EPCGSettingsType::GPU; }
 	void CreateKernels(FPCGGPUCompilationContext& InOutContext, UObject* InObjectOuter, TArray<UPCGComputeKernel*>& OutKernels, TArray<FPCGKernelEdge>& OutEdges) const override;
 
diff --git a/Engine/Plugins/PCG/Source/PCG/Private/Compute/Elements/PCGCustomHLSLKernel.cpp b/Engine/Plugins/PCG/Source/PCG/Private/Compute/Elements/PCGCustomHLSLKernel.cpp
--- a/Engine/Plugins/PCG/Source/PCG/Private/Compute/Elements/PCGCustomHLSLKernel.cpp
+++ b/Engine/Plugins/PCG/Source/PCG/Private/Compute/Elements/PCGCustomHLSLKernel.cpp
@@ -266,14 +266,14 @@ TSharedPtr<const FPCGDataCollectionDesc> UPCGCustomHLSLKernel::ComputeOutputBind
 		OutDataDesc = FPCGDataCollectionDesc::MakeShared();
 
 		// Generators always produce a single point data with known point count.
-		OutDataDesc->DataDescs.Emplace(EPCGDataType::Point, KernelParams->GetValueInt(GET_MEMBER_NAME_CHECKED(UPCGCustomHLSLSettings, NumElements)));
+		OutDataDesc->GetDataDescriptionsMutable().Emplace(EPCGDataType::Point, KernelParams->GetValueInt(GET_MEMBER_NAME_CHECKED(UPCGCustomHLSLSettings, NumElements)));
 	}
 	else if (InOutputPinLabel == FirstOutputPinLabel && CustomHLSLSettings->KernelType == EPCGKernelType::TextureGenerator)
 	{
 		OutDataDesc = FPCGDataCollectionDesc::MakeShared();
 
 		// Generators always produce a single texture data with known size.
-		OutDataDesc->DataDescs.Emplace(EPCGDataType::BaseTexture, FIntVector4(CustomHLSLSettings->NumElements2D.X, CustomHLSLSettings->NumElements2D.Y, 0, 0));
+		OutDataDesc->GetDataDescriptionsMutable().Emplace(EPCGDataType::BaseTexture, FIntVector4(CustomHLSLSettings->NumElements2D.X, CustomHLSLSettings->NumElements2D.Y, 0, 0));
 	}
 	else if (ensure(OutputPinProperties))
 	{
@@ -298,14 +298,14 @@ TSharedPtr<const FPCGDataCollectionDesc> UPCGCustomHLSLKernel::ComputeOutputBind
 	// from matching attribute names (across all input pins), and then fall back to collecting keys from all string key attributes across all inputs.
 	bool bOutputHasStringKeys = false;
 
-	for (const FPCGDataDesc& DataDesc : OutDataDesc->DataDescs)
+	for (const FPCGDataDesc& DataDesc : OutDataDesc->GetDataDescriptions())
 	{
 		auto HasStringKey = [](const FPCGKernelAttributeDesc& InAttributeDesc)
 		{
-			return InAttributeDesc.AttributeKey.Type == EPCGKernelAttributeType::StringKey || InAttributeDesc.AttributeKey.Type == EPCGKernelAttributeType::Name;
+			return InAttributeDesc.GetAttributeKey().GetType() == EPCGKernelAttributeType::StringKey || InAttributeDesc.GetAttributeKey().GetType() == EPCGKernelAttributeType::Name;
 		};
 
-		if (DataDesc.AttributeDescs.FindByPredicate(HasStringKey))
+		if (DataDesc.GetAttributeDescriptions().FindByPredicate(HasStringKey))
 		{
 			bOutputHasStringKeys = true;
 			break;
@@ -329,11 +329,11 @@ TSharedPtr<const FPCGDataCollectionDesc> UPCGCustomHLSLKernel::ComputeOutputBind
 
 			bool bFoundStringKeyAttribute = false;
 
-			for (FPCGDataDesc& DataDesc : OutDataDesc->DataDescs)
+			for (FPCGDataDesc& DataDesc : OutDataDesc->GetDataDescriptionsMutable())
 			{
-				for (FPCGKernelAttributeDesc& AttributeDesc : DataDesc.AttributeDescs)
+				for (FPCGKernelAttributeDesc& AttributeDesc : DataDesc.GetAttributeDescriptionsMutable())
 				{
-					bFoundStringKeyAttribute |= (AttributeDesc.AttributeKey.Type == EPCGKernelAttributeType::StringKey || AttributeDesc.AttributeKey.Type == EPCGKernelAttributeType::Name);
+					bFoundStringKeyAttribute |= (AttributeDesc.GetAttributeKey().GetType() == EPCGKernelAttributeType::StringKey || AttributeDesc.GetAttributeKey().GetType() == EPCGKernelAttributeType::Name);
 				}
 			}
 
@@ -345,11 +345,11 @@ TSharedPtr<const FPCGDataCollectionDesc> UPCGCustomHLSLKernel::ComputeOutputBind
 
 		if (!RelevantInputDataDescs.IsEmpty())
 		{
-			for (FPCGDataDesc& DataDesc : OutDataDesc->DataDescs)
+			for (FPCGDataDesc& DataDesc : OutDataDesc->GetDataDescriptionsMutable())
 			{
-				for (FPCGKernelAttributeDesc& AttributeDesc : DataDesc.AttributeDescs)
+				for (FPCGKernelAttributeDesc& AttributeDesc : DataDesc.GetAttributeDescriptionsMutable())
 				{
-					if (AttributeDesc.AttributeKey.Type != EPCGKernelAttributeType::StringKey && AttributeDesc.AttributeKey.Type != EPCGKernelAttributeType::Name)
+					if (AttributeDesc.GetAttributeKey().GetType() != EPCGKernelAttributeType::StringKey && AttributeDesc.GetAttributeKey().GetType() != EPCGKernelAttributeType::Name)
 					{
 						continue;
 					}
@@ -360,12 +360,12 @@ TSharedPtr<const FPCGDataCollectionDesc> UPCGCustomHLSLKernel::ComputeOutputBind
 					{
 						// Try to find string keys for matching attributes on inputs. E.g. if we are processing an output attribute named 'MeshPath',
 						// look at data on all input pins for an attribute named MeshPath and assume we could use any of its values - copy the string keys.
-						for (const FPCGDataDesc& InputDataDesc : InputPinDataDesc->DataDescs)
+						for (const FPCGDataDesc& InputDataDesc : InputPinDataDesc->GetDataDescriptions())
 						{
-							for (const FPCGKernelAttributeDesc& InputAttributeDesc : InputDataDesc.AttributeDescs)
+							for (const FPCGKernelAttributeDesc& InputAttributeDesc : InputDataDesc.GetAttributeDescriptions())
 							{
-								const bool bIsString = InputAttributeDesc.AttributeKey.Type == EPCGKernelAttributeType::StringKey || InputAttributeDesc.AttributeKey.Type == EPCGKernelAttributeType::Name;
-								if (bIsString && InputAttributeDesc.AttributeKey.Identifier == AttributeDesc.AttributeKey.Identifier)
+								const bool bIsString = InputAttributeDesc.GetAttributeKey().GetType() == EPCGKernelAttributeType::StringKey || InputAttributeDesc.GetAttributeKey().GetType() == EPCGKernelAttributeType::Name;
+								if (bIsString && InputAttributeDesc.GetAttributeKey().GetIdentifier() == AttributeDesc.GetAttributeKey().GetIdentifier())
 								{
 									AttributeDesc.AddUniqueStringKeys(InputAttributeDesc.GetUniqueStringKeys());
 									bFoundMatchingAttribute = true;
@@ -383,11 +383,11 @@ TSharedPtr<const FPCGDataCollectionDesc> UPCGCustomHLSLKernel::ComputeOutputBind
 						{
 							// Try to find string keys for matching attributes on inputs. E.g. if we are processing an output attribute named 'MeshPath',
 							// look at data on all input pins for an attribute named MeshPath and assume we could use any of its values - copy the string keys.
-							for (const FPCGDataDesc& InputDataDesc : InputPinDataDesc->DataDescs)
+							for (const FPCGDataDesc& InputDataDesc : InputPinDataDesc->GetDataDescriptions())
 							{
-								for (const FPCGKernelAttributeDesc& InputAttributeDesc : InputDataDesc.AttributeDescs)
+								for (const FPCGKernelAttributeDesc& InputAttributeDesc : InputDataDesc.GetAttributeDescriptions())
 								{
-									if (InputAttributeDesc.AttributeKey.Type == EPCGKernelAttributeType::StringKey || InputAttributeDesc.AttributeKey.Type == EPCGKernelAttributeType::Name)
+									if (InputAttributeDesc.GetAttributeKey().GetType() == EPCGKernelAttributeType::StringKey || InputAttributeDesc.GetAttributeKey().GetType() == EPCGKernelAttributeType::Name)
 									{
 										AttributeDesc.AddUniqueStringKeys(InputAttributeDesc.GetUniqueStringKeys());
 									}
@@ -654,7 +654,7 @@ int UPCGCustomHLSLKernel::ComputeThreadCount(const UPCGDataBinding* InBinding) c
 			if (const FPCGPinPropertiesGPU* OutputPin = CustomHLSLSettings->OutputPins.IsEmpty() ? nullptr : &CustomHLSLSettings->OutputPins[0])
 			{
 				const TSharedPtr<const FPCGDataCollectionDesc> Desc = InBinding->GetCachedKernelPinDataDesc(this, OutputPin->Label, /*bIsInput=*/false);
-				ThreadCount = Desc ? OutputPin->GetElementCountMultiplier() * Desc->ComputeTotalElementCount(OutputPin->AllowedTypes) : 0;
+				ThreadCount = Desc ? OutputPin->GetElementCountMultiplier() * Desc->ComputeTotalElementCount() : 0;
 			}
 		}
 		else if (CustomHLSLSettings->DispatchThreadCount == EPCGDispatchThreadCount::FromProductOfInputPins)
@@ -944,9 +944,7 @@ void UPCGCustomHLSLKernel::ParseShaderSource()
 			}
 
 			// Add the attribute if it hasn't already been referenced.
-			FPCGKernelAttributeKey Key;
-			Key.Type = static_cast<EPCGKernelAttributeType>(AttributeType);
-			Key.SetSelector(FPCGAttributePropertySelector::CreateSelectorFromString(NameStr));
+			const FPCGKernelAttributeKey Key(FPCGAttributePropertySelector::CreateSelectorFromString(NameStr), static_cast<EPCGKernelAttributeType>(AttributeType));
 
 			if (Key.IsValid())
 			{
@@ -1531,7 +1529,7 @@ bool UPCGCustomHLSLKernel::PerformStaticValidation()
 
 	for (const FPCGKernelAttributeKey& AttributeKey : KernelAttributeKeys)
 	{
-		if (AttributeKey.Type == EPCGKernelAttributeType::Invalid)
+		if (AttributeKey.GetType() == EPCGKernelAttributeType::Invalid)
 		{
 			const UEnum* AttributeTypeEnum = StaticEnum<EPCGKernelAttributeType>();
 			check(AttributeTypeEnum);
@@ -1539,8 +1537,8 @@ bool UPCGCustomHLSLKernel::PerformStaticValidation()
 #if PCG_KERNEL_LOGGING_ENABLED
 			AddStaticLogEntry(FText::Format(
 				LOCTEXT("InvalidAttributeTypeNone", "Attribute '{0}' has invalid GPU attribute type '{1}', check the 'Attributes to Create' array on your pins."),
-				FText::FromName(AttributeKey.Identifier.Name),
-				FText::FromString(AttributeTypeEnum->GetNameStringByValue(static_cast<int64>(AttributeKey.Type)))),
+				FText::FromName(AttributeKey.GetIdentifier().Name),
+				FText::FromString(AttributeTypeEnum->GetNameStringByValue(static_cast<int64>(AttributeKey.GetType())))),
 				EPCGKernelLogVerbosity::Error);
 #endif
 			return false;
@@ -1867,14 +1865,14 @@ bool UPCGCustomHLSLKernel::AreAttributesValid(FPCGContext* InContext, FText* Out
 				bool bFoundMatchingAttributeName = false;
 
 				// Verify that the attribute exists on at least one data in pin data collection.
-				for (const FPCGDataDesc& DataDesc : PinDesc->DataDescs)
+				for (const FPCGDataDesc& DataDesc : PinDesc->GetDataDescriptions())
 				{
-					AttrDesc = DataDesc.AttributeDescs.FindByPredicate([&AttrKey, &bFoundMatchingAttributeName](const FPCGKernelAttributeDesc& Desc)
+					AttrDesc = DataDesc.GetAttributeDescriptions().FindByPredicate([&AttrKey, &bFoundMatchingAttributeName](const FPCGKernelAttributeDesc& Desc)
 					{
-						const bool bAttributeNameMatches = Desc.AttributeKey.Identifier == AttrKey.Identifier;
+						const bool bAttributeNameMatches = Desc.GetAttributeKey().GetIdentifier() == AttrKey.GetIdentifier();
 						bFoundMatchingAttributeName |= bAttributeNameMatches;
 
-						return bAttributeNameMatches && Desc.AttributeKey.Type == AttrKey.Type;
+						return bAttributeNameMatches && Desc.GetAttributeKey().GetType() == AttrKey.GetType();
 					});
 
 					if (AttrDesc)
@@ -1883,7 +1881,7 @@ bool UPCGCustomHLSLKernel::AreAttributesValid(FPCGContext* InContext, FText* Out
 					}
 				}
 
-				if (!AttrDesc && !PinDesc->DataDescs.IsEmpty())
+				if (!AttrDesc && !PinDesc->GetDataDescriptions().IsEmpty())
 				{
 #if PCG_KERNEL_LOGGING_ENABLED
 					if (OutErrorText)
@@ -1943,7 +1941,7 @@ int UPCGCustomHLSLKernel::GetElementCountForInputPin(const FPCGPinProperties& In
 
 	const TSharedPtr<const FPCGDataCollectionDesc> InputDesc = InBinding->GetCachedKernelPinDataDesc(this, InInputPinProps.Label, /*bIsInputPin=*/true);
 
-	return InputDesc ? InputDesc->ComputeTotalElementCount(InInputPinProps.AllowedTypes) : 0;
+	return InputDesc ? InputDesc->ComputeTotalElementCount() : 0;
 }
 
 #undef LOCTEXT_NAMESPACE
diff --git a/Engine/Plugins/PCG/Source/PCG/Private/Compute/Elements/PCGNormalToDensityKernel.cpp b/Engine/Plugins/PCG/Source/PCG/Private/Compute/Elements/PCGNormalToDensityKernel.cpp
--- a/Engine/Plugins/PCG/Source/PCG/Private/Compute/Elements/PCGNormalToDensityKernel.cpp
+++ b/Engine/Plugins/PCG/Source/PCG/Private/Compute/Elements/PCGNormalToDensityKernel.cpp
@@ -34,7 +34,7 @@ TSharedPtr<const FPCGDataCollectionDesc> UPCGNormalToDensityKernel::ComputeOutpu
 	TSharedPtr<FPCGDataCollectionDesc> OutDataDesc = FPCGDataCollectionDesc::MakeSharedFrom(InDataDesc);
 
 	// Allocate density property if the transform is also allocated. Otherwise the densities will all be the same value and full-width allocation is unnecessary.
-	for (FPCGDataDesc& DataDesc : OutDataDesc->DataDescs)
+	for (FPCGDataDesc& DataDesc : OutDataDesc->GetDataDescriptionsMutable())
 	{
 		if (!!(DataDesc.GetAllocatedProperties() & EPCGPointNativeProperties::Transform))
 		{
diff --git a/Engine/Plugins/PCG/Source/PCG/Private/Compute/Elements/PCGTransformPointsKernel.cpp b/Engine/Plugins/PCG/Source/PCG/Private/Compute/Elements/PCGTransformPointsKernel.cpp
--- a/Engine/Plugins/PCG/Source/PCG/Private/Compute/Elements/PCGTransformPointsKernel.cpp
+++ b/Engine/Plugins/PCG/Source/PCG/Private/Compute/Elements/PCGTransformPointsKernel.cpp
@@ -41,7 +41,7 @@ TSharedPtr<const FPCGDataCollectionDesc> UPCGTransformPointsKernel::ComputeOutpu
 	}
 
 	// Allocate properties if needed
-	for (FPCGDataDesc& DataDesc : OutDataDesc->DataDescs)
+	for (FPCGDataDesc& DataDesc : OutDataDesc->GetDataDescriptionsMutable())
 	{
 		DataDesc.AllocateProperties(Settings->GetPropertiesToAllocate());
 	}
diff --git a/Engine/Plugins/PCG/Source/PCG/Private/Compute/PCGComputeGraph.cpp b/Engine/Plugins/PCG/Source/PCG/Private/Compute/PCGComputeGraph.cpp
--- a/Engine/Plugins/PCG/Source/PCG/Private/Compute/PCGComputeGraph.cpp
+++ b/Engine/Plugins/PCG/Source/PCG/Private/Compute/PCGComputeGraph.cpp
@@ -210,14 +210,14 @@ void UPCGComputeGraph::DebugLogDataDescriptions(const UPCGDataBinding* InBinding
 			return;
 		}
 
-		for (int32 DataIndex = 0; DataIndex < InDataDesc->DataDescs.Num(); ++DataIndex)
+		for (int32 DataIndex = 0; DataIndex < InDataDesc->GetDataDescriptions().Num(); ++DataIndex)
 		{
-			const FPCGDataDesc& DataDesc = InDataDesc->DataDescs[DataIndex];
+			const FPCGDataDesc& DataDesc = InDataDesc->GetDataDescriptions()[DataIndex];
 
 			UE_LOG(LogPCG, Display, TEXT("\t\tData Index %d"), DataIndex);
 			UE_LOG(LogPCG, Display, TEXT("\t\t\tType: %s"), *DataDesc.GetType().ToString());
 			UE_LOG(LogPCG, Display, TEXT("\t\t\tNum Elements: %s (%s)"), *DataDesc.GetElementCount().ToString(), *UEnum::GetValueAsString(DataDesc.GetElementDimension()));
-			UE_LOG(LogPCG, Display, TEXT("\t\t\tAttributes (%d)"), DataDesc.AttributeDescs.Num());
+			UE_LOG(LogPCG, Display, TEXT("\t\t\tAttributes (%d)"), DataDesc.GetAttributeDescriptions().Num());
 
 			TSharedPtr<const FPCGProxyForGPUDataCollection> UpstreamDataCollectionInfo;
 
@@ -234,9 +234,8 @@ void UPCGComputeGraph::DebugLogDataDescriptions(const UPCGDataBinding* InBinding
 				}
 			}
 
-			for (int32 AttributeIndex = 0; AttributeIndex < DataDesc.AttributeDescs.Num(); ++AttributeIndex)
+			for (const FPCGKernelAttributeDesc& AttributeDesc : DataDesc.GetAttributeDescriptions())
 			{
-				const FPCGKernelAttributeDesc& AttributeDesc = DataDesc.AttributeDescs[AttributeIndex];
 				const bool bAllocated = DataDesc.IsAttributeAllocated(AttributeDesc);
 
 				// If the data comes from an upstream / source buffer then look up its info.
@@ -255,9 +254,9 @@ void UPCGComputeGraph::DebugLogDataDescriptions(const UPCGDataBinding* InBinding
 				// If attribute comes from an upstream GPU buffer, display the attribute ID in the upstream buffer.
 				int32 SourceBufferAttributeId = INDEX_NONE;
 
-				if (AttributeDesc.AttributeId >= PCGDataCollectionPackingConstants::NUM_RESERVED_ATTRS && UpstreamDataCollectionInfo)
+				if (AttributeDesc.GetAttributeId() >= PCGDataCollectionPackingConstants::NUM_RESERVED_ATTRS && UpstreamDataCollectionInfo)
 				{
-					const int32 GraphAttributeIndex = PCGComputeHelpers::GetMetadataAttributeIndexFromAttributeId(AttributeDesc.AttributeId);
+					const int32 GraphAttributeIndex = PCGComputeHelpers::GetMetadataAttributeIndexFromAttributeId(AttributeDesc.GetAttributeId());
 					const int32* FoundIndex = InBinding->GetGraphAttributeToSourceBufferAttributeIndex().Find({ GraphAttributeIndex, UpstreamDataCollectionInfo->GetBuffer() });
 
 					SourceBufferAttributeId = PCGComputeHelpers::GetAttributeIdFromMetadataAttributeIndex(FoundIndex ? *FoundIndex : GraphAttributeIndex);
@@ -266,28 +265,28 @@ void UPCGComputeGraph::DebugLogDataDescriptions(const UPCGDataBinding* InBinding
 				if (SourceBufferAttributeId == INDEX_NONE)
 				{
 					UE_LOG(LogPCG, Display, TEXT("\t\t\t\tID: %d\t\tName: %s\t\tType: %s\t\tAllocated: %s\t\tUniqueStringKeys%s"),
-						AttributeDesc.AttributeId,
-						*AttributeDesc.AttributeKey.Identifier.ToString(),
-						*PCGKernelAttributeTypeEnum->GetDisplayNameTextByValue(static_cast<int64>(AttributeDesc.AttributeKey.Type)).ToString(),
+						AttributeDesc.GetAttributeId(),
+						*AttributeDesc.GetAttributeKey().GetIdentifier().ToString(),
+						*PCGKernelAttributeTypeEnum->GetDisplayNameTextByValue(static_cast<int64>(AttributeDesc.GetAttributeKey().GetType())).ToString(),
 						bAllocated ? TEXT("true") : TEXT("false"),
 						UniqueStringKeys.IsEmpty() ? TEXT(": ") : *UniqueStringKeys);
 				}
 				else
 				{
 					UE_LOG(LogPCG, Display, TEXT("\t\t\t\tID: %d\t\tName: %s\t\tType: %s\t\tAllocated: %s\t\tUniqueStringKeys%s\t\tSource Buffer ID: %d"),
-						AttributeDesc.AttributeId,
-						*AttributeDesc.AttributeKey.Identifier.ToString(),
-						*PCGKernelAttributeTypeEnum->GetDisplayNameTextByValue(static_cast<int64>(AttributeDesc.AttributeKey.Type)).ToString(),
+						AttributeDesc.GetAttributeId(),
+						*AttributeDesc.GetAttributeKey().GetIdentifier().ToString(),
+						*PCGKernelAttributeTypeEnum->GetDisplayNameTextByValue(static_cast<int64>(AttributeDesc.GetAttributeKey().GetType())).ToString(),
 						bAllocated ? TEXT("true") : TEXT("false"),
 						UniqueStringKeys.IsEmpty() ? TEXT(": ") : *UniqueStringKeys,
 						SourceBufferAttributeId);
 				}
 			}
 
-			UE_LOG(LogPCG, Display, TEXT("\t\t\tTags (%d)"), DataDesc.TagStringKeys.Num());
+			UE_LOG(LogPCG, Display, TEXT("\t\t\tTags (%d)"), DataDesc.GetTagStringKeys().Num());
 
 			FString TagStringKeys;
-			for (int32 TagStringKey : DataDesc.TagStringKeys)
+			for (int32 TagStringKey : DataDesc.GetTagStringKeys())
 			{
 				TagStringKeys += TEXT(", ");
 				TagStringKeys += FString::Format(TEXT("{0}"), { TagStringKey });
@@ -316,7 +315,7 @@ void UPCGComputeGraph::DebugLogDataDescriptions(const UPCGDataBinding* InBinding
 				{
 					if (const TSharedPtr<const FPCGDataCollectionDesc> DataDesc = InBinding->GetCachedKernelPinDataDesc(PCGKernel, Pin.Label, true))
 					{
-						UE_LOG(LogPCG, Display, TEXT("\tInput Pin: %s (%d data)"), *Pin.Label.ToString(), DataDesc->DataDescs.Num());
+						UE_LOG(LogPCG, Display, TEXT("\tInput Pin: %s (%d data)"), *Pin.Label.ToString(), DataDesc->GetDataDescriptions().Num());
 
 						LogDataDescription(DataDesc, FPCGKernelPin(KernelIndex, Pin.Label, true));
 					}
@@ -331,7 +330,7 @@ void UPCGComputeGraph::DebugLogDataDescriptions(const UPCGDataBinding* InBinding
 				{
 					if (const TSharedPtr<const FPCGDataCollectionDesc> DataDesc = InBinding->GetCachedKernelPinDataDesc(PCGKernel, Pin.Label, false))
 					{
-						UE_LOG(LogPCG, Display, TEXT("\tOutput Pin: %s (%d data)"), *Pin.Label.ToString(), DataDesc->DataDescs.Num());
+						UE_LOG(LogPCG, Display, TEXT("\tOutput Pin: %s (%d data)"), *Pin.Label.ToString(), DataDesc->GetDataDescriptions().Num());
 
 						LogDataDescription(DataDesc, FPCGKernelPin(KernelIndex, Pin.Label, false));
 					}
@@ -382,7 +381,7 @@ TSharedPtr<const FPCGDataCollectionDesc> UPCGComputeGraph::ComputeExternalPinDes
 		FPCGDataDesc DataDesc;
 		if (ensure(ComputeTaggedDataPinDesc(TaggedData, InBinding, DataDesc)))
 		{
-			OutDataDesc->DataDescs.Add(MoveTemp(DataDesc));
+			OutDataDesc->GetDataDescriptionsMutable().Add(MoveTemp(DataDesc));
 		}
 	}
 
@@ -401,9 +400,9 @@ bool UPCGComputeGraph::ComputeTaggedDataPinDesc(const FPCGTaggedData& InTaggedDa
 		if (TSharedPtr<const FPCGProxyForGPUDataCollection> DataGPU = Proxy->GetGPUInfo())
 		{
 			// Remap attribute IDs and string keys.
-			for (FPCGKernelAttributeDesc& AttrDesc : OutDescription.AttributeDescs)
+			for (FPCGKernelAttributeDesc& AttrDesc : OutDescription.GetAttributeDescriptionsMutable())
 			{
-				const int32 Index = PCGComputeHelpers::GetMetadataAttributeIndexFromAttributeId(AttrDesc.AttributeId);
+				const int32 Index = PCGComputeHelpers::GetMetadataAttributeIndexFromAttributeId(AttrDesc.GetAttributeId());
 				if (Index == INDEX_NONE)
 				{
 					continue;
@@ -414,7 +413,7 @@ bool UPCGComputeGraph::ComputeTaggedDataPinDesc(const FPCGTaggedData& InTaggedDa
 				{
 					const int32 NewId = PCGComputeHelpers::GetAttributeIdFromMetadataAttributeIndex(*AttributeIndex);
 					//UE_LOG(LogPCG, Warning, TEXT("Remapped attribute '%s' ID: %d -> %d"), *AttrDesc.Name.ToString(), AttrDesc.Index, NewId);
-					AttrDesc.AttributeId = NewId;
+					AttrDesc.SetAttributeId(NewId);
 				}
 
 				// Map from the attribute IDs in the buffer / upstream compute graph to attribute IDs usable in the current graph.
@@ -446,7 +445,7 @@ bool UPCGComputeGraph::ComputeTaggedDataPinDesc(const FPCGTaggedData& InTaggedDa
 			}
 
 			// Remap tag string keys.
-			for (int32& TagStringKey : OutDescription.TagStringKeys)
+			for (int32& TagStringKey : OutDescription.GetTagStringKeysMutable())
 			{
 				if (const int32* GraphStringKey = InBinding->GetSourceBufferToGraphStringKey().Find({ DataGPU->GetBuffer(), TagStringKey }))
 				{
diff --git a/Engine/Plugins/PCG/Source/PCG/Private/Compute/PCGComputeKernel.cpp b/Engine/Plugins/PCG/Source/PCG/Private/Compute/PCGComputeKernel.cpp
--- a/Engine/Plugins/PCG/Source/PCG/Private/Compute/PCGComputeKernel.cpp
+++ b/Engine/Plugins/PCG/Source/PCG/Private/Compute/PCGComputeKernel.cpp
@@ -171,7 +171,7 @@ bool UPCGComputeKernel::IsKernelDataValid(FPCGContext* InContext) const
 					{
 						if (const TSharedPtr<const FPCGDataCollectionDesc> InputDesc = DataBinding->GetCachedKernelPinDataDesc(this, InitPin, /*bIsInputPin=*/true))
 						{
-							const uint32 InputNumData = InputDesc->DataDescs.Num();
+							const uint32 InputNumData = InputDesc->GetDataDescriptions().Num();
 
 							// Only test pins that have a non-trivial number of data.
 							if (InputNumData > 1)
@@ -280,17 +280,19 @@ void UPCGComputeKernel::ComputeDataDescFromPinProperties(const FPCGPinProperties
 		// Copies unique (non-reserved) attribute descriptions from 'InDataDesc' to 'OutDataDesc'.
 		auto AddAttributesFromData = [](const FPCGDataDesc& InDataDesc, FPCGDataDesc& OutDataDesc)
 		{
-			for (const FPCGKernelAttributeDesc& InAttrDesc : InDataDesc.AttributeDescs)
+			TArray<FPCGKernelAttributeDesc>& OutAttributeDescs = OutDataDesc.GetAttributeDescriptionsMutable();
+
+			for (const FPCGKernelAttributeDesc& InAttrDesc : InDataDesc.GetAttributeDescriptions())
 			{
-				if (InAttrDesc.AttributeId >= PCGDataCollectionPackingConstants::NUM_RESERVED_ATTRS
-					&& !OutDataDesc.AttributeDescs.ContainsByPredicate([InAttrDesc](const FPCGKernelAttributeDesc& OutAttrDesc)
+				if (InAttrDesc.GetAttributeId() >= PCGDataCollectionPackingConstants::NUM_RESERVED_ATTRS
+					&& !OutAttributeDescs.ContainsByPredicate([InAttrDesc](const FPCGKernelAttributeDesc& OutAttrDesc)
 						{
 							// Note: We shouldn't need to check for uniqueness of attr Index, since attributes should all have unique index via
 							// the GlobalAttributeLookupTable
-							return InAttrDesc.AttributeKey.Identifier == OutAttrDesc.AttributeKey.Identifier;
+							return InAttrDesc.GetAttributeKey().GetIdentifier() == OutAttrDesc.GetAttributeKey().GetIdentifier();
 						}))
 				{
-					OutDataDesc.AttributeDescs.Emplace(InAttrDesc);
+					OutAttributeDescs.Emplace(InAttrDesc);
 				}
 			}
 		};
@@ -331,7 +333,7 @@ void UPCGComputeKernel::ComputeDataDescFromPinProperties(const FPCGPinProperties
 					}
 				}
 
-				FPCGDataDesc& DataDesc = OutPinDesc->DataDescs.Emplace_GetRef(OutputPinProps.AllowedTypes, InitialNumElements);
+				FPCGDataDesc& DataDesc = OutPinDesc->GetDataDescriptionsMutable().Emplace_GetRef(OutputPinProps.AllowedTypes, InitialNumElements);
 
 				// For each data index, loop over all the pins and create the uber-data.
 				for (int InputPinIndex = 0; InputPinIndex < NumInitPins; ++InputPinIndex)
@@ -341,12 +343,12 @@ void UPCGComputeKernel::ComputeDataDescFromPinProperties(const FPCGPinProperties
 					int ClampedDataIndex = DataIndex;
 
 					// If this pin does not have the same number of data, clamp it to the last data.
-					if (DataIndex >= InputDesc->DataDescs.Num())
+					if (DataIndex >= InputDesc->GetDataDescriptions().Num())
 					{
-						ClampedDataIndex = InputDesc->DataDescs.Num() - 1;
+						ClampedDataIndex = InputDesc->GetDataDescriptions().Num() - 1;
 					}
 
-					const FPCGDataDesc& InputDataDesc = InputDesc->DataDescs.IsValidIndex(ClampedDataIndex) ? InputDesc->DataDescs[ClampedDataIndex] : FPCGDataDesc(FPCGDataTypeIdentifier{EPCGDataType::Any}, 0);
+					const FPCGDataDesc& InputDataDesc = InputDesc->GetDataDescriptions().IsValidIndex(ClampedDataIndex) ? InputDesc->GetDataDescriptions()[ClampedDataIndex] : FPCGDataDesc(FPCGDataTypeIdentifier{EPCGDataType::Any}, 0);
 					const EPCGElementDimension SourceDimension = InputDataDesc.GetElementDimension();
 
 					if (Props.ElementCountMode == EPCGElementCountMode::FromInputData)
@@ -382,9 +384,9 @@ void UPCGComputeKernel::ComputeDataDescFromPinProperties(const FPCGPinProperties
 					}
 
 					// Add unique tags from this data.
-					for (const int32 TagStringKey : InputDataDesc.TagStringKeys)
+					for (const int32 TagStringKey : InputDataDesc.GetTagStringKeys())
 					{
-						DataDesc.TagStringKeys.AddUnique(TagStringKey);
+						DataDesc.AddUniqueTagStringKey(TagStringKey);
 					}
 
 					DataDesc.AllocateProperties(InputDataDesc.GetAllocatedProperties());
@@ -397,9 +399,9 @@ void UPCGComputeKernel::ComputeDataDescFromPinProperties(const FPCGPinProperties
 			// If this is the only input pin, we can just copy it.
 			if (NumInitPins == 1)
 			{
-				for (const FPCGDataDesc& InputDataDesc : InputDescs[0]->DataDescs)
+				for (const FPCGDataDesc& InputDataDesc : InputDescs[0]->GetDataDescriptions())
 				{
-					FPCGDataDesc& DataDesc = OutPinDesc->DataDescs.Emplace_GetRef(OutputPinProps.AllowedTypes, /*InitialElementCount=*/0);
+					FPCGDataDesc& DataDesc = OutPinDesc->GetDataDescriptionsMutable().Emplace_GetRef(OutputPinProps.AllowedTypes, /*InitialElementCount=*/0);
 					DataDesc.CombineElementCount(InputDataDesc, EPCGElementMultiplicity::Sum);
 
 					if (Props.AttributeInheritanceMode == EPCGAttributeInheritanceMode::CopyAttributeSetup)
@@ -407,7 +409,7 @@ void UPCGComputeKernel::ComputeDataDescFromPinProperties(const FPCGPinProperties
 						AddAttributesFromData(InputDataDesc, DataDesc);
 					}
 
-					DataDesc.TagStringKeys = InputDataDesc.TagStringKeys;
+					DataDesc.SetTagStringKeys(InputDataDesc.GetTagStringKeys());
 					DataDesc.AllocateProperties(InputDataDesc.GetAllocatedProperties());
 				}
 			}
@@ -422,11 +424,11 @@ void UPCGComputeKernel::ComputeDataDescFromPinProperties(const FPCGPinProperties
 					{
 						const TSharedPtr<const FPCGDataCollectionDesc> OtherInputDesc = InputDescs[OtherInputPinIndex];
 
-						for (const FPCGDataDesc& InputDataDesc : InputDesc->DataDescs)
+						for (const FPCGDataDesc& InputDataDesc : InputDesc->GetDataDescriptions())
 						{
-							for (const FPCGDataDesc& OtherInputDataDesc : OtherInputDesc->DataDescs)
+							for (const FPCGDataDesc& OtherInputDataDesc : OtherInputDesc->GetDataDescriptions())
 							{
-								FPCGDataDesc& DataDesc = OutPinDesc->DataDescs.Emplace_GetRef(OutputPinProps.AllowedTypes, /*InitialElementCount=*/0);
+								FPCGDataDesc& DataDesc = OutPinDesc->GetDataDescriptionsMutable().Emplace_GetRef(OutputPinProps.AllowedTypes, /*InitialElementCount=*/0);
 
 								if (Props.ElementCountMode == EPCGElementCountMode::FromInputData)
 								{
@@ -462,14 +464,14 @@ void UPCGComputeKernel::ComputeDataDescFromPinProperties(const FPCGPinProperties
 								}
 
 								// Add unique tags from both input data.
-								for (const int32 TagStringKey : InputDataDesc.TagStringKeys)
+								for (const int32 TagStringKey : InputDataDesc.GetTagStringKeys())
 								{
-									DataDesc.TagStringKeys.AddUnique(TagStringKey);
+									DataDesc.AddUniqueTagStringKey(TagStringKey);
 								}
 
-								for (const int32 TagStringKey : OtherInputDataDesc.TagStringKeys)
+								for (const int32 TagStringKey : OtherInputDataDesc.GetTagStringKeys())
 								{
-									DataDesc.TagStringKeys.AddUnique(TagStringKey);
+									DataDesc.AddUniqueTagStringKey(TagStringKey);
 								}
 
 								DataDesc.AllocateProperties(InputDataDesc.GetAllocatedProperties() | OtherInputDataDesc.GetAllocatedProperties());
@@ -486,7 +488,7 @@ void UPCGComputeKernel::ComputeDataDescFromPinProperties(const FPCGPinProperties
 				// Find the maximum number of data among the init pins. Note, they should all be the same number of data, or only one data.
 				for (int I = 0; I < NumInitPins; ++I)
 				{
-					const int DataCount = InputDescs[I]->DataDescs.Num();
+					const int DataCount = InputDescs[I]->GetDataDescriptions().Num();
 
 					// If any init pins have no data, then the pairing fails and no data are produced.
 					if (DataCount == 0)
@@ -516,7 +518,7 @@ void UPCGComputeKernel::ComputeDataDescFromPinProperties(const FPCGPinProperties
 
 		// Apply element count multiplier.
 		const uint32 Multiplier = OutputPinProps.GetElementCountMultiplier();
-		for (FPCGDataDesc& Desc : OutPinDesc->DataDescs)
+		for (FPCGDataDesc& Desc : OutPinDesc->GetDataDescriptionsMutable())
 		{
 			Desc.ScaleElementCount(Multiplier);
 		}
@@ -529,11 +531,11 @@ void UPCGComputeKernel::ComputeDataDescFromPinProperties(const FPCGPinProperties
 
 			if (ElementDimension == EPCGElementDimension::One)
 			{
-				OutPinDesc->DataDescs.Emplace(OutputPinProps.AllowedTypes, Props.ElementCount);
+				OutPinDesc->GetDataDescriptionsMutable().Emplace(OutputPinProps.AllowedTypes, Props.ElementCount);
 			}
 			else if (ElementDimension == EPCGElementDimension::Two)
 			{
-				OutPinDesc->DataDescs.Emplace(OutputPinProps.AllowedTypes, Props.NumElements2D);
+				OutPinDesc->GetDataDescriptionsMutable().Emplace(OutputPinProps.AllowedTypes, Props.NumElements2D);
 			}
 			else
 			{
diff --git a/Engine/Plugins/PCG/Source/PCG/Private/Compute/PCGDataBinding.cpp b/Engine/Plugins/PCG/Source/PCG/Private/Compute/PCGDataBinding.cpp
--- a/Engine/Plugins/PCG/Source/PCG/Private/Compute/PCGDataBinding.cpp
+++ b/Engine/Plugins/PCG/Source/PCG/Private/Compute/PCGDataBinding.cpp
@@ -220,6 +220,26 @@ void UPCGDataBinding::ReceiveDataFromGPU_GameThread(
 	}
 }
 
+const TMap<UPCGDataBinding::FSourceBufferAndAttributeIndex, int32>& UPCGDataBinding::GetSourceBufferAttributeToGraphAttributeIndex() const
+{
+	return SourceBufferAttributeToGraphAttributeIndex;
+}
+
+const TMap<UPCGDataBinding::FGraphAttributeIndexAndSourceBuffer, int32>& UPCGDataBinding::GetGraphAttributeToSourceBufferAttributeIndex() const
+{
+	return GraphAttributeToSourceBufferAttributeIndex;
+}
+
+const TMap<UPCGDataBinding::FSourceBufferAndStringKey, int32>& UPCGDataBinding::GetSourceBufferToGraphStringKey() const
+{
+	return SourceBufferToGraphStringKey;
+}
+
+const TMap<UPCGDataBinding::FSourceBufferAndStringKey, int32>& UPCGDataBinding::GetGraphToSourceBufferStringKey() const
+{
+	return GraphToSourceBufferStringKey;
+}
+
 IPCGGraphExecutionSource* UPCGDataBinding::GetExecutionSource() const
 {
 	if (TSharedPtr<FPCGContextHandle> SharedHandle = ContextHandle.Pin())
@@ -393,17 +413,17 @@ void UPCGDataBinding::AddInputDataAttributesToTable()
 			FPCGDataDesc SourceDesc;
 			if (DataCollectionInfo && Proxy->GetDescription(SourceDesc))
 			{
-				for (FPCGKernelAttributeDesc& SourceAttrDesc : SourceDesc.AttributeDescs)
+				for (const FPCGKernelAttributeDesc& SourceAttrDesc : SourceDesc.GetAttributeDescriptions())
 				{
 					// Convert from absolute ID to index into metadata attributes.
-					const int32 SourceAttributeTableIndex = PCGComputeHelpers::GetMetadataAttributeIndexFromAttributeId(SourceAttrDesc.AttributeId);
+					const int32 SourceAttributeTableIndex = PCGComputeHelpers::GetMetadataAttributeIndexFromAttributeId(SourceAttrDesc.GetAttributeId());
 					if (SourceAttributeTableIndex == INDEX_NONE)
 					{
 						// Invalid index means the attribute is a property rather than metadata.
 						continue;
 					}
 
-					const int32 GraphAttributeIndex = AttributeTable.AddAttribute(SourceAttrDesc.AttributeKey.Identifier, SourceAttrDesc.AttributeKey.Type);
+					const int32 GraphAttributeIndex = AttributeTable.AddAttribute(SourceAttrDesc.GetAttributeKey().GetIdentifier(), SourceAttrDesc.GetAttributeKey().GetType());
 
 					if (GraphAttributeIndex != SourceAttributeTableIndex)
 					{
diff --git a/Engine/Plugins/PCG/Source/PCG/Private/Compute/PCGDataDescription.cpp b/Engine/Plugins/PCG/Source/PCG/Private/Compute/PCGDataDescription.cpp
--- a/Engine/Plugins/PCG/Source/PCG/Private/Compute/PCGDataDescription.cpp
+++ b/Engine/Plugins/PCG/Source/PCG/Private/Compute/PCGDataDescription.cpp
@@ -197,13 +197,18 @@ void FPCGKernelAttributeTable::DebugLog() const
 	for (const FPCGKernelAttributeKey& Attribute : AttributeTable)
 	{
 		UE_LOG(LogPCG, Display, TEXT("\tName: %s\t\tID: %d\t\tType: %s"),
-			*Attribute.Identifier.ToString(),
+			*Attribute.GetIdentifier().ToString(),
 			GetAttributeId(Attribute),
-			*PCGKernelAttributeTypeEnum->GetDisplayNameTextByValue(static_cast<int64>(Attribute.Type)).ToString());
+			*PCGKernelAttributeTypeEnum->GetDisplayNameTextByValue(static_cast<int64>(Attribute.GetType())).ToString());
 	}
 }
 #endif
 
+void FPCGKernelAttributeDesc::SetAttributeId(int32 InAttributeId)
+{
+	AttributeId = InAttributeId;
+}
+
 void FPCGKernelAttributeDesc::AddUniqueStringKeys(const TArray<int32>& InOtherStringKeys)
 {
 	for (int32 OtherStringKey : InOtherStringKeys)
@@ -272,7 +277,7 @@ bool FPCGDataDesc::HasElementsMetadataDomainAttributes() const
 {
 	return AttributeDescs.ContainsByPredicate([](const FPCGKernelAttributeDesc& AttributeDesc)
 	{
-		return AttributeDesc.AttributeId >= PCGDataCollectionPackingConstants::NUM_RESERVED_ATTRS && AttributeDesc.AttributeKey.Identifier.MetadataDomain != PCGMetadataDomainID::Data;
+		return AttributeDesc.GetAttributeId() >= PCGDataCollectionPackingConstants::NUM_RESERVED_ATTRS && AttributeDesc.GetAttributeKey().GetIdentifier().MetadataDomain != PCGMetadataDomainID::Data;
 	});
 }
 
@@ -280,7 +285,7 @@ bool FPCGDataDesc::ContainsAttribute(FPCGAttributeIdentifier InAttributeIdentifi
 {
 	return AttributeDescs.ContainsByPredicate([InAttributeIdentifier](const FPCGKernelAttributeDesc& AttributeDesc)
 	{
-		return AttributeDesc.AttributeKey.Identifier == InAttributeIdentifier;
+		return AttributeDesc.GetAttributeKey().GetIdentifier() == InAttributeIdentifier;
 	});
 }
 
@@ -288,13 +293,16 @@ bool FPCGDataDesc::ContainsAttribute(FPCGAttributeIdentifier InAttributeIdentifi
 {
 	return AttributeDescs.ContainsByPredicate([InAttributeKey = FPCGKernelAttributeKey(InAttributeIdentifier, InAttributeType)](const FPCGKernelAttributeDesc& AttributeDesc)
 	{
-		return AttributeDesc.AttributeKey == InAttributeKey;
+		return AttributeDesc.GetAttributeKey() == InAttributeKey;
 	});
 }
 
 void FPCGDataDesc::AddAttribute(FPCGKernelAttributeKey InAttribute, const UPCGDataBinding* InBinding, const TArray<int32>* InOptionalUniqueStringKeys)
 {
+	check(InBinding);
+
 	const int32 AttributeId = InBinding->GetAttributeId(InAttribute);
+
 	if (AttributeId == INDEX_NONE)
 	{
 		return;
@@ -307,20 +315,20 @@ void FPCGDataDesc::AddAttribute(FPCGKernelAttributeKey InAttribute, const UPCGDa
 	{
 		FPCGKernelAttributeDesc& ExistingAttributeDesc = AttributeDescs[AttributeIndex];
 
-		if (ExistingAttributeDesc.AttributeKey.Identifier == InAttribute.Identifier)
+		if (ExistingAttributeDesc.GetAttributeKey().GetIdentifier() == InAttribute.GetIdentifier())
 		{
 			AttributeDescs.RemoveAtSwap(AttributeIndex);
 		}
 	}
 
-	AttributeDescs.Emplace(AttributeId, InAttribute.Type, InAttribute.Identifier, InOptionalUniqueStringKeys);
+	AttributeDescs.Emplace(AttributeId, InAttribute.GetType(), InAttribute.GetIdentifier(), InOptionalUniqueStringKeys);
 }
 
 FIntVector4 FPCGDataDesc::GetElementCountForAttribute(const FPCGKernelAttributeDesc& AttributeDesc) const
 {
 	// A buffer that is not fully allocated uses only a single value to represent all elements.
 	// Data domain only has a single element.
-	if (!IsAttributeAllocated(AttributeDesc) || AttributeDesc.AttributeKey.Identifier.MetadataDomain == PCGMetadataDomainID::Data)
+	if (!IsAttributeAllocated(AttributeDesc) || AttributeDesc.GetAttributeKey().GetIdentifier().MetadataDomain == PCGMetadataDomainID::Data)
 	{
 		return FIntVector4(1, 0, 0, 0);
 	}
@@ -373,6 +381,11 @@ void FPCGDataDesc::AddElementCount(FIntVector4 InElementCountToAdd)
 	ElementCount += InElementCountToAdd;
 }
 
+void FPCGDataDesc::ScaleElementCount(int32 InMultiplier)
+{
+	ElementCount *= InMultiplier;
+}
+
 FPCGDataDesc& FPCGDataDesc::CombineElementCount(const FPCGDataDesc& Other, EPCGElementMultiplicity Multiplicity)
 {
 	if (Multiplicity == EPCGElementMultiplicity::Sum)
@@ -511,6 +524,11 @@ bool FPCGDataDesc::IsAttributeAllocated(int32 InAttributeId) const
 	return true;
 }
 
+TArrayView<int32> FPCGDataDesc::GetTagStringKeysMutable()
+{
+	return MakeArrayView(TagStringKeys);
+}
+
 void FPCGDataDesc::InitializeAttributeDescs(const UPCGData* InData, const UPCGDataBinding* InBinding)
 {
 	TRACE_CPUPROFILER_EVENT_SCOPE(FPCGDataDesc::InitializeAttributeDescs);
@@ -652,13 +670,6 @@ uint32 FPCGDataCollectionDesc::ComputeTotalElementCount() const
 	return ProcessingElementCount;
 }
 
-uint32 FPCGDataCollectionDesc::ComputeTotalElementCount(FPCGDataTypeIdentifier InDataType) const
-{
-	// InDataType not used anymore
-	// @todo_pcg Remove in future pass.
-	return ComputeTotalElementCount();
-}
-
 bool FPCGDataCollectionDesc::GetAttributeDesc(FPCGAttributeIdentifier InAttributeIdentifier, FPCGKernelAttributeDesc& OutAttributeDesc, bool& bOutConflictingTypesFound, bool& bOutPresentOnAllData) const
 {
 	// Will be set to the type of the first attribute across all the data which has a matching name, and then is used to detect if subsequent
@@ -672,9 +683,9 @@ bool FPCGDataCollectionDesc::GetAttributeDesc(FPCGAttributeIdentifier InAttribut
 	{
 		bool bFoundOnData = false;
 
-		for (const FPCGKernelAttributeDesc& AttributeDesc : DataDesc.AttributeDescs)
+		for (const FPCGKernelAttributeDesc& AttributeDesc : DataDesc.GetAttributeDescriptions())
 		{
-			if (AttributeDesc.AttributeKey.Identifier == InAttributeIdentifier && ensure(AttributeDesc.AttributeKey.Type != EPCGKernelAttributeType::Invalid))
+			if (AttributeDesc.GetAttributeKey().GetIdentifier() == InAttributeIdentifier && ensure(AttributeDesc.GetAttributeKey().GetType() != EPCGKernelAttributeType::Invalid))
 			{
 				bFoundOnData = true;
 
@@ -682,9 +693,9 @@ bool FPCGDataCollectionDesc::GetAttributeDesc(FPCGAttributeIdentifier InAttribut
 				{
 					// Take the first matching attribute.
 					OutAttributeDesc = AttributeDesc;
-					FoundType = AttributeDesc.AttributeKey.Type;
+					FoundType = AttributeDesc.GetAttributeKey().GetType();
 				}
-				else if (FoundType != AttributeDesc.AttributeKey.Type)
+				else if (FoundType != AttributeDesc.GetAttributeKey().GetType())
 				{
 					// Signal conflict found.
 					bOutConflictingTypesFound = true;
@@ -735,9 +746,9 @@ void FPCGDataCollectionDesc::GetUniqueStringKeyValues(int32 InAttributeId, TArra
 {
 	for (const FPCGDataDesc& DataDesc : DataDescs)
 	{
-		for (const FPCGKernelAttributeDesc& AttributeDesc : DataDesc.AttributeDescs)
+		for (const FPCGKernelAttributeDesc& AttributeDesc : DataDesc.GetAttributeDescriptions())
 		{
-			if (AttributeDesc.AttributeId == InAttributeId)
+			if (AttributeDesc.GetAttributeId() == InAttributeId)
 			{
 				for (int32 StringKey : AttributeDesc.GetUniqueStringKeys())
 				{
diff --git a/Engine/Plugins/PCG/Source/PCG/Private/Compute/PCGDataUploadAdaptor.cpp b/Engine/Plugins/PCG/Source/PCG/Private/Compute/PCGDataUploadAdaptor.cpp
--- a/Engine/Plugins/PCG/Source/PCG/Private/Compute/PCGDataUploadAdaptor.cpp
+++ b/Engine/Plugins/PCG/Source/PCG/Private/Compute/PCGDataUploadAdaptor.cpp
@@ -72,7 +72,7 @@ FPCGDataUploadAdaptor::FPCGDataUploadAdaptor(UPCGDataBinding* InDataBinding, con
 		}
 	}
 
-	if (UniqueDataCollection && UniqueDataCollection->GetDescription() && UniqueDataCollection->GetDescription()->DataDescs.Num() != NumDataForThisPin)
+	if (UniqueDataCollection && UniqueDataCollection->GetDescription() && UniqueDataCollection->GetDescription()->GetDataDescriptions().Num() != NumDataForThisPin)
 	{
 		// We're not using all the data items from the buffer, no reuse for now.
 		// TODO: Support indirection of data items.
@@ -86,9 +86,9 @@ FPCGDataUploadAdaptor::FPCGDataUploadAdaptor(UPCGDataBinding* InDataBinding, con
 
 		ensureAlways(ExternalBufferMaxStringKeyValue == INDEX_NONE);
 
-		for (const FPCGDataDesc& DataDesc : UniqueDataCollection->GetDescription()->DataDescs)
+		for (const FPCGDataDesc& DataDesc : UniqueDataCollection->GetDescription()->GetDataDescriptions())
 		{
-			for (const FPCGKernelAttributeDesc& AttrDesc : DataDesc.AttributeDescs)
+			for (const FPCGKernelAttributeDesc& AttrDesc : DataDesc.GetAttributeDescriptions())
 			{
 				for (int32 StringKeyValue : AttrDesc.GetUniqueStringKeys())
 				{
@@ -248,9 +248,9 @@ FRDGBufferSRVRef FPCGDataUploadAdaptor::GetBufferToGraphStringKeySRV(FRDGBuilder
 FRDGBufferSRVRef FPCGDataUploadAdaptor::GetDataElementCountsBufferSRV(FRDGBuilder& InGraphBuilder)
 {
 	TArray<uint32> DataElementCounts;
-	DataElementCounts.Reserve(TargetDataCollectionDesc->DataDescs.Num());
+	DataElementCounts.Reserve(TargetDataCollectionDesc->GetDataDescriptions().Num());
 
-	for (const FPCGDataDesc& DataDesc : TargetDataCollectionDesc->DataDescs)
+	for (const FPCGDataDesc& DataDesc : TargetDataCollectionDesc->GetDataDescriptions())
 	{
 		ensure(DataDesc.GetElementDimension() == EPCGElementDimension::One);
 		DataElementCounts.Add(DataDesc.GetElementCount().X);
diff --git a/Engine/Plugins/PCG/Source/PCG/Private/Compute/Packing/PCGDataCollectionPacking.cpp b/Engine/Plugins/PCG/Source/PCG/Private/Compute/Packing/PCGDataCollectionPacking.cpp
--- a/Engine/Plugins/PCG/Source/PCG/Private/Compute/Packing/PCGDataCollectionPacking.cpp
+++ b/Engine/Plugins/PCG/Source/PCG/Private/Compute/Packing/PCGDataCollectionPacking.cpp
@@ -36,7 +36,7 @@ namespace PCGDataCollectionPackingHelpers
 
 		const PCGMetadataValueKey ValueKey = InAttributeBase->GetValueKey(InEntryKey);
 		const int16 TypeId = InAttributeBase->GetTypeId();
-		const int StrideBytes = PCGDataDescriptionHelpers::GetAttributeTypeStrideBytes(InAttributeDesc.AttributeKey.Type);
+		const int StrideBytes = PCGDataDescriptionHelpers::GetAttributeTypeStrideBytes(InAttributeDesc.GetAttributeKey().GetType());
 
 		switch (TypeId)
 		{
@@ -212,57 +212,62 @@ namespace PCGDataCollectionPackingHelpers
 	FPCGMetadataAttributeBase* CreateAttributeFromAttributeDesc(UPCGMetadata* Metadata, const FPCGKernelAttributeDesc& AttributeDesc)
 	{
 		check(Metadata);
-		FPCGMetadataDomain* MetadataDomain = Metadata->GetMetadataDomain(AttributeDesc.AttributeKey.Identifier.MetadataDomain);
+
+		const FPCGKernelAttributeKey& AttributeKey = AttributeDesc.GetAttributeKey();
+		FPCGMetadataDomain* MetadataDomain = Metadata->GetMetadataDomain(AttributeKey.GetIdentifier().MetadataDomain);
+
 		if (!MetadataDomain)
 		{
 			return nullptr;
 		}
+
+		const FName AttributeName = AttributeKey.GetIdentifier().Name;
 		
-		switch (AttributeDesc.AttributeKey.Type)
+		switch (AttributeKey.GetType())
 		{
 		case EPCGKernelAttributeType::Bool:
 		{
-			return MetadataDomain->FindOrCreateAttribute<bool>(AttributeDesc.AttributeKey.Identifier.Name);
+			return MetadataDomain->FindOrCreateAttribute<bool>(AttributeName);
 		}
 		case EPCGKernelAttributeType::Int:
 		{
-			return MetadataDomain->FindOrCreateAttribute<int>(AttributeDesc.AttributeKey.Identifier.Name);
+			return MetadataDomain->FindOrCreateAttribute<int>(AttributeName);
 		}
 		case EPCGKernelAttributeType::Float:
 		{
-			return MetadataDomain->FindOrCreateAttribute<float>(AttributeDesc.AttributeKey.Identifier.Name);
+			return MetadataDomain->FindOrCreateAttribute<float>(AttributeName);
 		}
 		case EPCGKernelAttributeType::Float2:
 		{
-			return MetadataDomain->FindOrCreateAttribute<FVector2D>(AttributeDesc.AttributeKey.Identifier.Name);
+			return MetadataDomain->FindOrCreateAttribute<FVector2D>(AttributeName);
 		}
 		case EPCGKernelAttributeType::Float3:
 		{
-			return MetadataDomain->FindOrCreateAttribute<FVector>(AttributeDesc.AttributeKey.Identifier.Name);
+			return MetadataDomain->FindOrCreateAttribute<FVector>(AttributeName);
 		}
 		case EPCGKernelAttributeType::Float4:
 		{
-			return MetadataDomain->FindOrCreateAttribute<FVector4>(AttributeDesc.AttributeKey.Identifier.Name);
+			return MetadataDomain->FindOrCreateAttribute<FVector4>(AttributeName);
 		}
 		case EPCGKernelAttributeType::Rotator:
 		{
-			return MetadataDomain->FindOrCreateAttribute<FRotator>(AttributeDesc.AttributeKey.Identifier.Name);
+			return MetadataDomain->FindOrCreateAttribute<FRotator>(AttributeName);
 		}
 		case EPCGKernelAttributeType::Quat:
 		{
-			return MetadataDomain->FindOrCreateAttribute<FQuat>(AttributeDesc.AttributeKey.Identifier.Name);
+			return MetadataDomain->FindOrCreateAttribute<FQuat>(AttributeName);
 		}
 		case EPCGKernelAttributeType::Transform:
 		{
-			return MetadataDomain->FindOrCreateAttribute<FTransform>(AttributeDesc.AttributeKey.Identifier.Name);
+			return MetadataDomain->FindOrCreateAttribute<FTransform>(AttributeName);
 		}
 		case EPCGKernelAttributeType::StringKey:
 		{
-			return MetadataDomain->FindOrCreateAttribute<FString>(AttributeDesc.AttributeKey.Identifier.Name);
+			return MetadataDomain->FindOrCreateAttribute<FString>(AttributeName);
 		}
 		case EPCGKernelAttributeType::Name:
 		{
-			return MetadataDomain->FindOrCreateAttribute<FName>(AttributeDesc.AttributeKey.Identifier.Name);
+			return MetadataDomain->FindOrCreateAttribute<FName>(AttributeName);
 		}
 		default:
 			return nullptr;
@@ -277,10 +282,10 @@ namespace PCGDataCollectionPackingHelpers
 		const float* DataAsFloat = static_cast<const float*>(InPackedData);
 		const int32* DataAsInt = static_cast<const int32*>(InPackedData);
 
-		FPCGAttributePropertyOutputSelector Selector = FPCGAttributePropertySelector::CreateAttributeSelector<FPCGAttributePropertyOutputSelector>(InAttributeDesc.AttributeKey.Identifier.Name);
-		OutData->SetDomainFromDomainID(InAttributeDesc.AttributeKey.Identifier.MetadataDomain, Selector);
+		FPCGAttributePropertyOutputSelector Selector = FPCGAttributePropertySelector::CreateAttributeSelector<FPCGAttributePropertyOutputSelector>(InAttributeDesc.GetAttributeKey().GetIdentifier().Name);
+		OutData->SetDomainFromDomainID(InAttributeDesc.GetAttributeKey().GetIdentifier().MetadataDomain, Selector);
 
-		switch (InAttributeDesc.AttributeKey.Type)
+		switch (InAttributeDesc.GetAttributeKey().GetType())
 		{
 		case EPCGKernelAttributeType::Bool:
 		{
@@ -490,10 +495,10 @@ namespace PCGDataCollectionPackingHelpers
 		check(PCGComputeHelpers::IsTypeAllowedInDataCollection(Type));
 		uint64 DataSizeBytes = PCGDataCollectionPackingConstants::DATA_HEADER_SIZE_BYTES;
 
-		for (const FPCGKernelAttributeDesc& AttributeDesc : InDataDesc.AttributeDescs)
+		for (const FPCGKernelAttributeDesc& AttributeDesc : InDataDesc.GetAttributeDescriptions())
 		{
 			const uint64 NumValues = static_cast<uint64>(InDataDesc.GetElementCountForAttribute(AttributeDesc).X);
-			DataSizeBytes += static_cast<uint64>(PCGDataDescriptionHelpers::GetAttributeTypeStrideBytes(AttributeDesc.AttributeKey.Type)) * NumValues;
+			DataSizeBytes += static_cast<uint64>(PCGDataDescriptionHelpers::GetAttributeTypeStrideBytes(AttributeDesc.GetAttributeKey().GetType())) * NumValues;
 		}
 
 		return DataSizeBytes;
@@ -501,7 +506,7 @@ namespace PCGDataCollectionPackingHelpers
 
 	uint32 ComputePackedHeaderSizeBytes(TSharedPtr<const FPCGDataCollectionDesc> InDataCollectionDesc)
 	{
-		return InDataCollectionDesc ? PCGDataCollectionPackingConstants::DATA_COLLECTION_HEADER_SIZE_BYTES + PCGDataCollectionPackingConstants::DATA_HEADER_SIZE_BYTES * InDataCollectionDesc->DataDescs.Num() : 0;
+		return InDataCollectionDesc ? PCGDataCollectionPackingConstants::DATA_COLLECTION_HEADER_SIZE_BYTES + PCGDataCollectionPackingConstants::DATA_HEADER_SIZE_BYTES * InDataCollectionDesc->GetDataDescriptions().Num() : 0;
 	}
 	
 	uint64 ComputePackedSizeBytes(TSharedPtr<const FPCGDataCollectionDesc> InDataCollectionDesc)
@@ -516,7 +521,7 @@ namespace PCGDataCollectionPackingHelpers
 		bool bRequiresHeader = false;
 		uint64 TotalDataSizeBytes = 0;
 
-		const TArray<FPCGDataDesc>& DataDescs = InDataCollectionDesc->DataDescs;
+		const TConstArrayView<FPCGDataDesc> DataDescs = InDataCollectionDesc->GetDataDescriptions();
 	
 		for (const FPCGDataDesc& DataDesc : DataDescs)
 		{
@@ -542,7 +547,7 @@ namespace PCGDataCollectionPackingHelpers
 			return;
 		}
 
-		const TArray<FPCGDataDesc>& DataDescs = InDataCollectionDesc->DataDescs;
+		const TConstArrayView<FPCGDataDesc> DataDescs = InDataCollectionDesc->GetDataDescriptions();
 	
 		const uint32 HeaderSizeBytes = ComputePackedHeaderSizeBytes(InDataCollectionDesc);
 		const uint32 HeaderSizeUints = HeaderSizeBytes >> 2;
@@ -583,24 +588,23 @@ namespace PCGDataCollectionPackingHelpers
 			}
 	
 			// Data i: attribute count (including intrinsic point properties)
-			OutPackedDataCollectionHeader[WriteAddressUints++] = DataDesc.AttributeDescs.Num();
+			OutPackedDataCollectionHeader[WriteAddressUints++] = DataDesc.GetAttributeDescriptions().Num();
 	
 			// Data i: element count
 			OutPackedDataCollectionHeader[WriteAddressUints++] = DataDesc.GetElementCount().X;
 	
 			const uint32 DataAttributesHeaderStartAddressBytes = WriteAddressUints << 2;
 	
-			for (int32 AttrIndex = 0; AttrIndex < DataDesc.AttributeDescs.Num(); ++AttrIndex)
+			for (const FPCGKernelAttributeDesc& AttributeDesc : DataDesc.GetAttributeDescriptions())
 			{
-				const FPCGKernelAttributeDesc& AttributeDesc = DataDesc.AttributeDescs[AttrIndex];
-				const int AttributeStride = PCGDataDescriptionHelpers::GetAttributeTypeStrideBytes(AttributeDesc.AttributeKey.Type);
+				const int AttributeStride = PCGDataDescriptionHelpers::GetAttributeTypeStrideBytes(AttributeDesc.GetAttributeKey().GetType());
 				const int AttributeElementsCount = DataDesc.GetElementCountForAttribute(AttributeDesc).X;
 	
 				// Scatter from attributes that are present into header which has slots for all possible attributes.
-				WriteAddressUints = (AttributeDesc.AttributeId * PCGDataCollectionPackingConstants::ATTRIBUTE_HEADER_SIZE_BYTES + DataAttributesHeaderStartAddressBytes) >> 2;
+				WriteAddressUints = (AttributeDesc.GetAttributeId() * PCGDataCollectionPackingConstants::ATTRIBUTE_HEADER_SIZE_BYTES + DataAttributesHeaderStartAddressBytes) >> 2;
 	
 				// Data i element j: packed attribute info and data address.
-				const uint32 PackedAttributeInfo = (DataDesc.IsAttributeAllocated(AttributeDesc) ? PCGDataCollectionPackingConstants::AttributeAllocatedMask : 0) + (AttributeDesc.AttributeId << 8) + AttributeStride;
+				const uint32 PackedAttributeInfo = (DataDesc.IsAttributeAllocated(AttributeDesc) ? PCGDataCollectionPackingConstants::AttributeAllocatedMask : 0) + (AttributeDesc.GetAttributeId() << 8) + AttributeStride;
 				OutPackedDataCollectionHeader[WriteAddressUints++] = PackedAttributeInfo;
 				OutPackedDataCollectionHeader[WriteAddressUints++] = DataStartAddressBytes;
 	
@@ -624,7 +628,7 @@ namespace PCGDataCollectionPackingHelpers
 			return;
 		}
 
-		const TArray<FPCGDataDesc>& DataDescs = InDataCollectionDesc->DataDescs;
+		const TConstArrayView<FPCGDataDesc> DataDescs = InDataCollectionDesc->GetDataDescriptions();
 		const TArray<FPCGTaggedData> InputData = InDataCollection.GetInputsByPin(InPin);
 	
 		for (const FPCGTaggedData& TaggedData : InputData)
@@ -691,12 +695,12 @@ namespace PCGDataCollectionPackingHelpers
 					continue;
 				}
 	
-				for (const FPCGKernelAttributeDesc& AttributeDesc : DataDesc.AttributeDescs)
+				for (const FPCGKernelAttributeDesc& AttributeDesc : DataDesc.GetAttributeDescriptions())
 				{
-					const uint32 AttributeId = AttributeDesc.AttributeId;
+					const uint32 AttributeId = AttributeDesc.GetAttributeId();
 	
-					const FPCGMetadataDomain* MetadataDomain = Metadata->GetConstMetadataDomain(AttributeDesc.AttributeKey.Identifier.MetadataDomain);
-					const FPCGMetadataAttributeBase* AttributeBase = (AttributeId >= PCGDataCollectionPackingConstants::NUM_RESERVED_ATTRS && MetadataDomain) ? MetadataDomain->GetConstAttribute(AttributeDesc.AttributeKey.Identifier.Name) : nullptr;
+					const FPCGMetadataDomain* MetadataDomain = Metadata->GetConstMetadataDomain(AttributeDesc.GetAttributeKey().GetIdentifier().MetadataDomain);
+					const FPCGMetadataAttributeBase* AttributeBase = (AttributeId >= PCGDataCollectionPackingConstants::NUM_RESERVED_ATTRS && MetadataDomain) ? MetadataDomain->GetConstAttribute(AttributeDesc.GetAttributeKey().GetIdentifier().Name) : nullptr;
 	
 					uint32 AddressUints = GetElementDataStartAddressUints(OutPackedDataCollection.GetData(), DataIndex, AttributeId);
 	
@@ -847,16 +851,16 @@ namespace PCGDataCollectionPackingHelpers
 			}
 			else if (const UPCGParamData* ParamData = Cast<UPCGParamData>(InputData[DataIndex].Data))
 			{
-				for (const FPCGKernelAttributeDesc& AttributeDesc : DataDesc.AttributeDescs)
+				for (const FPCGKernelAttributeDesc& AttributeDesc : DataDesc.GetAttributeDescriptions())
 				{
-					const FPCGMetadataDomain* MetadataDomain = Metadata->GetConstMetadataDomain(AttributeDesc.AttributeKey.Identifier.MetadataDomain);
-					const FPCGMetadataAttributeBase* AttributeBase = MetadataDomain ? MetadataDomain->GetConstAttribute(AttributeDesc.AttributeKey.Identifier.Name) : nullptr;
+					const FPCGMetadataDomain* MetadataDomain = Metadata->GetConstMetadataDomain(AttributeDesc.GetAttributeKey().GetIdentifier().MetadataDomain);
+					const FPCGMetadataAttributeBase* AttributeBase = MetadataDomain ? MetadataDomain->GetConstAttribute(AttributeDesc.GetAttributeKey().GetIdentifier().Name) : nullptr;
 					if (!AttributeBase)
 					{
 						continue;
 					}
 	
-					uint32 AddressUints = GetElementDataStartAddressUints(OutPackedDataCollection.GetData(), DataIndex, AttributeDesc.AttributeId);
+					uint32 AddressUints = GetElementDataStartAddressUints(OutPackedDataCollection.GetData(), DataIndex, AttributeDesc.GetAttributeId());
 	
 					// Pack attribute. Validate first element only for perf.
 					ensure(PackAttributeHelper(AttributeBase, AttributeDesc, /*InEntryKey*/0, InDataBinding->GetStringTable(), OutPackedDataCollection, AddressUints));
@@ -885,7 +889,7 @@ namespace PCGDataCollectionPackingHelpers
 			return EPCGUnpackDataCollectionResult::NoData;
 		}
 
-		const TArray<FPCGDataDesc>& DataDescs = InDataCollectionDesc->DataDescs;
+		const TConstArrayView<FPCGDataDesc> DataDescs = InDataCollectionDesc->GetDataDescriptions();
 	
 		const void* PackedData = InPackedData.GetData();
 		const float* DataAsFloat = static_cast<const float*>(PackedData);
@@ -922,7 +926,7 @@ namespace PCGDataCollectionPackingHelpers
 				NumElements = FMath::Min(NumElements, static_cast<uint32>(DataDesc.GetElementCount().X));
 			}
 	
-			const TArray<FPCGKernelAttributeDesc>& AttributeDescs = DataDesc.AttributeDescs;
+			const TConstArrayView<FPCGKernelAttributeDesc> AttributeDescs = DataDesc.GetAttributeDescriptions();
 			check(NumAttributes == AttributeDescs.Num());
 			check(AttributeDescs.Num() <= PCGDataCollectionPackingConstants::MAX_NUM_ATTRS);
 	
@@ -972,7 +976,7 @@ namespace PCGDataCollectionPackingHelpers
 				OutTaggedData.Data = OutPointData;
 				OutTaggedData.Pin = InPin;
 	
-				for (const int32 TagStringKey : DataDesc.TagStringKeys)
+				for (const int32 TagStringKey : DataDesc.GetTagStringKeys())
 				{
 					if (ensure(InStringTable.IsValidIndex(TagStringKey)))
 					{
@@ -982,9 +986,9 @@ namespace PCGDataCollectionPackingHelpers
 	
 				for (const FPCGKernelAttributeDesc& AttributeDesc : AttributeDescs)
 				{
-					if (AttributeDesc.AttributeId >= PCGDataCollectionPackingConstants::NUM_RESERVED_ATTRS)
+					if (AttributeDesc.GetAttributeId() >= PCGDataCollectionPackingConstants::NUM_RESERVED_ATTRS)
 					{
-						MetadataAttributes[AttributeDesc.AttributeId] = CreateAttributeFromAttributeDesc(Metadata, AttributeDesc);
+						MetadataAttributes[AttributeDesc.GetAttributeId()] = CreateAttributeFromAttributeDesc(Metadata, AttributeDesc);
 					}
 				}
 	
@@ -1000,8 +1004,8 @@ namespace PCGDataCollectionPackingHelpers
 				// Loop over attributes.
 				for (const FPCGKernelAttributeDesc& AttributeDesc : AttributeDescs)
 				{
-					const uint32 AttributeId = AttributeDesc.AttributeId;
-					const uint32 AttributeStrideUints = PCGDataDescriptionHelpers::GetAttributeTypeStrideBytes(AttributeDesc.AttributeKey.Type) >> 2;
+					const uint32 AttributeId = AttributeDesc.GetAttributeId();
+					const uint32 AttributeStrideUints = PCGDataDescriptionHelpers::GetAttributeTypeStrideBytes(AttributeDesc.GetAttributeKey().GetType()) >> 2;
 					const uint32 AddressUints = GetElementDataStartAddressUints(DataAsUint, DataIndex, AttributeId);
 	
 					if (AttributeId < PCGDataCollectionPackingConstants::NUM_RESERVED_ATTRS)
@@ -1234,7 +1238,7 @@ namespace PCGDataCollectionPackingHelpers
 					{
 						TRACE_CPUPROFILER_EVENT_SCOPE(UnpackAttribute);
 	
-						if (FPCGMetadataAttributeBase* AttributeBase = MetadataAttributes[AttributeDesc.AttributeId])
+						if (FPCGMetadataAttributeBase* AttributeBase = MetadataAttributes[AttributeDesc.GetAttributeId()])
 						{
 							ensure(UnpackAttributeHelper(InContext, PackedData, AttributeDesc, InStringTable, AddressUints, DataDesc.GetElementCountForAttribute(AttributeDesc).X, OutPointData));
 						}
@@ -1320,7 +1324,7 @@ namespace PCGDataCollectionPackingHelpers
 	
 				for (const FPCGKernelAttributeDesc& AttributeDesc : AttributeDescs)
 				{
-					MetadataAttributes[AttributeDesc.AttributeId] = CreateAttributeFromAttributeDesc(Metadata, AttributeDesc);
+					MetadataAttributes[AttributeDesc.GetAttributeId()] = CreateAttributeFromAttributeDesc(Metadata, AttributeDesc);
 				}
 	
 				// No work to do if there are no elements to process.
@@ -1334,9 +1338,9 @@ namespace PCGDataCollectionPackingHelpers
 				{
 					TRACE_CPUPROFILER_EVENT_SCOPE(UnpackAttribute);
 	
-					if (FPCGMetadataAttributeBase* AttributeBase = MetadataAttributes[AttributeDesc.AttributeId])
+					if (FPCGMetadataAttributeBase* AttributeBase = MetadataAttributes[AttributeDesc.GetAttributeId()])
 					{
-						const uint32 AddressUints = GetElementDataStartAddressUints(DataAsUint, DataIndex, AttributeDesc.AttributeId);
+						const uint32 AddressUints = GetElementDataStartAddressUints(DataAsUint, DataIndex, AttributeDesc.GetAttributeId());
 						ensure(UnpackAttributeHelper(InContext, PackedData, AttributeDesc, InStringTable, AddressUints, NumElements, OutParamData));
 					}
 				}
diff --git a/Engine/Plugins/PCG/Source/PCG/Private/Compute/Packing/PCGMeshSpawnerPacking.cpp b/Engine/Plugins/PCG/Source/PCG/Private/Compute/Packing/PCGMeshSpawnerPacking.cpp
--- a/Engine/Plugins/PCG/Source/PCG/Private/Compute/Packing/PCGMeshSpawnerPacking.cpp
+++ b/Engine/Plugins/PCG/Source/PCG/Private/Compute/Packing/PCGMeshSpawnerPacking.cpp
@@ -45,9 +45,9 @@ namespace PCGMeshSpawnerPackingHelpers
 				return;
 			}
 
-			const uint32 StrideFloats = PCGDataDescriptionHelpers::GetAttributeTypeStrideBytes(AttributeDesc.AttributeKey.Type) / sizeof(float);
+			const uint32 StrideFloats = PCGDataDescriptionHelpers::GetAttributeTypeStrideBytes(AttributeDesc.GetAttributeKey().GetType()) / sizeof(float);
 
-			OutAttributeIdOffsetStrides.Emplace(static_cast<uint32>(AttributeDesc.AttributeId), OffsetFloats, StrideFloats, /*Unused*/0);
+			OutAttributeIdOffsetStrides.Emplace(static_cast<uint32>(AttributeDesc.GetAttributeId()), OffsetFloats, StrideFloats, /*Unused*/0);
 
 			OffsetFloats += StrideFloats;
 		}
diff --git a/Engine/Plugins/PCG/Source/PCG/Private/Elements/Metadata/PCGMetadataPartitionKernel.cpp b/Engine/Plugins/PCG/Source/PCG/Private/Elements/Metadata/PCGMetadataPartitionKernel.cpp
--- a/Engine/Plugins/PCG/Source/PCG/Private/Elements/Metadata/PCGMetadataPartitionKernel.cpp
+++ b/Engine/Plugins/PCG/Source/PCG/Private/Elements/Metadata/PCGMetadataPartitionKernel.cpp
@@ -51,7 +51,7 @@ bool UPCGMetadataPartitionKernel::IsKernelDataValid(FPCGContext* InContext) cons
 
 		if (!bPresentOnAllData)
 		{
-			if (!InputDesc->DataDescs.IsEmpty())
+			if (!InputDesc->GetDataDescriptions().IsEmpty())
 			{
 				PCG_KERNEL_VALIDATION_ERR(InContext, GetSettings(), FText::Format(
 					LOCTEXT("PartitionAttributeMissing", "Partition attribute was not present on all incoming data."),
@@ -69,7 +69,7 @@ bool UPCGMetadataPartitionKernel::IsKernelDataValid(FPCGContext* InContext) cons
 			return false;
 		}
 
-		if (AttributeDesc.AttributeKey.Type != EPCGKernelAttributeType::StringKey)
+		if (AttributeDesc.GetAttributeKey().GetType() != EPCGKernelAttributeType::StringKey)
 		{
 			PCG_KERNEL_VALIDATION_ERR(Context, GetSettings(), FText::Format(
 				LOCTEXT("PartitionAttributeTypeInvalid", "Attribute '{0}' not usable for partitioning, only attributes of type String Key are currently supported."),
@@ -101,7 +101,7 @@ TSharedPtr<const FPCGDataCollectionDesc> UPCGMetadataPartitionKernel::ComputeOut
 		return nullptr;
 	}
 
-	if (InputPinDesc->DataDescs.IsEmpty())
+	if (InputPinDesc->GetDataDescriptions().IsEmpty())
 	{
 		UE_LOG(LogPCG, Verbose, TEXT("UPCGMetadataPartitionKernel: No analysis data found from count kernel."));
 		return nullptr;
@@ -109,10 +109,10 @@ TSharedPtr<const FPCGDataCollectionDesc> UPCGMetadataPartitionKernel::ComputeOut
 
 	const TArray<int32> AnalysisDataIndices = InBinding->GetPinInputDataIndices(this, PCGMetadataPartitionConstants::ElementCountsPinLabel);
 
-	if (AnalysisDataIndices.Num() != InputPinDesc->DataDescs.Num())
+	if (AnalysisDataIndices.Num() != InputPinDesc->GetDataDescriptions().Num())
 	{
 		UE_LOG(LogPCG, Error, TEXT("UPCGMetadataPartitionKernel: Number of received analysis data items (%d) does not match input data count (%d)."),
-			AnalysisDataIndices.Num(), InputPinDesc->DataDescs.Num());
+			AnalysisDataIndices.Num(), InputPinDesc->GetDataDescriptions().Num());
 		return nullptr;
 	}
 
@@ -169,13 +169,13 @@ TSharedPtr<const FPCGDataCollectionDesc> UPCGMetadataPartitionKernel::ComputeOut
 		{
 			if (AttributeValueToCount[Value] > 0)
 			{
-				FPCGDataDesc& DataDesc = OutDataDesc->DataDescs.Emplace_GetRef(EPCGDataType::Point, AttributeValueToCount[Value]);
-				DataDesc.AttributeDescs.Append(InputPinDesc->DataDescs[DataIndex].AttributeDescs);
-				DataDesc.AllocateProperties(InputPinDesc->DataDescs[DataIndex].GetAllocatedProperties());
+				FPCGDataDesc& DataDesc = OutDataDesc->GetDataDescriptionsMutable().Emplace_GetRef(EPCGDataType::Point, AttributeValueToCount[Value]);
+				DataDesc.GetAttributeDescriptionsMutable().Append(InputPinDesc->GetDataDescriptions()[DataIndex].GetAttributeDescriptions());
+				DataDesc.AllocateProperties(InputPinDesc->GetDataDescriptions()[DataIndex].GetAllocatedProperties());
 
-				for (FPCGKernelAttributeDesc& AttributeDesc : DataDesc.AttributeDescs)
+				for (FPCGKernelAttributeDesc& AttributeDesc : DataDesc.GetAttributeDescriptionsMutable())
 				{
-					if (AttributeDesc.AttributeId == PartitionAttributeId)
+					if (AttributeDesc.GetAttributeId() == PartitionAttributeId)
 					{
 						// By definition each data emitted by the partition node will have a single value for this attribute.
 						AttributeDesc.SetStringKeys(MakeConstArrayView(&Value, 1));
diff --git a/Engine/Plugins/PCG/Source/PCG/Private/Elements/PCGCopyPointsAnalysisKernel.cpp b/Engine/Plugins/PCG/Source/PCG/Private/Elements/PCGCopyPointsAnalysisKernel.cpp
--- a/Engine/Plugins/PCG/Source/PCG/Private/Elements/PCGCopyPointsAnalysisKernel.cpp
+++ b/Engine/Plugins/PCG/Source/PCG/Private/Elements/PCGCopyPointsAnalysisKernel.cpp
@@ -68,15 +68,15 @@ TSharedPtr<const FPCGDataCollectionDesc> UPCGCopyPointsAnalysisKernel::ComputeOu
 
 	TSharedPtr<FPCGDataCollectionDesc> OutDataDesc = FPCGDataCollectionDesc::MakeShared();
 
-	const uint32 NumSourceDatas = SourcesDesc->DataDescs.Num();
-	const uint32 NumTargetDatas = TargetsDesc->DataDescs.Num();
+	const uint32 NumSourceDatas = SourcesDesc->GetDataDescriptions().Num();
+	const uint32 NumTargetDatas = TargetsDesc->GetDataDescriptions().Num();
 
 	// Output: a single attribute set with a single attribute true/false value per NumSource*NumTarget possible outputs.
 	const bool bCopyEachSourceOnEveryTarget = KernelParams->GetValueBool(GET_MEMBER_NAME_CHECKED(UPCGCopyPointsSettings, bCopyEachSourceOnEveryTarget));
 
 	// Validation should have ensured already that these are valid: (S, T) = (N, N), (N, 1) or (1, N).
 	const int32 ElementCount = bCopyEachSourceOnEveryTarget ? (NumSourceDatas * NumTargetDatas) : FMath::Max(NumSourceDatas, NumTargetDatas);
-	FPCGDataDesc& DataDesc = OutDataDesc->DataDescs.Emplace_GetRef(EPCGDataType::Param, ElementCount);
+	FPCGDataDesc& DataDesc = OutDataDesc->GetDataDescriptionsMutable().Emplace_GetRef(EPCGDataType::Param, ElementCount);
 
 	FPCGKernelAttributeKey SelectedFlagAttributeKey = FPCGKernelAttributeKey(PCGCopyPointsConstants::SelectedFlagAttributeName, EPCGKernelAttributeType::Bool);
 	DataDesc.AddAttribute(SelectedFlagAttributeKey, InBinding);
diff --git a/Engine/Plugins/PCG/Source/PCG/Private/Elements/PCGCopyPointsKernel.cpp b/Engine/Plugins/PCG/Source/PCG/Private/Elements/PCGCopyPointsKernel.cpp
--- a/Engine/Plugins/PCG/Source/PCG/Private/Elements/PCGCopyPointsKernel.cpp
+++ b/Engine/Plugins/PCG/Source/PCG/Private/Elements/PCGCopyPointsKernel.cpp
@@ -117,8 +117,8 @@ TSharedPtr<const FPCGDataCollectionDesc> UPCGCopyPointsKernel::ComputeOutputBind
 	const EPCGCopyPointsInheritanceMode SeedInheritance = KernelParams->GetValueEnum<EPCGCopyPointsInheritanceMode>(GET_MEMBER_NAME_CHECKED(UPCGCopyPointsSettings, SeedInheritance));
 	const EPCGCopyPointsMetadataInheritanceMode AttributeInheritance = KernelParams->GetValueEnum<EPCGCopyPointsMetadataInheritanceMode>(GET_MEMBER_NAME_CHECKED(UPCGCopyPointsSettings, AttributeInheritance));
 	const EPCGCopyPointsTagInheritanceMode TagInheritance = KernelParams->GetValueEnum<EPCGCopyPointsTagInheritanceMode>(GET_MEMBER_NAME_CHECKED(UPCGCopyPointsSettings, TagInheritance));
-	const int32 NumSources = SourcePinDesc->DataDescs.Num();
-	const int32 NumTargets = TargetPinDesc->DataDescs.Num();
+	const int32 NumSources = SourcePinDesc->GetDataDescriptions().Num();
+	const int32 NumTargets = TargetPinDesc->GetDataDescriptions().Num();
 	const int32 NumIterations = bCopyEachSourceOnEveryTarget ? NumSources * NumTargets : FMath::Max(NumSources, NumTargets);
 
 	// Properties that might need to be allocated
@@ -136,51 +136,56 @@ TSharedPtr<const FPCGDataCollectionDesc> UPCGCopyPointsKernel::ComputeOutputBind
 			const int32 SourceIndex = bCopyEachSourceOnEveryTarget ? (I / NumTargets) : FMath::Min(I, NumSources - 1);
 			const int32 TargetIndex = bCopyEachSourceOnEveryTarget ? (I % NumTargets) : FMath::Min(I, NumTargets - 1);
 
-			const FPCGDataDesc& SourceDesc = SourcePinDesc->DataDescs[SourceIndex];
-			const FPCGDataDesc& TargetDesc = TargetPinDesc->DataDescs[TargetIndex];
+			const FPCGDataDesc& SourceDesc = SourcePinDesc->GetDataDescriptions()[SourceIndex];
+			const FPCGDataDesc& TargetDesc = TargetPinDesc->GetDataDescriptions()[TargetIndex];
 
-			FPCGDataDesc& ResultDataDesc = OutDataDesc->DataDescs.Emplace_GetRef(EPCGDataType::Point, SourceDesc.GetElementCount() * TargetDesc.GetElementCount());
+			FPCGDataDesc& ResultDataDesc = OutDataDesc->GetDataDescriptionsMutable().Emplace_GetRef(EPCGDataType::Point, SourceDesc.GetElementCount() * TargetDesc.GetElementCount());
+
+			const TConstArrayView<FPCGKernelAttributeDesc> SourceAttributeDescs = SourceDesc.GetAttributeDescriptions();
+			const TConstArrayView<FPCGKernelAttributeDesc> TargetAttributeDescs = TargetDesc.GetAttributeDescriptions();
+
+			TArray<FPCGKernelAttributeDesc>& ResultAttributeDescs = ResultDataDesc.GetAttributeDescriptionsMutable();
 
 			if (AttributeInheritance == EPCGCopyPointsMetadataInheritanceMode::SourceFirst)
 			{
-				ResultDataDesc.AttributeDescs = SourceDesc.AttributeDescs;
+				ResultAttributeDescs = SourceAttributeDescs;
 
-				for (const FPCGKernelAttributeDesc& AttrDesc : TargetDesc.AttributeDescs)
+				for (const FPCGKernelAttributeDesc& AttrDesc : TargetAttributeDescs)
 				{
-					ResultDataDesc.AttributeDescs.AddUnique(AttrDesc);
+					ResultAttributeDescs.AddUnique(AttrDesc);
 				}
 			}
 			else if (AttributeInheritance == EPCGCopyPointsMetadataInheritanceMode::TargetFirst)
 			{
-				ResultDataDesc.AttributeDescs = TargetDesc.AttributeDescs;
+				ResultAttributeDescs = TargetAttributeDescs;
 
-				for (const FPCGKernelAttributeDesc& AttrDesc : SourceDesc.AttributeDescs)
+				for (const FPCGKernelAttributeDesc& AttrDesc : SourceAttributeDescs)
 				{
-					ResultDataDesc.AttributeDescs.AddUnique(AttrDesc);
+					ResultAttributeDescs.AddUnique(AttrDesc);
 				}
 			}
 			else if (AttributeInheritance == EPCGCopyPointsMetadataInheritanceMode::SourceOnly)
 			{
-				ResultDataDesc.AttributeDescs = SourceDesc.AttributeDescs;
+				ResultAttributeDescs = SourceAttributeDescs;
 			}
 			else if (AttributeInheritance == EPCGCopyPointsMetadataInheritanceMode::TargetOnly)
 			{
-				ResultDataDesc.AttributeDescs = TargetDesc.AttributeDescs;
+				ResultAttributeDescs = TargetAttributeDescs;
 			}
 
 			if (TagInheritance == EPCGCopyPointsTagInheritanceMode::Source || TagInheritance == EPCGCopyPointsTagInheritanceMode::Both)
 			{
-				for (const int32 SourceTagStringKey : SourceDesc.TagStringKeys)
+				for (const int32 SourceTagStringKey : SourceDesc.GetTagStringKeys())
 				{
-					ResultDataDesc.TagStringKeys.AddUnique(SourceTagStringKey);
+					ResultDataDesc.AddUniqueTagStringKey(SourceTagStringKey);
 				}
 			}
 			
 			if (TagInheritance == EPCGCopyPointsTagInheritanceMode::Target || TagInheritance == EPCGCopyPointsTagInheritanceMode::Both)
 			{
-				for (const int32 TargetTagStringKey : TargetDesc.TagStringKeys)
+				for (const int32 TargetTagStringKey : TargetDesc.GetTagStringKeys())
 				{
-					ResultDataDesc.TagStringKeys.AddUnique(TargetTagStringKey);
+					ResultDataDesc.AddUniqueTagStringKey(TargetTagStringKey);
 				}
 			}
 
diff --git a/Engine/Plugins/PCG/Source/PCG/Private/Elements/PCGCopyPointsKernelShared.cpp b/Engine/Plugins/PCG/Source/PCG/Private/Elements/PCGCopyPointsKernelShared.cpp
--- a/Engine/Plugins/PCG/Source/PCG/Private/Elements/PCGCopyPointsKernelShared.cpp
+++ b/Engine/Plugins/PCG/Source/PCG/Private/Elements/PCGCopyPointsKernelShared.cpp
@@ -44,8 +44,8 @@ bool PCGCopyPointsKernel::IsKernelDataValid(const UPCGComputeKernel* InKernel, c
 		const bool bMatchBasedOnAttribute = KernelParams->GetValueBool(GET_MEMBER_NAME_CHECKED(UPCGCopyPointsSettings, bMatchBasedOnAttribute));
 		const FName MatchAttribute = KernelParams->GetValueName(GET_MEMBER_NAME_CHECKED(UPCGCopyPointsSettings, MatchAttribute));
 
-		const int32 NumSources = SourcePinDesc->DataDescs.Num();
-		const int32 NumTargets = TargetPinDesc->DataDescs.Num();
+		const int32 NumSources = SourcePinDesc->GetDataDescriptions().Num();
+		const int32 NumTargets = TargetPinDesc->GetDataDescriptions().Num();
 
 		if (!bCopyEachSourceOnEveryTarget && (NumSources != NumTargets && NumSources != 1 && NumTargets != 1 && NumSources != 0 && NumTargets != 0))
 		{
@@ -64,7 +64,7 @@ bool PCGCopyPointsKernel::IsKernelDataValid(const UPCGComputeKernel* InKernel, c
 				// todo_pcg: Can generalize this to any type?
 				constexpr EPCGKernelAttributeType MatchAttributeType = EPCGKernelAttributeType::Int;
 
-				for (const FPCGDataDesc& DataDesc : PinDataDesc->DataDescs)
+				for (const FPCGDataDesc& DataDesc : PinDataDesc->GetDataDescriptions())
 				{
 					if (!DataDesc.ContainsAttribute(MatchAttributeName, MatchAttributeType))
 					{
@@ -77,7 +77,7 @@ bool PCGCopyPointsKernel::IsKernelDataValid(const UPCGComputeKernel* InKernel, c
 				}
 
 				// Valid for execution if we have some data to process.
-				return !PinDataDesc->DataDescs.IsEmpty();
+				return !PinDataDesc->GetDataDescriptions().IsEmpty();
 			};
 
 			if (!ValidateAttributeExists(SourcePinDesc) || !ValidateAttributeExists(TargetPinDesc))
diff --git a/Engine/Plugins/PCG/Source/PCG/Private/Elements/PCGStaticMeshSpawnerKernel.cpp b/Engine/Plugins/PCG/Source/PCG/Private/Elements/PCGStaticMeshSpawnerKernel.cpp
--- a/Engine/Plugins/PCG/Source/PCG/Private/Elements/PCGStaticMeshSpawnerKernel.cpp
+++ b/Engine/Plugins/PCG/Source/PCG/Private/Elements/PCGStaticMeshSpawnerKernel.cpp
@@ -78,13 +78,13 @@ TSharedPtr<const FPCGDataCollectionDesc> UPCGStaticMeshSpawnerKernel::ComputeOut
 			{
 				TArray<int32> StringKeys;
 
-				for (FPCGDataDesc& DataDesc : OutDataDesc->DataDescs)
+				for (FPCGDataDesc& DataDesc : OutDataDesc->GetDataDescriptionsMutable())
 				{
 					StringKeys.Empty();
 
-					for (const FPCGKernelAttributeDesc& AttrDesc : DataDesc.AttributeDescs)
+					for (const FPCGKernelAttributeDesc& AttrDesc : DataDesc.GetAttributeDescriptions())
 					{
-						if (AttrDesc.AttributeKey.Identifier.Name == SelectorByAttribute->AttributeName)
+						if (AttrDesc.GetAttributeKey().GetIdentifier().Name == SelectorByAttribute->AttributeName)
 						{
 							StringKeys = AttrDesc.GetUniqueStringKeys();
 							break;
diff --git a/Engine/Plugins/PCG/Source/PCG/Public/Compute/PCGDataBinding.h b/Engine/Plugins/PCG/Source/PCG/Public/Compute/PCGDataBinding.h
--- a/Engine/Plugins/PCG/Source/PCG/Public/Compute/PCGDataBinding.h
+++ b/Engine/Plugins/PCG/Source/PCG/Public/Compute/PCGDataBinding.h
@@ -108,18 +108,18 @@ public:
 	using FGraphAttributeIndexAndSourceBuffer = TPair</*Graph Metadata Attribute Index*/int32, TRefCountPtr<FRDGPooledBuffer>>;
 
 	/** Get map from source buffer and source buffer attribute index to this-graph attribute index. */
-	const TMap<FSourceBufferAndAttributeIndex, /*Graph Attribute Index*/int32>& GetSourceBufferAttributeToGraphAttributeIndex() const { return SourceBufferAttributeToGraphAttributeIndex; }
+	const TMap<FSourceBufferAndAttributeIndex, /*Graph Attribute Index*/int32>& GetSourceBufferAttributeToGraphAttributeIndex() const;
 
 	/** Get map from this-graph attribute index and upstream source buffer to upstream attribute index. */
-	const TMap<FGraphAttributeIndexAndSourceBuffer, /*Source Buffer Attribute Index*/int32>& GetGraphAttributeToSourceBufferAttributeIndex() const { return GraphAttributeToSourceBufferAttributeIndex; }
+	const TMap<FGraphAttributeIndexAndSourceBuffer, /*Source Buffer Attribute Index*/int32>& GetGraphAttributeToSourceBufferAttributeIndex() const;
 
 	using FSourceBufferAndStringKey = TPair<TRefCountPtr<FRDGPooledBuffer>, /*String Key*/int32>;
 
 	/** Map from string key value in graph to the key value for this string in the source buffer. */
-	const TMap<FSourceBufferAndStringKey, /*Graph String Key*/int32>& GetSourceBufferToGraphStringKey() const { return SourceBufferToGraphStringKey; }
+	const TMap<FSourceBufferAndStringKey, /*Graph String Key*/int32>& GetSourceBufferToGraphStringKey() const;
 
 	/** Map from string key value in source buffer to the key value usable in this graph. */
-	const TMap<FSourceBufferAndStringKey, /*Source Buffer String Key*/int32>& GetGraphToSourceBufferStringKey() const { return GraphToSourceBufferStringKey; }
+	const TMap<FSourceBufferAndStringKey, /*Source Buffer String Key*/int32>& GetGraphToSourceBufferStringKey() const;
 
 	PCG_API IPCGGraphExecutionSource* GetExecutionSource() const;
 
diff --git a/Engine/Plugins/PCG/Source/PCG/Public/Compute/PCGDataDescription.h b/Engine/Plugins/PCG/Source/PCG/Public/Compute/PCGDataDescription.h
--- a/Engine/Plugins/PCG/Source/PCG/Public/Compute/PCGDataDescription.h
+++ b/Engine/Plugins/PCG/Source/PCG/Public/Compute/PCGDataDescription.h
@@ -45,13 +45,24 @@ struct FPCGKernelAttributeKey
 		, Type(InType)
 	{}
 
-	explicit FPCGKernelAttributeKey(const FPCGAttributePropertySelector& InSelector, EPCGKernelAttributeType InType);
+	PCG_API explicit FPCGKernelAttributeKey(const FPCGAttributePropertySelector& InSelector, EPCGKernelAttributeType InType);
+
+	const FPCGAttributeIdentifier& GetIdentifier() const { return Identifier; }
+	EPCGKernelAttributeType GetType() const { return Type; }
 
 	/** To be called everytime the Selector changes, to update the Identifier. Return true if it has changed. */
 	bool UpdateIdentifierFromSelector();
-	
+
+	PCG_API bool IsValid() const;
+
+	PCG_API bool operator==(const FPCGKernelAttributeKey& Other) const;
+
+	PCG_API friend uint32 GetTypeHash(const FPCGKernelAttributeKey& In);
+
+private:
 	void SetSelector(const FPCGAttributePropertySelector& InSelector);
 
+private:
 	/** Cached identifier. Need to be updated if the Selector ever change. */
 	UPROPERTY()
 	FPCGAttributeIdentifier Identifier;
@@ -59,13 +70,6 @@ struct FPCGKernelAttributeKey
 	UPROPERTY(EditAnywhere, Category = "Settings", meta = (DisplayAfter = "Attribute"))
 	EPCGKernelAttributeType Type = EPCGKernelAttributeType::Float;
 
-	bool IsValid() const;
-
-	bool operator==(const FPCGKernelAttributeKey& Other) const;
-
-	friend uint32 GetTypeHash(const FPCGKernelAttributeKey& In);
-
-protected:
 	/** Selector to specify which attribute to create and on which domain. At the moment, only support `@Data` domains or no domain (default one). */
 	UPROPERTY(EditAnywhere, Category = "Settings", meta = (PCG_DiscardPropertySelection, PCG_DiscardExtraSelection))
 	FPCGAttributePropertyOutputNoSourceSelector Name;
@@ -93,7 +97,7 @@ struct FPCGKernelAttributeTable
 	void DebugLog() const;
 #endif
 
-protected:
+private:
 	UPROPERTY()
 	TArray<FPCGKernelAttributeKey> AttributeTable;
 };
@@ -104,15 +108,15 @@ struct FPCGKernelAttributeDesc
 public:
 	FPCGKernelAttributeDesc() = default;
 
-	explicit FPCGKernelAttributeDesc(int32 InIndex, EPCGKernelAttributeType InType, FPCGAttributeIdentifier InIdentifier)
-		: AttributeId(InIndex)
-		, AttributeKey(InIdentifier, InType)
+	explicit FPCGKernelAttributeDesc(int32 InAttributeId, EPCGKernelAttributeType InAttributeType, FPCGAttributeIdentifier InAttributeIdentifier)
+		: AttributeId(InAttributeId)
+		, AttributeKey(InAttributeIdentifier, InAttributeType)
 	{
 	}
 
-	explicit FPCGKernelAttributeDesc(int32 InIndex, EPCGKernelAttributeType InType, FPCGAttributeIdentifier InIdentifier, const TArray<int32>* InUniqueStringKeys)
-		: AttributeId(InIndex)
-		, AttributeKey(InIdentifier, InType)
+	explicit FPCGKernelAttributeDesc(int32 InAttributeId, EPCGKernelAttributeType InAttributeType, FPCGAttributeIdentifier InAttributeIdentifier, const TArray<int32>* InUniqueStringKeys)
+		: AttributeId(InAttributeId)
+		, AttributeKey(InAttributeIdentifier, InAttributeType)
 	{
 		if (InUniqueStringKeys)
 		{
@@ -120,20 +124,22 @@ public:
 		}
 	}
 
-	const TArray<int32>& GetUniqueStringKeys() const { return UniqueStringKeys; }
+	int32 GetAttributeId() const { return AttributeId; }
+	void SetAttributeId(int32 InAttributeId);
 
-	void AddUniqueStringKeys(const TArray<int32>& InOtherStringKeys);
+	const FPCGKernelAttributeKey& GetAttributeKey() const { return AttributeKey; }
 
-	void SetStringKeys(const TConstArrayView<int32>& InStringKeys);
+	PCG_API void AddUniqueStringKeys(const TArray<int32>& InOtherStringKeys);
+	PCG_API void SetStringKeys(const TConstArrayView<int32>& InStringKeys);
+	const TArray<int32>& GetUniqueStringKeys() const { return UniqueStringKeys; }
 
-	bool operator==(const FPCGKernelAttributeDesc& Other) const;
+	PCG_API bool operator==(const FPCGKernelAttributeDesc& Other) const;
 
-public:
+private:
 	int32 AttributeId = INDEX_NONE;
 
 	FPCGKernelAttributeKey AttributeKey{};
 
-private:
 	/* All possible string keys arriving on this attribute (string keys are indices into the string table in the data binding). */
 	TArray<int32> UniqueStringKeys;
 };
@@ -143,57 +149,63 @@ struct FPCGDataDesc
 {
 public:
 	FPCGDataDesc() = default;
-	FPCGDataDesc(FPCGDataTypeIdentifier InType, int32 InElementCount);
-	FPCGDataDesc(FPCGDataTypeIdentifier InType, FIntPoint InElementCount);
+	PCG_API FPCGDataDesc(FPCGDataTypeIdentifier InType, int32 InElementCount);
+	PCG_API FPCGDataDesc(FPCGDataTypeIdentifier InType, FIntPoint InElementCount);
 	FPCGDataDesc(FPCGDataTypeIdentifier InType, FIntVector3 InElementCount);
 	FPCGDataDesc(FPCGDataTypeIdentifier InType, FIntVector4 InElementCount);
 	FPCGDataDesc(const FPCGTaggedData& InTaggedData, const UPCGDataBinding* InBinding);
 
 	bool HasElementsMetadataDomainAttributes() const;
 
-	bool ContainsAttribute(FPCGAttributeIdentifier InAttributeIdentifier) const;
+	PCG_API bool ContainsAttribute(FPCGAttributeIdentifier InAttributeIdentifier) const;
 
-	bool ContainsAttribute(FPCGAttributeIdentifier InAttributeIdentifier, EPCGKernelAttributeType InAttributeType) const;
+	PCG_API bool ContainsAttribute(FPCGAttributeIdentifier InAttributeIdentifier, EPCGKernelAttributeType InAttributeType) const;
 
-	void AddAttribute(FPCGKernelAttributeKey InAttribute, const UPCGDataBinding* InBinding, const TArray<int32>* InOptionalUniqueStringKeys = nullptr);
+	PCG_API void AddAttribute(FPCGKernelAttributeKey InAttribute, const UPCGDataBinding* InBinding, const TArray<int32>* InOptionalUniqueStringKeys = nullptr);
 
 	const FPCGDataTypeIdentifier& GetType() const { return Type; }
 
 	FIntVector4 GetElementCount() const { return ElementCount; }
 	FIntVector4 GetElementCountForAttribute(const FPCGKernelAttributeDesc& AttributeDesc) const;
 	EPCGElementDimension GetElementDimension() const { return ElementDimension; }
-	int32 ComputeTotalElementCount() const;
+	PCG_API int32 ComputeTotalElementCount() const;
 
 	void AddElementCount(int32 InElementCountToAdd);
 	void AddElementCount(FIntPoint InElementCountToAdd);
 	void AddElementCount(FIntVector3 InElementCountToAdd);
 	void AddElementCount(FIntVector4 InElementCountToAdd);
-	void ScaleElementCount(int32 InMultiplier) { ElementCount *= InMultiplier; }
+	void ScaleElementCount(int32 InMultiplier);
 
 	FPCGDataDesc& CombineElementCount(const FPCGDataDesc& Other, EPCGElementMultiplicity Multiplicity);
 	static FPCGDataDesc CombineElementCount(const FPCGDataDesc& A, const FPCGDataDesc& B, EPCGElementMultiplicity Multiplicity);
 
+	TConstArrayView<FPCGKernelAttributeDesc> GetAttributeDescriptions() const { return MakeConstArrayView(AttributeDescs); }
+	TArray<FPCGKernelAttributeDesc>& GetAttributeDescriptionsMutable() { return AttributeDescs; }
+
 	void AllocateProperties(EPCGPointNativeProperties InProperties) { AllocatedPointProperties |= InProperties; }
 	EPCGPointNativeProperties GetAllocatedProperties() const { return AllocatedPointProperties; }
 
-	bool IsAttributeAllocated(const FPCGKernelAttributeDesc& InAttributeDesc) const { return IsAttributeAllocated(InAttributeDesc.AttributeId); }
-	bool IsAttributeAllocated(int32 InAttributeId) const;
+	bool IsAttributeAllocated(const FPCGKernelAttributeDesc& InAttributeDesc) const { return IsAttributeAllocated(InAttributeDesc.GetAttributeId()); }
+	PCG_API bool IsAttributeAllocated(int32 InAttributeId) const;
+
+	TConstArrayView<int32> GetTagStringKeys() const { return MakeConstArrayView(TagStringKeys); }
+	TArrayView<int32> GetTagStringKeysMutable();
+	void SetTagStringKeys(TConstArrayView<int32> InTagStringKeys) { TagStringKeys = InTagStringKeys; }
+	void AddUniqueTagStringKey(int32 InTagStringKey) { TagStringKeys.AddUnique(InTagStringKey); }
 
 private:
 	void InitializeAttributeDescs(const UPCGData* InData, const UPCGDataBinding* InBinding);
 
-public:
-	// @todo_pcg: Encapsulate these members.
-	TArray<FPCGKernelAttributeDesc> AttributeDescs;
-	TArray<int32, TInlineAllocator<4>> TagStringKeys;
-
 private:
 	FPCGDataTypeIdentifier Type = EPCGDataType::Point;
 	FIntVector4 ElementCount = FIntVector4::ZeroValue;
 	EPCGElementDimension ElementDimension = EPCGElementDimension::One;
+	TArray<FPCGKernelAttributeDesc> AttributeDescs;
 
 	/** Properties that should be fully allocated for this data. Only applies to EPCGDataType::Point currently. */
 	EPCGPointNativeProperties AllocatedPointProperties = EPCGPointNativeProperties::None;
+
+	TArray<int32, TInlineAllocator<4>> TagStringKeys;
 };
 
 /** Data description for a data collection (FPCGDataCollection). */
@@ -211,32 +223,33 @@ struct FPCGDataCollectionDesc
 	/** Create an FPCGDataCollectionDesc from another FPCGDataCollectionDesc. */
 	static PCG_API TSharedPtr<FPCGDataCollectionDesc> MakeSharedFrom(const TSharedPtr<const FPCGDataCollectionDesc> InOtherDataDesc);
 
+	TConstArrayView<FPCGDataDesc> GetDataDescriptions() const { return MakeConstArrayView(DataDescs); }
+	TArray<FPCGDataDesc>& GetDataDescriptionsMutable() { return DataDescs; }
+
 	/** Compute total number of processing elements. */
 	PCG_API uint32 ComputeTotalElementCount() const;
 
-	/** Compute total number of processing elements of the given type. */
-	uint32 ComputeTotalElementCount(FPCGDataTypeIdentifier InDataType) const;
-
 	/** Get description of first attribute with matching identifier in input data. Returns true if such an attribute found and also signals whether multiple matching attributes
 	* with conflicting names are present. */
-	bool GetAttributeDesc(FPCGAttributeIdentifier InAttributeIdentifier, FPCGKernelAttributeDesc& OutAttributeDesc, bool& bOutConflictingTypesFound, bool& bOutPresentOnAllData) const;
+	PCG_API bool GetAttributeDesc(FPCGAttributeIdentifier InAttributeIdentifier, FPCGKernelAttributeDesc& OutAttributeDesc, bool& bOutConflictingTypesFound, bool& bOutPresentOnAllData) const;
 
-	bool ContainsAttributeOnAnyData(FPCGAttributeIdentifier InAttributeIdentifier) const;
+	PCG_API bool ContainsAttributeOnAnyData(FPCGAttributeIdentifier InAttributeIdentifier) const;
 
 	/** Makes attribute present on all data. If data has existing attribute with same name then the given type will be applied. */
-	void AddAttributeToAllData(FPCGKernelAttributeKey InAttribute, const UPCGDataBinding* InBinding, const TArray<int32>* InOptionalUniqueStringKeys = nullptr);
-
-	void AllocatePropertiesForAllData(EPCGPointNativeProperties InProperties);
+	PCG_API void AddAttributeToAllData(FPCGKernelAttributeKey InAttribute, const UPCGDataBinding* InBinding, const TArray<int32>* InOptionalUniqueStringKeys = nullptr);
 
-	void GetUniqueStringKeyValues(int32 InAttributeId, TArray<int32>& OutUniqueStringKeys) const;
+	PCG_API void AllocatePropertiesForAllData(EPCGPointNativeProperties InProperties);
 
-	int GetNumStringKeyValues(int32 InAttributeId) const;
+	PCG_API void GetUniqueStringKeyValues(int32 InAttributeId, TArray<int32>& OutUniqueStringKeys) const;
 
-	/** Description of each data in this data collection. */
-	TArray<FPCGDataDesc> DataDescs;
+	PCG_API int GetNumStringKeyValues(int32 InAttributeId) const;
 
 private:
 	FPCGDataCollectionDesc() = default;
+
+private:
+	/** Description of each data in this data collection. */
+	TArray<FPCGDataDesc> DataDescs;
 };
 
 namespace PCGDataDescriptionHelpers
diff --git a/Engine/Plugins/PCG/Source/PCG/Public/PCGSettings.h b/Engine/Plugins/PCG/Source/PCG/Public/PCGSettings.h
--- a/Engine/Plugins/PCG/Source/PCG/Public/PCGSettings.h
+++ b/Engine/Plugins/PCG/Source/PCG/Public/PCGSettings.h
@@ -572,20 +572,14 @@ public:
 	// GPU section
 public:
 
-	/** [EXPERIMENTAL] Whether this node should be executed on the GPU.
-	* Note: This API function is subject to change without deprecation.
-	*/
+	/** Whether this node should be executed on the GPU. */
 	virtual bool ShouldExecuteOnGPU() const { return bExecuteOnGPU; }
 
 #if WITH_EDITOR
-	/** [EXPERIMENTAL] Create the GPU kernels and kernel edges that should execute for this node.
-	* Note: This API function is subject to change without deprecation.
-	*/
+	/** Create the GPU kernels and kernel edges that should execute for this node. */
 	virtual void CreateKernels(FPCGGPUCompilationContext& InOutContext, UObject* InObjectOuter, TArray<UPCGComputeKernel*>& OutKernels, TArray<FPCGKernelEdge>& OutEdges) const {};
 
-	/** [EXPERIMENTAL] Whether to display GPU execution option in node settings UI. 
-	* Note: This API function is subject to change without deprecation.
-	*/
+	/** Whether to display GPU execution option in node settings UI. */
 	UFUNCTION()
 	virtual bool DisplayExecuteOnGPUSetting() const { return false; }
 #endif
@@ -600,34 +594,34 @@ private:
 	UE_API void AddExecutionDependencyPin(TArray<FPCGPinProperties>& OutPins) const;
 
 protected:
-	/** [EXPERIMENTAL] Note: This variable is subject to change without deprecation. */
-	UPROPERTY(EditAnywhere, Category = "GPU", meta = (Tooltip = "Whether this node should be executed on the GPU.", EditCondition = "DisplayExecuteOnGPUSetting()", EditConditionHides, HideEditConditionToggle))
+	/** Whether this node should be executed on the GPU. */
+	UPROPERTY(EditAnywhere, Category = "GPU", meta = (EditCondition = "DisplayExecuteOnGPUSetting()", EditConditionHides, HideEditConditionToggle))
 	bool bExecuteOnGPU = false;
 
 public:
-	/** [EXPERIMENTAL] Note: This variable is subject to change without deprecation. */
-	UPROPERTY(EditAnywhere, Category = "GPU", AdvancedDisplay, meta = (Tooltip = "Dump the cooked HLSL into the log after it is generated.", EditCondition = "bExecuteOnGPU", EditConditionHides))
+	/** Dump the cooked HLSL into the log after it is generated. */
+	UPROPERTY(EditAnywhere, Category = "GPU", AdvancedDisplay, meta = (EditCondition = "bExecuteOnGPU", EditConditionHides))
 	bool bDumpCookedHLSL = false;
 
-	/** [EXPERIMENTAL] Note: This variable is subject to change without deprecation. */
-	UPROPERTY(EditAnywhere, Category = "GPU", AdvancedDisplay, meta = (Tooltip = "Dump the data descriptions of input/output pins to the log.", EditCondition = "bExecuteOnGPU", EditConditionHides))
+	/** Dump the data descriptions of input/output pins to the log. */
+	UPROPERTY(EditAnywhere, Category = "GPU", AdvancedDisplay, meta = (EditCondition = "bExecuteOnGPU", EditConditionHides))
 	bool bDumpDataDescriptions = false;
 
-	/** [EXPERIMENTAL] Note: This variable is subject to change without deprecation. */
-	UPROPERTY(EditAnywhere, Category = "GPU", AdvancedDisplay, meta = (Tooltip = "Enable use of 'WriteDebugValue(uint Index, float Value)' function in your kernel. Allows you to write float values to a buffer for logging on the CPU.", EditCondition = "bExecuteOnGPU", EditConditionHides))
+	/** Enable use of 'WriteDebugValue(uint Index, float Value)' function in your kernel. Allows you to write float values to a buffer for logging on the CPU. */
+	UPROPERTY(EditAnywhere, Category = "GPU", AdvancedDisplay, meta = (EditCondition = "bExecuteOnGPU", EditConditionHides))
 	bool bPrintShaderDebugValues = false;
 
-	/** [EXPERIMENTAL] Note: This variable is subject to change without deprecation. */
-	UPROPERTY(EditAnywhere, Category = "GPU", AdvancedDisplay, meta = (Tooltip = "Size (in number of floats) of the shader debug print buffer.", EditCondition="bExecuteOnGPU && bPrintShaderDebugValues", EditConditionHides))
+	/** Size (in number of floats) of the shader debug print buffer. */
+	UPROPERTY(EditAnywhere, Category = "GPU", AdvancedDisplay, meta = (EditCondition="bExecuteOnGPU && bPrintShaderDebugValues", EditConditionHides))
 	int DebugBufferSize = 16;
 
 #if WITH_EDITORONLY_DATA
-	/** [EXPERIMENTAL] Note: This variable is subject to change without deprecation. */
-	UPROPERTY(EditAnywhere, Transient, Category = "GPU", AdvancedDisplay, meta = (Tooltip = "Will trigger a render capture when this node executes and a debug object is selected in the graph editor. Transient and Editor-only. Render captures must be enabled (e.g. -AttachRenderDoc or -AttachPIX).", EditCondition = "bExecuteOnGPU", EditConditionHides))
+	/** Will trigger a render capture when this node executes and a debug object is selected in the graph editor. Transient and Editor-only. Render captures must be enabled (e.g. -AttachRenderDoc or -AttachPIX). */
+	UPROPERTY(EditAnywhere, Transient, Category = "GPU", AdvancedDisplay, meta = (EditCondition = "bExecuteOnGPU", EditConditionHides))
 	bool bTriggerRenderCapture = false;
 
-	/** [EXPERIMENTAL] Note: This variable is subject to change without deprecation. */
-	UPROPERTY(EditAnywhere, Transient, Category = "GPU", AdvancedDisplay, meta = (Tooltip = "Index of kernel emitted by this node to repeatedly dispatch every frame to enable profiling. Set to -1 to disable profiling. PCG_GPU_KERNEL_PROFILING must be defined in PCG.Build.cs.", EditCondition = "bExecuteOnGPU", EditConditionHides))
+	/** Index of kernel emitted by this node to repeatedly dispatch every frame to enable profiling. Set to -1 to disable profiling. PCG_GPU_KERNEL_PROFILING must be defined in PCG.Build.cs. */
+	UPROPERTY(EditAnywhere, Transient, Category = "GPU", AdvancedDisplay, meta = (EditCondition = "bExecuteOnGPU", EditConditionHides))
 	int ProfileKernelIndex = INDEX_NONE;
 #endif
 
