[PCG] Fix static init order of StaticInLabels/StaticOutLabels being initialized before PCGPinConstants::DefaultInputLabel/DefaultOutputLabel by moving statics inside function call

#rnx

diff --git a/Engine/Plugins/PCG/Source/PCG/Private/PCGInputOutputSettings.cpp b/Engine/Plugins/PCG/Source/PCG/Private/PCGInputOutputSettings.cpp
--- a/Engine/Plugins/PCG/Source/PCG/Private/PCGInputOutputSettings.cpp
+++ b/Engine/Plugins/PCG/Source/PCG/Private/PCGInputOutputSettings.cpp
@@ -24,55 +24,30 @@ namespace PCGInputOutputPrivate
 		FText Tooltip;
 	};
 
-	static TArray<FLabelAndTooltip> StaticInLabels =
+	TArray<FLabelAndTooltip>& GetStaticInLabels()
 	{
+		static TArray<FLabelAndTooltip> StaticInLabels =
 		{
-			PCGPinConstants::DefaultInputLabel, 
-			LOCTEXT("InputOutputInPinTooltip", "Default input pin.")
-		}
-	};
+			{
+				PCGPinConstants::DefaultInputLabel,
+				LOCTEXT("InputOutputInPinTooltip", "Default input pin.")
+			}
+		};
 
-	static TArray<FLabelAndTooltip> StaticAdvancedInLabels =
-	{
-		{
-			PCGInputOutputConstants::DefaultInputLabel,
-			LOCTEXT("InputOutputInputPinTooltip", "DEPRECATED: Please use custom pins to pass inputs into the graph, or use the 'getter' nodes (e.g. Get Volume Data, Get Actor Data, ...) to import data from the level.\n\n"
-			"Takes the output of the Actor pin and if the 'Input Type' setting on the PCG Component is set to Landscape, combines it with the result of the Landscape pin. "
-			"If the Actor data is two dimensional it will be projected onto the landscape, otherwise it will be intersected.")
-		},
-		{
-			PCGInputOutputConstants::DefaultActorLabel,
-			LOCTEXT("InputOutputActorPinTooltip", "DEPRECATED: Please use custom pins to pass inputs into the graph, or use the 'getter' nodes (e.g. Get Volume Data, Get Actor Data, ...) to import data from the level.\n\n"
-			"If this is a partitioned component, then this will be the intersection of the current partition actor bounds with the following. "
-			"If the actor is a Landscape Proxy, then this provide a landscape data. "
-			"Otherwise if the actor is a volume, this will provide a volume shape matching the actor bounds. "
-			"Otherwise if the 'Parse Actor Components' setting is enabled on the PCG Component, this will be all compatible components on the actor (Landscape Splines, Splines, Shapes, Primitives) unioned together. "
-			"Otherwise a single point will be provided at the actor position.")
-		},
-		{
-			PCGInputOutputConstants::DefaultOriginalActorLabel,
-			LOCTEXT("InputOutputOriginalActorPinTooltip", "DEPRECATED: Please use custom pins to pass inputs into the graph, or use the 'getter' nodes (e.g. Get Volume Data, Get Actor Data, ...) to import data from the level.\n\n"
-			"If the actor is a partition actor, this will pull data from the generating PCG actor. Otherwise it will provide the same data as the Actor pin.")
-		},
-		{
-			PCGInputOutputConstants::DefaultLandscapeLabel,
-			LOCTEXT("InputOutputLandscapePinTooltip", "DEPRECATED: Please use custom pins to pass inputs into the graph, or use the 'getter' nodes (e.g. Get Volume Data, Get Actor Data, ...) to import data from the level.\n\n"
-			"Provides the landscape represented by this actor if it is a Landscape Proxy, otherwise it returns any landscapes overlapping this actor in the level.")
-		},
-		{
-			PCGInputOutputConstants::DefaultLandscapeHeightLabel,
-			LOCTEXT("InputOutputLandscapeHeightPinTooltip", "DEPRECATED: Please use custom pins to pass inputs into the graph, or use the 'getter' nodes (e.g. Get Volume Data, Get Actor Data, ...) to import data from the level.\n\n"
-			"Similar to Landscape pin, but only provides height data and not other layers.")
-		}
-	};
+		return StaticInLabels;
+	}
 
-	static TArray<FLabelAndTooltip> StaticOutLabels =
+	TArray<FLabelAndTooltip>& GetStaticOutLabels()
 	{
+		static TArray<FLabelAndTooltip> StaticOutLabels =
 		{
-			PCGPinConstants::DefaultOutputLabel
-		}
-	};
-	static TArray<FLabelAndTooltip> StaticAdvancedOutLabels;
+			{
+				PCGPinConstants::DefaultOutputLabel
+			}
+		};
+
+		return StaticOutLabels;
+	}
 }
 
 bool FPCGInputOutputElement::ExecuteInternal(FPCGContext* Context) const
@@ -136,7 +111,7 @@ void UPCGGraphInputOutputSettings::SetInput(bool bInIsInput)
 		if (Pins.IsEmpty())
 		{
 			const EPCGDataType DefaultPinDataType = bIsInput ? EPCGDataType::Spatial : EPCGDataType::Any;
-			const TArray<PCGInputOutputPrivate::FLabelAndTooltip>& StaticLabels = bIsInput ? PCGInputOutputPrivate::StaticInLabels : PCGInputOutputPrivate::StaticOutLabels;
+			const TArray<PCGInputOutputPrivate::FLabelAndTooltip>& StaticLabels = bIsInput ? PCGInputOutputPrivate::GetStaticInLabels() : PCGInputOutputPrivate::GetStaticOutLabels();
 
 			Algo::Transform(StaticLabels, Pins, [DefaultPinDataType](const PCGInputOutputPrivate::FLabelAndTooltip& InLabelAndTooltip) {
 				FPCGPinProperties Res = FPCGPinProperties(InLabelAndTooltip.Label, DefaultPinDataType, /*bMultiConnections=*/true, /*bMultiData=*/true, InLabelAndTooltip.Tooltip);
@@ -186,7 +161,7 @@ TArray<FPCGPinProperties> UPCGGraphInputOutputSettings::DefaultPinProperties(boo
 	// It is important for serialization that this is not modified, or it could break existing graphs.
 	TArray<FPCGPinProperties> PinProperties;
 	const EPCGDataType DefaultPinDataType = bIsInput ? EPCGDataType::Spatial : EPCGDataType::Any;
-	const TArray<PCGInputOutputPrivate::FLabelAndTooltip>& StaticLabels = bIsInput ? PCGInputOutputPrivate::StaticInLabels : PCGInputOutputPrivate::StaticOutLabels;
+	const TArray<PCGInputOutputPrivate::FLabelAndTooltip>& StaticLabels = bIsInput ? PCGInputOutputPrivate::GetStaticInLabels() : PCGInputOutputPrivate::GetStaticOutLabels();
 
 	Algo::Transform(StaticLabels, PinProperties, [bInvisiblePin, DefaultPinDataType](const PCGInputOutputPrivate::FLabelAndTooltip& InLabelAndTooltip) {
 		FPCGPinProperties Res = FPCGPinProperties(InLabelAndTooltip.Label, DefaultPinDataType, /*bMultiConnections=*/true, /*bMultiData=*/true, InLabelAndTooltip.Tooltip);
