[PCG] Minor usability touchups.

#rnx

diff --git a/Engine/Plugins/PCG/Source/PCGEditor/Private/EditorMode/Tools/Line/PCGDrawSplineTool.cpp b/Engine/Plugins/PCG/Source/PCGEditor/Private/EditorMode/Tools/Line/PCGDrawSplineTool.cpp
--- a/Engine/Plugins/PCG/Source/PCGEditor/Private/EditorMode/Tools/Line/PCGDrawSplineTool.cpp
+++ b/Engine/Plugins/PCG/Source/PCGEditor/Private/EditorMode/Tools/Line/PCGDrawSplineTool.cpp
@@ -421,6 +421,24 @@ FInputRayHit UPCGDrawSplineToolBase::IsHitByClick(const FInputDeviceRay& ClickPo
 	return {};
 }
 
+void UPCGDrawSplineToolBase::ClearSplineOnFirstInteraction(UPCGInteractiveToolSettings_SplineBase* ToolSettings)
+{
+	if(ToolSettings && ToolSettings->DrawMode != PreviousDrawMode)
+	{
+		if (USplineComponent* SplineComponent = ToolSettings->GetWorkingSplineComponent(ToolSettings->DataInstance))
+		{
+			SplineComponent->bSplineHasBeenEdited = true;
+			SplineComponent->ClearSplinePoints(true);
+		}
+
+		// For "continuous" draw modes like FreeDraw, we reset the spline data so so the user can re-drag as needed.
+		if (ToolSettings->DrawMode != EPCGDrawSplineDrawMode::FreeDraw)
+		{
+			PreviousDrawMode = ToolSettings->DrawMode;
+		}
+	}
+}
+
 void UPCGDrawSplineToolBase::OnClicked(const FInputDeviceRay& ClickPosition)
 {
 	using namespace PCGDrawSplineTool;
@@ -428,6 +446,7 @@ void UPCGDrawSplineToolBase::OnClicked(const FInputDeviceRay& ClickPosition)
 	if (TOptional<Scene::FViewHitResult> Result = Scene::ViewportRay(ClickPosition, CachedRayParams))
 	{
 		UPCGInteractiveToolSettings_SplineBase* ToolSettings = GetSplineSettings();
+		ClearSplineOnFirstInteraction(ToolSettings);
 		
 		USplineComponent* WorkingSpline = ToolSettings->GetWorkingSplineComponent(ToolSettings->DataInstance);
 		switch (ToolSettings->DrawMode)
@@ -485,17 +504,10 @@ void UPCGDrawSplineToolBase::OnClickPress(const FInputDeviceRay& ClickPosition)
 	using namespace UE::PCG::EditorMode::Scene;
 	if (TOptional<FViewHitResult> Result = ViewportRay(ClickPosition, CachedRayParams))
 	{
-		UPCGInteractiveToolSettings_SplineBase* ToolSettings = GetSplineSettings();
-
-		// For "continuous" draw modes like FreeDraw, we reset the spline data so so the user can re-drag until happy.
 		// @todo_pcg
 		// When the Spline Tool exposes more controls, such as selecting spline points to move them, resetting the spline etc., we will remove this.
-		if (ToolSettings->DrawMode == EPCGDrawSplineDrawMode::FreeDraw)
-		{
-			USplineComponent* SplineComponent = ToolSettings->GetWorkingSplineComponent(ToolSettings->DataInstance);
-			SplineComponent->bSplineHasBeenEdited = true;
-			SplineComponent->ClearSplinePoints(true);
-		}
+		UPCGInteractiveToolSettings_SplineBase* ToolSettings = GetSplineSettings();
+		ClearSplineOnFirstInteraction(ToolSettings);
 		
 		AddSplinePoint(Result->ImpactPosition, Result->ImpactNormal);
 		
diff --git a/Engine/Plugins/PCG/Source/PCGEditor/Private/EditorMode/Tools/Line/PCGDrawSplineTool.h b/Engine/Plugins/PCG/Source/PCGEditor/Private/EditorMode/Tools/Line/PCGDrawSplineTool.h
--- a/Engine/Plugins/PCG/Source/PCGEditor/Private/EditorMode/Tools/Line/PCGDrawSplineTool.h
+++ b/Engine/Plugins/PCG/Source/PCGEditor/Private/EditorMode/Tools/Line/PCGDrawSplineTool.h
@@ -20,7 +20,8 @@ enum class EPCGDrawSplineDrawMode : uint8
 {
 	TangentDrag UMETA(ToolTip = "Click to place a point and then drag to set its tangent. Clicking without dragging will create sharp corners."),
 	ClickAutoTangent UMETA(ToolTip = "Click and drag new points, with the tangent set automatically."),
-	FreeDraw UMETA(ToolTip = "Drag to place multiple points, with spacing controlled by Min Point Spacing.")
+	FreeDraw UMETA(ToolTip = "Drag to place multiple points, with spacing controlled by Min Point Spacing."),
+	None UMETA(Hidden)
 };
 
 UCLASS(Abstract)
@@ -132,6 +133,8 @@ private:
 	void AddSplinePoint(const FVector& HitLocation, const FVector& HitNormal) const;
 	FVector GetNewPointNormal(const FVector& HitNormal) const;
 
+	void ClearSplineOnFirstInteraction(UPCGInteractiveToolSettings_SplineBase* InSettings);
+
 	UPROPERTY(Instanced)
 	TObjectPtr<USingleClickOrDragInputBehavior> ClickOrDragBehavior = nullptr;
 	
@@ -140,6 +143,8 @@ private:
 	bool bFreeDrawPlacedPreviewPoint = false;
 	int32 FreeDrawStrokeStartIndex = 0;
 
+	EPCGDrawSplineDrawMode PreviousDrawMode = EPCGDrawSplineDrawMode::None;
+
 protected:
 	double CachedMinPointSpacingSquared = 0;
 };
diff --git a/Engine/Plugins/PCG/Source/PCGEditor/Private/EditorMode/Tools/Paint/PCGPaintTool.cpp b/Engine/Plugins/PCG/Source/PCGEditor/Private/EditorMode/Tools/Paint/PCGPaintTool.cpp
--- a/Engine/Plugins/PCG/Source/PCGEditor/Private/EditorMode/Tools/Paint/PCGPaintTool.cpp
+++ b/Engine/Plugins/PCG/Source/PCGEditor/Private/EditorMode/Tools/Paint/PCGPaintTool.cpp
@@ -98,6 +98,10 @@ void UPCGPaintTool::Setup()
 	AddInputBehavior(MouseWheelInputBehavior);
 	
 	Super::Setup();	
+
+	// Finally, set our default values for brush properties
+	BrushProperties->BrushStrength = 1.0;
+	BrushProperties->BrushFalloffAmount = 0.0;
 }
 
 void UPCGPaintTool::Shutdown(EToolShutdownType ShutdownType)
