[PCG] Added graph defaults for editor mode, and made the selection also available to graph instances.

#rnx

diff --git a/Engine/Plugins/PCG/Source/PCGEditor/Private/EditorMode/PCGEdModeSettings.cpp b/Engine/Plugins/PCG/Source/PCGEditor/Private/EditorMode/PCGEdModeSettings.cpp
new file mode 100644
--- /dev/null
+++ b/Engine/Plugins/PCG/Source/PCGEditor/Private/EditorMode/PCGEdModeSettings.cpp
@@ -0,0 +1,28 @@
+// Copyright Epic Games, Inc. All Rights Reserved.
+
+#include "EditorMode/PCGEdModeSettings.h"
+
+#include "PCGGraph.h"
+#include "EditorMode/Tools/Line/PCGDrawSplineTool.h"
+#include "EditorMode/Tools/Paint/PCGPaintTool.h"
+#include "EditorMode/Tools/Volume/PCGVolumeTool.h"
+
+UPCGEditorModeSettings::UPCGEditorModeSettings(const FObjectInitializer& ObjectInitializer)
+	: UDeveloperSettings(ObjectInitializer)
+{
+	FPCGPerInteractiveToolSettingSettings& DrawSpline = InteractiveToolSettings.Emplace_GetRef();
+	DrawSpline.SettingsClass = UPCGInteractiveToolSettings_Spline::StaticClass();
+	DrawSpline.DefaultGraph = TSoftObjectPtr<UPCGGraphInterface>(FSoftObjectPath(TEXT("/PCG/EdMode/DrawSpline/SpacingTool.SpacingTool")));
+
+	FPCGPerInteractiveToolSettingSettings& DrawSplineSurface = InteractiveToolSettings.Emplace_GetRef();
+	DrawSplineSurface.SettingsClass = UPCGInteractiveToolSettings_SplineSurface::StaticClass();
+	DrawSplineSurface.DefaultGraph = TSoftObjectPtr<UPCGGraphInterface>(FSoftObjectPath(TEXT("/PCG/EdMode/DrawSplineSurface/AreaTool.AreaTool")));
+
+	FPCGPerInteractiveToolSettingSettings& PaintTool = InteractiveToolSettings.Emplace_GetRef();
+	PaintTool.SettingsClass = UPCGInteractiveToolSettings_PaintTool::StaticClass();
+	PaintTool.DefaultGraph = TSoftObjectPtr<UPCGGraphInterface>(FSoftObjectPath(TEXT("/PCG/EdMode/Paint/PlaceTool.PlaceTool")));
+
+	FPCGPerInteractiveToolSettingSettings& CreateVolume = InteractiveToolSettings.Emplace_GetRef();
+	CreateVolume.SettingsClass = UPCGInteractiveToolSettings_Volume::StaticClass();
+	CreateVolume.DefaultGraph = TSoftObjectPtr<UPCGGraphInterface>(FSoftObjectPath(TEXT("/PCG/EdMode/Volume/ScatterTool.ScatterTool")));
+}
\ No newline at end of file
diff --git a/Engine/Plugins/PCG/Source/PCGEditor/Private/EditorMode/PCGEdModeSettings.h b/Engine/Plugins/PCG/Source/PCGEditor/Private/EditorMode/PCGEdModeSettings.h
--- a/Engine/Plugins/PCG/Source/PCGEditor/Private/EditorMode/PCGEdModeSettings.h
+++ b/Engine/Plugins/PCG/Source/PCGEditor/Private/EditorMode/PCGEdModeSettings.h
@@ -7,7 +7,7 @@
 #include "PCGEdModeSettings.generated.h"
 
 class UPCGInteractiveToolSettings;
-class UPCGGraph;
+class UPCGGraphInterface;
 
 USTRUCT()
 struct FPCGPerInteractiveToolSettingSettings
@@ -20,7 +20,7 @@ struct FPCGPerInteractiveToolSettingSettings
 	// @todo_pcg GetAssetFilter only works on UClasses due to need for UFUNCTION. Will reenable this via details customization instead.
 	//UPROPERTY(EditAnywhere, Config, Category = "Default", meta=(GetAssetFilter = "FilterBySettingsClass"))
 	UPROPERTY(EditAnywhere, Config, Category = "Default")
-	TSoftObjectPtr<UPCGGraph> DefaultGraph;
+	TSoftObjectPtr<UPCGGraphInterface> DefaultGraph;
 };
 
 /** User preferences related to the PCG Editor Mode. */
@@ -30,9 +30,7 @@ class UPCGEditorModeSettings : public UDeveloperSettings
 	GENERATED_BODY()
 
 public:
-	UPCGEditorModeSettings(const FObjectInitializer& ObjectInitializer)
-		: UDeveloperSettings(ObjectInitializer)
-	{}
+	UPCGEditorModeSettings(const FObjectInitializer& ObjectInitializer);
 
 	virtual FName GetContainerName() const override { return FName("Editor"); }
 	virtual FName GetCategoryName() const override { return FName("Plugins"); }
