[PCG] Allow nested property override in the Static/Skinned Mesh Spawner

diff --git a/Engine/Plugins/PCG/Source/PCG/Private/MeshSelectors/PCGMeshSelectorByAttribute.cpp b/Engine/Plugins/PCG/Source/PCG/Private/MeshSelectors/PCGMeshSelectorByAttribute.cpp
--- a/Engine/Plugins/PCG/Source/PCG/Private/MeshSelectors/PCGMeshSelectorByAttribute.cpp
+++ b/Engine/Plugins/PCG/Source/PCG/Private/MeshSelectors/PCGMeshSelectorByAttribute.cpp
@@ -186,20 +186,14 @@ bool UPCGMeshSelectorByAttribute::SelectMeshInstances(
 			PropertyOverrides.Emplace(MeshSelector, StaticMeshPropertyString);
 		}
 
-		// Validate all the selectors are actual FPCGSoftISMComponentDescriptor properties
-		TArray<FPCGAttributePropertySelector> ValidSelectorOverrides;
+		static FPCGSoftISMComponentDescriptor DummyDescriptor{};
 
-		for (const FPCGObjectPropertyOverrideDescription& PropertyOverride : PropertyOverrides)
-		{
-			if (FPCGSoftISMComponentDescriptor::StaticStruct()->FindPropertyByName(FName(PropertyOverride.PropertyTarget)))
-			{
-				ValidSelectorOverrides.Emplace_GetRef() = FPCGAttributePropertySelector::CreateFromOtherSelector<FPCGAttributePropertySelector>(PropertyOverride.InputSource);
-			}
-			else
-			{
-				PCGE_LOG_C(Error, GraphAndLog, &Context, FText::Format(LOCTEXT("OverriddenPropertyNotFound", "Property '{0}' not a valid property with an ISM Descriptor. It will be ignored."), FText::FromString(PropertyOverride.PropertyTarget)));
-			}
-		}
+		// Dummy descriptor, to initialize the overrides and do some validation
+		FPCGObjectOverrides Overrides(&DummyDescriptor);
+		Overrides.Initialize(PropertyOverrides, &DummyDescriptor, InPointData, &Context);
+
+		// Validate all the selectors are actual FPCGSoftISMComponentDescriptor properties
+		TArray<FPCGAttributePropertySelector> ValidSelectorOverrides = Overrides.GetAllInputSelectors();
 
 		// If there are valid overrides, partition the points on those attributes so that an instance can be created for each
 		if (!ValidSelectorOverrides.IsEmpty())
@@ -220,8 +214,7 @@ bool UPCGMeshSelectorByAttribute::SelectMeshInstances(
 			}
 
 			// Use the Object Override to map the user's input selector and property to the descriptor
-			FPCGObjectOverrides Overrides(&Descriptor);
-			Overrides.Initialize(PropertyOverrides, &Descriptor, InPointData, &Context);
+			Overrides.UpdateTemplateObject(&Descriptor);
 
 			// Since they are already partitioned and identical, we can just use the value on the first point
 			check(!Context.AttributeOverridePartition[I].IsEmpty());
diff --git a/Engine/Plugins/PCG/Source/PCG/Private/MeshSelectors/PCGSkinnedMeshSelector.cpp b/Engine/Plugins/PCG/Source/PCG/Private/MeshSelectors/PCGSkinnedMeshSelector.cpp
--- a/Engine/Plugins/PCG/Source/PCG/Private/MeshSelectors/PCGSkinnedMeshSelector.cpp
+++ b/Engine/Plugins/PCG/Source/PCG/Private/MeshSelectors/PCGSkinnedMeshSelector.cpp
@@ -132,20 +132,14 @@ bool UPCGSkinnedMeshSelector::SelectInstances(
 			PropertyOverrides.Emplace(MeshAttribute, SkinnedAssetPropertyString);
 		}
 
-		// Validate all the selectors are actual FPCGSoftSkinnedMeshComponentDescriptor properties
-		TArray<FPCGAttributePropertySelector> ValidSelectorOverrides;
+		static FPCGSoftSkinnedMeshComponentDescriptor DummyDescriptor{};
+		
+		// Dummy descriptor, to initialize the overrides and do some validation
+		FPCGObjectOverrides Overrides(&DummyDescriptor);
+		Overrides.Initialize(PropertyOverrides, &DummyDescriptor, InPointData, &Context);
 
-		for (const FPCGObjectPropertyOverrideDescription& PropertyOverride : PropertyOverrides)
-		{
-			if (FPCGSoftSkinnedMeshComponentDescriptor::StaticStruct()->FindPropertyByName(FName(PropertyOverride.PropertyTarget)))
-			{
-				ValidSelectorOverrides.Emplace_GetRef() = FPCGAttributePropertySelector::CreateFromOtherSelector<FPCGAttributePropertySelector>(PropertyOverride.InputSource);
-			}
-			else
-			{
-				PCGE_LOG_C(Error, GraphAndLog, &Context, FText::Format(LOCTEXT("OverriddenPropertyNotFound", "Property '{0}' not a valid property with an ABM Descriptor. It will be ignored."), FText::FromString(PropertyOverride.PropertyTarget)));
-			}
-		}
+		// Validate all the selectors are actual FPCGSoftSkinnedMeshComponentDescriptor properties
+		TArray<FPCGAttributePropertySelector> ValidSelectorOverrides = Overrides.GetAllInputSelectors();
 
 		// If there are valid overrides, partition the points on those attributes so that an instance can be created for each
 		if (!ValidSelectorOverrides.IsEmpty())
@@ -166,8 +160,7 @@ bool UPCGSkinnedMeshSelector::SelectInstances(
 			}
 
 			// Use the Object Override to map the user's input selector and property to the descriptor
-			FPCGObjectOverrides Overrides(&Descriptor);
-			Overrides.Initialize(PropertyOverrides, &Descriptor, InPointData, &Context);
+			Overrides.UpdateTemplateObject(&Descriptor);
 
 			// Since they are already partitioned and identical, we can just use the value on the first point
 			check(!Context.AttributeOverridePartition[I].IsEmpty());
diff --git a/Engine/Plugins/PCG/Source/PCG/Private/Metadata/PCGObjectPropertyOverride.cpp b/Engine/Plugins/PCG/Source/PCG/Private/Metadata/PCGObjectPropertyOverride.cpp
--- a/Engine/Plugins/PCG/Source/PCG/Private/Metadata/PCGObjectPropertyOverride.cpp
+++ b/Engine/Plugins/PCG/Source/PCG/Private/Metadata/PCGObjectPropertyOverride.cpp
@@ -171,6 +171,8 @@ void FPCGObjectSingleOverride::Initialize(const FPCGAttributePropertySelector& I
 	};
 
 	PCGMetadataAttribute::CallbackWithRightType(ObjectOverrideOutputAccessor->GetUnderlyingType(), CreateGetterSetter);
+
+	ObjectOverrideInputSelector = InputSelector;
 }
 
 bool FPCGObjectSingleOverride::IsValid() const
diff --git a/Engine/Plugins/PCG/Source/PCG/Public/Metadata/PCGObjectPropertyOverride.h b/Engine/Plugins/PCG/Source/PCG/Public/Metadata/PCGObjectPropertyOverride.h
--- a/Engine/Plugins/PCG/Source/PCG/Public/Metadata/PCGObjectPropertyOverride.h
+++ b/Engine/Plugins/PCG/Source/PCG/Public/Metadata/PCGObjectPropertyOverride.h
@@ -11,6 +11,8 @@
 #include "Metadata/Accessors/IPCGAttributeAccessor.h"
 #include "Metadata/Accessors/PCGAttributeAccessorKeys.h"
 
+#include "Algo/Transform.h"
+
 #include "PCGObjectPropertyOverride.generated.h"
 
 class IPCGAttributeAccessor;
@@ -90,10 +92,13 @@ struct FPCGObjectSingleOverride
 	/** Gathers overrides into an array, if they need loading. */
 	PCG_API void GatherAllOverridesToLoad(TArray<FSoftObjectPath>& OutObjectsToLoad) const;
 
+	const FPCGAttributePropertySelector& GetInputSelector() const { return ObjectOverrideInputSelector; }
+
 private:
 	TUniquePtr<const IPCGAttributeAccessorKeys> InputKeys;
 	TUniquePtr<const IPCGAttributeAccessor> ObjectOverrideInputAccessor;
 	TUniquePtr<IPCGAttributeAccessor> ObjectOverrideOutputAccessor;
+	FPCGAttributePropertySelector ObjectOverrideInputSelector;
 
 #if WITH_EDITOR
 	// To be propagated to the object if requested, allow the object to react to those overrides.
@@ -248,6 +253,15 @@ struct FPCGObjectOverrides
 		}
 	}
 
+	/** Utility function to get all the valid input selectors. Useful to do some partitioning after the override is initialized (and validated all selectors are valid). */
+	TArray<FPCGAttributePropertySelector> GetAllInputSelectors() const
+	{
+		TArray<FPCGAttributePropertySelector> Selectors;
+		Selectors.Reserve(ObjectSingleOverrides.Num());
+		Algo::Transform(ObjectSingleOverrides, Selectors, [](const FPCGObjectSingleOverride& SingleOverride) { return SingleOverride.GetInputSelector(); });
+		return Selectors;
+	}
+
 private:
 	bool bPropagateEditChangeEvent = false;
 	UObject* CachedOutputObject = nullptr;
