[PCG] Fixed a few crashes (crash on deleted points in paint), missing data updates and a gnarly component overwrite.

#rnx

diff --git a/Engine/Plugins/PCG/Source/PCG/Internal/Data/Tool/PCGToolBaseData.h b/Engine/Plugins/PCG/Source/PCG/Internal/Data/Tool/PCGToolBaseData.h
--- a/Engine/Plugins/PCG/Source/PCG/Internal/Data/Tool/PCGToolBaseData.h
+++ b/Engine/Plugins/PCG/Source/PCG/Internal/Data/Tool/PCGToolBaseData.h
@@ -111,7 +111,11 @@ public:
 private:
 	/** Implement your initialization logic here. Called by Initialize. */
 	virtual void InitializeInternal(const FPCGInteractiveToolWorkingDataContext& Context);
+
 protected:
+	/** Cleanup call that will be done at the end of the OnToolApply and OnToolCancel calls. */
+	PCG_API virtual void OnToolShutdown(const FPCGInteractiveToolWorkingDataContext& Context) {}
+
 	/** Each working data struct keeps track of its generated resources. They have to be reverted during OnToolCancel. */
 	UPROPERTY(Transient)
 	FPCGInteractiveToolWorkingDataGeneratedResources GeneratedResources;
diff --git a/Engine/Plugins/PCG/Source/PCG/Internal/Data/Tool/PCGToolPointData.h b/Engine/Plugins/PCG/Source/PCG/Internal/Data/Tool/PCGToolPointData.h
--- a/Engine/Plugins/PCG/Source/PCG/Internal/Data/Tool/PCGToolPointData.h
+++ b/Engine/Plugins/PCG/Source/PCG/Internal/Data/Tool/PCGToolPointData.h
@@ -27,6 +27,8 @@ public:
 	PCG_API virtual void OnToolStart(const FPCGInteractiveToolWorkingDataContext& Context) override;
 	PCG_API virtual void OnToolCancel(const FPCGInteractiveToolWorkingDataContext& Context) override;
 	PCG_API virtual void OnResetToolDataRequested(const FPCGInteractiveToolWorkingDataContext& Context) override;
+protected:
+	PCG_API virtual void OnToolShutdown(const FPCGInteractiveToolWorkingDataContext& Context) override;
 #endif
 
 protected:
diff --git a/Engine/Plugins/PCG/Source/PCG/Internal/Data/Tool/PCGToolSplineData.h b/Engine/Plugins/PCG/Source/PCG/Internal/Data/Tool/PCGToolSplineData.h
--- a/Engine/Plugins/PCG/Source/PCG/Internal/Data/Tool/PCGToolSplineData.h
+++ b/Engine/Plugins/PCG/Source/PCG/Internal/Data/Tool/PCGToolSplineData.h
@@ -28,6 +28,9 @@ public:
 	PCG_API virtual void OnToolStart(const FPCGInteractiveToolWorkingDataContext& Context) override;
 	PCG_API virtual void OnToolCancel(const FPCGInteractiveToolWorkingDataContext& Context) override;
 	PCG_API virtual void OnResetToolDataRequested(const FPCGInteractiveToolWorkingDataContext& Context) override;
+protected:
+	PCG_API virtual void OnToolShutdown(const FPCGInteractiveToolWorkingDataContext& Context) override;
+
 private:
 	virtual void InitializeInternal(const FPCGInteractiveToolWorkingDataContext& Context) override;
 
diff --git a/Engine/Plugins/PCG/Source/PCG/Internal/Data/Tool/PCGToolVolumeData.h b/Engine/Plugins/PCG/Source/PCG/Internal/Data/Tool/PCGToolVolumeData.h
--- a/Engine/Plugins/PCG/Source/PCG/Internal/Data/Tool/PCGToolVolumeData.h
+++ b/Engine/Plugins/PCG/Source/PCG/Internal/Data/Tool/PCGToolVolumeData.h
@@ -30,6 +30,8 @@ public:
 public:
 	PCG_API virtual void OnToolStart(const FPCGInteractiveToolWorkingDataContext& Context) override;
 	PCG_API virtual void OnToolCancel(const FPCGInteractiveToolWorkingDataContext& Context) override;
+protected:
+	PCG_API virtual void OnToolShutdown(const FPCGInteractiveToolWorkingDataContext& Context) override;
 
 private:
 	virtual void InitializeInternal(const FPCGInteractiveToolWorkingDataContext& Context) override;
diff --git a/Engine/Plugins/PCG/Source/PCG/Private/Data/Tool/PCGToolBaseData.cpp b/Engine/Plugins/PCG/Source/PCG/Private/Data/Tool/PCGToolBaseData.cpp
--- a/Engine/Plugins/PCG/Source/PCG/Private/Data/Tool/PCGToolBaseData.cpp
+++ b/Engine/Plugins/PCG/Source/PCG/Private/Data/Tool/PCGToolBaseData.cpp
@@ -47,11 +47,13 @@ void FPCGInteractiveToolWorkingData::OnToolStart(const FPCGInteractiveToolWorkin
 void FPCGInteractiveToolWorkingData::OnToolApply(const FPCGInteractiveToolWorkingDataContext& Context)
 {
 	GeneratedResources.Finalize();
+	OnToolShutdown(Context);
 }
 
 void FPCGInteractiveToolWorkingData::OnToolCancel(const FPCGInteractiveToolWorkingDataContext& Context)
 {
 	GeneratedResources.Cleanup();
+	OnToolShutdown(Context);
 }
 
 void FPCGInteractiveToolWorkingDataGeneratedResources::Finalize()
diff --git a/Engine/Plugins/PCG/Source/PCG/Private/Data/Tool/PCGToolPointData.cpp b/Engine/Plugins/PCG/Source/PCG/Private/Data/Tool/PCGToolPointData.cpp
--- a/Engine/Plugins/PCG/Source/PCG/Private/Data/Tool/PCGToolPointData.cpp
+++ b/Engine/Plugins/PCG/Source/PCG/Private/Data/Tool/PCGToolPointData.cpp
@@ -19,7 +19,10 @@ void FPCGInteractiveToolWorkingData_PointArrayData::InitializeInternal(const FPC
 
 void FPCGInteractiveToolWorkingData_PointArrayData::OnToolStart(const FPCGInteractiveToolWorkingDataContext& Context)
 {
-	OnToolStartPointArray = Cast<UPCGPointArrayData>(StaticDuplicateObject(GeneratedPointData, GetTransientPackage()));
+	if (!OnToolStartPointArray)
+	{
+		OnToolStartPointArray = Cast<UPCGPointArrayData>(StaticDuplicateObject(GeneratedPointData, GetTransientPackage()));
+	}
 
 	Super::OnToolStart(Context);
 }
@@ -30,11 +33,25 @@ void FPCGInteractiveToolWorkingData_PointArrayData::OnToolCancel(const FPCGInter
 	{
 		// We'd expect no relationship between the generated point data on tool cancel, but we'll make sure this is correct here.
 		GeneratedPointData = Cast<UPCGPointArrayData>(StaticDuplicateObject(OnToolStartPointArray, Context.OwningPCGComponent.Get()));
+
+		FPropertyChangedEvent EmptyPropertyChangedEvent(nullptr);
+		FCoreUObjectDelegates::OnObjectPropertyChanged.Broadcast(GeneratedPointData, EmptyPropertyChangedEvent);
 	}
 
 	Super::OnToolCancel(Context);
 }
 
+void FPCGInteractiveToolWorkingData_PointArrayData::OnToolShutdown(const FPCGInteractiveToolWorkingDataContext& Context)
+{
+	if (OnToolStartPointArray)
+	{
+		OnToolStartPointArray->MarkAsGarbage();
+		OnToolStartPointArray = nullptr;
+	}
+
+	Super::OnToolShutdown(Context);
+}
+
 void FPCGInteractiveToolWorkingData_PointArrayData::OnResetToolDataRequested(const FPCGInteractiveToolWorkingDataContext& Context)
 {
 	if (GeneratedPointData)
diff --git a/Engine/Plugins/PCG/Source/PCG/Private/Data/Tool/PCGToolSplineData.cpp b/Engine/Plugins/PCG/Source/PCG/Private/Data/Tool/PCGToolSplineData.cpp
--- a/Engine/Plugins/PCG/Source/PCG/Private/Data/Tool/PCGToolSplineData.cpp
+++ b/Engine/Plugins/PCG/Source/PCG/Private/Data/Tool/PCGToolSplineData.cpp
@@ -27,7 +27,7 @@ void FPCGInteractiveToolWorkingData_Spline::OnToolStart(const FPCGInteractiveToo
 {
 	Super::OnToolStart(Context);
 
-	if (SplineComponent.IsValid())
+	if (SplineComponent.IsValid() && !OnToolStartSplineComponent)
 	{
 		OnToolStartSplineComponent = Cast<USplineComponent>(StaticDuplicateObject(SplineComponent.Get(), GetTransientPackage()));
 	}
@@ -45,8 +45,8 @@ void FPCGInteractiveToolWorkingData_Spline::OnToolCancel(const FPCGInteractiveTo
 
 			SplineComponent->UpdateSpline();
 
-			OnToolStartSplineComponent->MarkAsGarbage();
-			OnToolStartSplineComponent = nullptr;
+			FPropertyChangedEvent EmptyPropertyChangedEvent(nullptr);
+			FCoreUObjectDelegates::OnObjectPropertyChanged.Broadcast(SplineComponent.Get(), EmptyPropertyChangedEvent);
 		}
 	}
 
@@ -54,11 +54,22 @@ void FPCGInteractiveToolWorkingData_Spline::OnToolCancel(const FPCGInteractiveTo
 	Super::OnToolCancel(Context);
 }
 
+void FPCGInteractiveToolWorkingData_Spline::OnToolShutdown(const FPCGInteractiveToolWorkingDataContext& Context)
+{
+	if (OnToolStartSplineComponent)
+	{
+		OnToolStartSplineComponent->MarkAsGarbage();
+		OnToolStartSplineComponent = nullptr;
+	}
+
+	Super::OnToolShutdown(Context);
+}
+
 void FPCGInteractiveToolWorkingData_Spline::OnResetToolDataRequested(const FPCGInteractiveToolWorkingDataContext& Context)
 {
 	if (SplineComponent.IsValid())
 	{
-		SplineComponent->ClearSplinePoints();
+		SplineComponent->ClearSplinePoints(/*bUpdateSpline=*/true);
 
 		FPropertyChangedEvent EmptyPropertyChangedEvent(nullptr);
 		FCoreUObjectDelegates::OnObjectPropertyChanged.Broadcast(SplineComponent.Get(), EmptyPropertyChangedEvent);
@@ -117,8 +128,8 @@ USplineComponent* FPCGInteractiveToolWorkingData_Spline::GenerateMatchingSplineC
 	// Make sure we do not inherit from the actor rotation and scale
 	MatchingSplineComponent->AttachToComponent(Context.OwningActor->GetRootComponent(), FAttachmentTransformRules(EAttachmentRule::KeepRelative, EAttachmentRule::KeepWorld, EAttachmentRule::KeepWorld, false));
 
-	// Transient until Applied. 
-	MatchingSplineComponent->SetFlags(RF_Transient);
+	// Transient until Applied, but always transactional
+	MatchingSplineComponent->SetFlags(RF_Transient | RF_Transactional);
 
 	if (Context.DataInstanceIdentifier.IsNone() == false)
 	{
@@ -215,8 +226,8 @@ USplineComponent* FPCGInteractiveToolWorkingData_SplineSurface::GenerateMatching
 	USplineComponent* MatchingSplineComponent = NewObject<USplineComponent>(Context.OwningActor.Get());
 	MatchingSplineComponent->SetMobility(Context.OwningActor->GetRootComponent()->GetMobility());
 	MatchingSplineComponent->SetupAttachment(Context.OwningActor->GetRootComponent());
-	// Transient until Applied. 
-	MatchingSplineComponent->SetFlags(RF_Transient);
+	// Transient until Applied, but always transactional
+	MatchingSplineComponent->SetFlags(RF_Transient | RF_Transactional);
 
 	if (Context.DataInstanceIdentifier.IsNone() == false)
 	{
diff --git a/Engine/Plugins/PCG/Source/PCG/Private/Data/Tool/PCGToolVolumeData.cpp b/Engine/Plugins/PCG/Source/PCG/Private/Data/Tool/PCGToolVolumeData.cpp
--- a/Engine/Plugins/PCG/Source/PCG/Private/Data/Tool/PCGToolVolumeData.cpp
+++ b/Engine/Plugins/PCG/Source/PCG/Private/Data/Tool/PCGToolVolumeData.cpp
@@ -26,19 +26,21 @@ void FPCGInteractiveToolWorkingData_Volume::OnToolStart(const FPCGInteractiveToo
 {
 	Super::OnToolStart(Context);
 
-	if (BrushComponent.IsValid())
+	bool bUpdateTransform = (!OnToolStartBrushComponent || !OnToolStartBrushBuilder);
+
+	if (BrushComponent.IsValid() && !OnToolStartBrushComponent)
 	{
 		OnToolStartBrushComponent = Cast<UBrushComponent>(StaticDuplicateObject(BrushComponent.Get(), GetTransientPackage()));
 	}
 
-	if (Volume.IsValid())
+	if (Volume.IsValid() && Volume->BrushBuilder && !OnToolStartBrushBuilder)
 	{
-		OnToolStartVolumeTransform = Volume->GetActorTransform();
+		OnToolStartBrushBuilder = DuplicateObject<UBrushBuilder>(Volume->BrushBuilder, GetTransientPackage());
+	}
 
-		if (Volume->BrushBuilder)
-		{
-			OnToolStartBrushBuilder = DuplicateObject<UBrushBuilder>(Volume->BrushBuilder, GetTransientPackage());
-		}
+	if (Volume.IsValid() && bUpdateTransform)
+	{
+		OnToolStartVolumeTransform = Volume->GetActorTransform();
 	}
 }
 
@@ -52,17 +54,35 @@ void FPCGInteractiveToolWorkingData_Volume::OnToolCancel(const FPCGInteractiveTo
 		{
 			// Recreate the brush from the original actor
 			UActorFactory::CreateBrushForVolumeActor(Volume.Get(), OnToolStartBrushBuilder);
-			OnToolStartBrushBuilder->MarkAsGarbage();
-			OnToolStartBrushBuilder = nullptr;
 		}
 
 		// Restore the brush component to its previous state
 		UEngine::CopyPropertiesForUnrelatedObjects(OnToolStartBrushComponent, BrushComponent.Get());
+
+		FPropertyChangedEvent EmptyPropertyChangedEvent(nullptr);
+		FCoreUObjectDelegates::OnObjectPropertyChanged.Broadcast(BrushComponent.Get(), EmptyPropertyChangedEvent);
+	}
+
+	Super::OnToolCancel(Context);
+}
+
+void FPCGInteractiveToolWorkingData_Volume::OnToolShutdown(const FPCGInteractiveToolWorkingDataContext& Context)
+{
+	if (OnToolStartBrushBuilder)
+	{
+		OnToolStartBrushBuilder->MarkAsGarbage();
+		OnToolStartBrushBuilder = nullptr;
+	}
+
+	if (OnToolStartBrushComponent)
+	{
 		OnToolStartBrushComponent->MarkAsGarbage();
 		OnToolStartBrushComponent = nullptr;
 	}
 
-	Super::OnToolCancel(Context);
+	OnToolStartVolumeTransform = FTransform::Identity;
+
+	Super::OnToolShutdown(Context);
 }
 
 UBrushComponent* FPCGInteractiveToolWorkingData_Volume::FindOrCreateBrushComponent(const FPCGInteractiveToolWorkingDataContext& Context)
@@ -126,8 +146,8 @@ UBrushComponent* FPCGInteractiveToolWorkingData_Volume::GenerateMatchingBrushCom
 	MatchingBrushComponent->SetCollisionObjectType(ECC_WorldStatic);
 	MatchingBrushComponent->SetCollisionResponseToAllChannels(ECR_Ignore);
 	MatchingBrushComponent->SetGenerateOverlapEvents(false);
-	// Transient until applied
-	MatchingBrushComponent->SetFlags(RF_Transient);
+	// Transient until applied, but always transactional.
+	MatchingBrushComponent->SetFlags(RF_Transient | RF_Transactional);
 
 	if (!Context.DataInstanceIdentifier.IsNone())
 	{
diff --git a/Engine/Plugins/PCG/Source/PCG/Public/PCGPointArray.h b/Engine/Plugins/PCG/Source/PCG/Public/PCGPointArray.h
--- a/Engine/Plugins/PCG/Source/PCG/Public/PCGPointArray.h
+++ b/Engine/Plugins/PCG/Source/PCG/Public/PCGPointArray.h
@@ -135,18 +135,13 @@ public:
 	inline void RemoveAt(int32 Index)
 	{
 		ensure(Index >= 0 && Index < NumValues);
-		const int32 OldNum = Values.Num();
-		
+
 		if (IsAllocated())
 		{
 			Values.RemoveAt(Index);
-			NumValues -= 1; 
 		}
 
-		if (OldNum > 0 && Values.Num() == 1)
-		{
-			Free();
-		}
+		NumValues -= 1;
 	}
 	
 	/** Move ranges of values(moves them back into the array).Used by the PCGAsync api */
diff --git a/Engine/Plugins/PCG/Source/PCGEditor/Private/EditorMode/Tools/Helpers/PCGEdModeActorHelpers.cpp b/Engine/Plugins/PCG/Source/PCGEditor/Private/EditorMode/Tools/Helpers/PCGEdModeActorHelpers.cpp
--- a/Engine/Plugins/PCG/Source/PCGEditor/Private/EditorMode/Tools/Helpers/PCGEdModeActorHelpers.cpp
+++ b/Engine/Plugins/PCG/Source/PCGEditor/Private/EditorMode/Tools/Helpers/PCGEdModeActorHelpers.cpp
@@ -98,7 +98,7 @@ namespace UE::PCG::EditorMode
 				return nullptr;
 			}
 
-			UPCGComponent* PCGComponent = NewObject<UPCGComponent>(Owner, ComponentName);
+			UPCGComponent* PCGComponent = NewObject<UPCGComponent>(Owner, MakeUniqueObjectName(Owner, UPCGComponent::StaticClass(), ComponentName));
 			PCGComponent->ComponentTags.Add(Tags::ToolGeneratedTag);
 			PCGComponent->RegisterComponent();
 			Owner->AddInstanceComponent(PCGComponent);
diff --git a/Engine/Plugins/PCG/Source/PCGEditor/Private/EditorMode/Tools/PCGInteractiveToolSettings.cpp b/Engine/Plugins/PCG/Source/PCGEditor/Private/EditorMode/Tools/PCGInteractiveToolSettings.cpp
--- a/Engine/Plugins/PCG/Source/PCGEditor/Private/EditorMode/Tools/PCGInteractiveToolSettings.cpp
+++ b/Engine/Plugins/PCG/Source/PCGEditor/Private/EditorMode/Tools/PCGInteractiveToolSettings.cpp
@@ -392,6 +392,12 @@ void UPCGInteractiveToolSettings::DeleteGeneratedResources() const
 	{
 		// We schedule a cleanup and deletion
 		FPCGTaskId CleanupTaskId = GeneratedResources.GeneratedPCGComponent->GetSubsystem()->ScheduleCleanup(GeneratedResources.GeneratedPCGComponent.Get(), true, {});
+		TArray<FPCGTaskId> DestroyTaskDependencies;
+		if (CleanupTaskId != InvalidPCGTaskId)
+		{
+			DestroyTaskDependencies.Add(CleanupTaskId);
+		}
+
 		GeneratedResources.GeneratedPCGComponent->GetSubsystem()->ScheduleGeneric([PCGComponent = GeneratedResources.GeneratedPCGComponent]() -> bool
 		{
 			if (PCGComponent.IsValid())
@@ -402,7 +408,7 @@ void UPCGInteractiveToolSettings::DeleteGeneratedResources() const
 			return true;
 		},
 		GeneratedResources.GeneratedPCGComponent.Get(),
-		{CleanupTaskId});
+		DestroyTaskDependencies);
 		
 		GeneratedResources.GeneratedPCGComponent.Reset();
 	}
@@ -509,6 +515,7 @@ void UPCGInteractiveToolSettings::RestoreGraphInstanceInitialState(UPCGGraphInst
 		if (ensure(OutGraphInstance->ParametersOverrides.Parameters.HasSameLayout(GraphInstanceCache->ParametersOverrides.Parameters)))
 		{
 			OutGraphInstance->ParametersOverrides = GraphInstanceCache->ParametersOverrides;
+			OutGraphInstance->OnGraphParametersChanged(EPCGGraphParameterEvent::UndoRedo, NAME_None);
 		}
 	}
 	else // Otherwise, initial instance state was null, or the final state became null, or the graph was changed. In those cases, revert it all.
@@ -518,11 +525,13 @@ void UPCGInteractiveToolSettings::RestoreGraphInstanceInitialState(UPCGGraphInst
 			if (GraphInstanceCache)
 			{
 				PCGComponent->SetGraph(GraphInstanceCache->GetGraph());
-				SetToolGraph(PCGComponent->GetGraph());
 				if (UPCGGraphInstance* Instance = PCGComponent->GetGraphInstance())
 				{
 					Instance->ParametersOverrides = GraphInstanceCache->ParametersOverrides;
+					Instance->OnGraphParametersChanged(EPCGGraphParameterEvent::UndoRedo, NAME_None);
 				}
+
+				SetToolGraph(PCGComponent->GetGraph());
 			}
 			else // If the original graph was null, reset back to null.
 			{
diff --git a/Engine/Plugins/PCG/Source/PCGEditor/Private/EditorMode/Tools/Paint/PCGPaintTool.cpp b/Engine/Plugins/PCG/Source/PCGEditor/Private/EditorMode/Tools/Paint/PCGPaintTool.cpp
--- a/Engine/Plugins/PCG/Source/PCGEditor/Private/EditorMode/Tools/Paint/PCGPaintTool.cpp
+++ b/Engine/Plugins/PCG/Source/PCGEditor/Private/EditorMode/Tools/Paint/PCGPaintTool.cpp
@@ -62,7 +62,7 @@ void UPCGInteractiveToolSettings_PaintTool::RegisterPropertyWatchers()
 
 	WatchProperty(MinPointSpacing, [this](double InValue)
 	{
-		CachedMinPointSpacingSquared = InValue * InValue;		
+		CachedMinPointSpacingSquared = InValue * InValue;
 	});
 }
 
@@ -299,6 +299,24 @@ double UPCGPaintTool::EstimateMaximumTargetDimension()
 	return 400.f;
 }
 
+void UPCGPaintTool::NotifyDataChanged(UPCGData* InData)
+{
+	if (!InData)
+	{
+		return;
+	}
+
+	FPropertyChangedEvent EmptyPropertyChangedEvent(nullptr);
+	if (UE::PCG::EditorMode::Tool::CVarCommitPropertyChangedEventOnTimer.GetValueOnAnyThread())
+	{
+		ToolSettings->PropertyChangedEventQueue.Add(InData, EmptyPropertyChangedEvent);
+	}
+	else
+	{
+		FCoreUObjectDelegates::OnObjectPropertyChanged.Broadcast(InData, EmptyPropertyChangedEvent);
+	}
+}
+
 void UPCGPaintTool::AddSinglePointFromBrush()
 {
 	UPCGPointArrayData* PointData = ToolSettings->GetTypedWorkingData<FPCGInteractiveToolWorkingData_PointArrayData>(ToolSettings->DataInstance)->GetPointArrayData();
@@ -340,16 +358,7 @@ void UPCGPaintTool::AddSinglePointFromBrush()
 
 	// GetToolData dynamically tracks the pointdata object and will rerun the graph when it changes
 	// To avoid doing this on tick, we queue up the changes
-	
-	FPropertyChangedEvent EmptyPropertyChangedEvent(nullptr);
-	if(UE::PCG::EditorMode::Tool::CVarCommitPropertyChangedEventOnTimer.GetValueOnAnyThread())
-	{
-		ToolSettings->PropertyChangedEventQueue.Add(PointData, EmptyPropertyChangedEvent);
-	}
-	else
-	{
-		FCoreUObjectDelegates::OnObjectPropertyChanged.Broadcast(PointData, EmptyPropertyChangedEvent);
-	}
+	NotifyDataChanged(PointData);
 
 	ToolSettings->LastPointTransform = PointTransform;
 }
@@ -358,7 +367,16 @@ void UPCGPaintTool::TryDeletePointsFromBrush()
 {
 	auto DeletePointsFromDataInstance = [this](FName DataInstanceName)
 	{
-		UPCGPointArrayData* PointData = ToolSettings->GetTypedWorkingData<FPCGInteractiveToolWorkingData_PointArrayData>(DataInstanceName)->GetPointArrayData();
+		check(ToolSettings);
+		const FPCGInteractiveToolWorkingData_PointArrayData* WorkingData = ToolSettings->GetTypedWorkingData<FPCGInteractiveToolWorkingData_PointArrayData>(DataInstanceName);
+		UPCGPointArrayData* PointData = WorkingData ? WorkingData->GetPointArrayData() : nullptr;
+
+		if (!PointData)
+		{
+			return;
+		}
+
+		// @todo_pcg - replace this by an octree search when we have proper iterative change for this kind of thing
 		FSphere Sphere(LastBrushStamp.WorldPosition, CurrentBrushRadius);
 
 		TArray<FTransform> Transforms = PointData->GetTransformsCopy();
@@ -370,15 +388,7 @@ void UPCGPaintTool::TryDeletePointsFromBrush()
 			}
 		}
 
-		FPropertyChangedEvent EmptyPropertyChangedEvent(nullptr);
-		if(UE::PCG::EditorMode::Tool::CVarCommitPropertyChangedEventOnTimer.GetValueOnAnyThread())
-		{
-			ToolSettings->PropertyChangedEventQueue.Add(PointData, EmptyPropertyChangedEvent);
-		}
-		else
-		{
-			FCoreUObjectDelegates::OnObjectPropertyChanged.Broadcast(PointData, EmptyPropertyChangedEvent);
-		}
+		NotifyDataChanged(PointData);
 	};
 	
 	if(ToolSettings->bErasePointsOfSelectedDataInstance)
@@ -392,7 +402,6 @@ void UPCGPaintTool::TryDeletePointsFromBrush()
 			DeletePointsFromDataInstance(DataInstanceName);
 		}
 	}
-	
 }
 
 bool UPCGPaintToolBuilder::CanBuildTool(const FToolBuilderState& SceneState) const
diff --git a/Engine/Plugins/PCG/Source/PCGEditor/Private/EditorMode/Tools/Paint/PCGPaintTool.h b/Engine/Plugins/PCG/Source/PCGEditor/Private/EditorMode/Tools/Paint/PCGPaintTool.h
--- a/Engine/Plugins/PCG/Source/PCGEditor/Private/EditorMode/Tools/Paint/PCGPaintTool.h
+++ b/Engine/Plugins/PCG/Source/PCGEditor/Private/EditorMode/Tools/Paint/PCGPaintTool.h
@@ -92,6 +92,8 @@ protected:
 private:
 	void AddSinglePointFromBrush();
 	void TryDeletePointsFromBrush();
+
+	void NotifyDataChanged(UPCGData* InData);
 	
 private:
 	UPROPERTY(Transient)
