[PCG] Allow GPU data collection reuse for data items that are not at beginning of array.

Fixes a severe performance issue where data not at beginning of input data collection is read back then reuploaded unnecessarily (0.7ms GT time per frame in one level on Switch).

[RN] [PCG] Fixed erroneous readbacks and re-uploads of GPU data items that are not placed first in the compute graph's input data collection.

diff --git a/Engine/Plugins/PCG/Source/PCG/Private/Compute/PCGDataUploadAdaptor.cpp b/Engine/Plugins/PCG/Source/PCG/Private/Compute/PCGDataUploadAdaptor.cpp
--- a/Engine/Plugins/PCG/Source/PCG/Private/Compute/PCGDataUploadAdaptor.cpp
+++ b/Engine/Plugins/PCG/Source/PCG/Private/Compute/PCGDataUploadAdaptor.cpp
@@ -26,6 +26,7 @@ FPCGDataUploadAdaptor::FPCGDataUploadAdaptor(UPCGDataBinding* InDataBinding, con
 	DataBinding = InDataBinding;
 
 	int NumDataForThisPin = 0;
+	int FirstDataForThisPin = INDEX_NONE;
 
 	// Are all our data items coming direct from GPU - and importantly are a complete, single data collection?
 	bool bMultipleBuffersEncountered = false;
@@ -41,6 +42,11 @@ FPCGDataUploadAdaptor::FPCGDataUploadAdaptor(UPCGDataBinding* InDataBinding, con
 
 		++NumDataForThisPin;
 
+		if (FirstDataForThisPin == INDEX_NONE)
+		{
+			FirstDataForThisPin = InputDataIndex;
+		}
+
 		const UPCGProxyForGPUData* DataGPU = Cast<UPCGProxyForGPUData>(DataBinding->GetInputDataCollection().TaggedData[InputDataIndex].Data);
 		if (!DataGPU)
 		{
@@ -63,7 +69,7 @@ FPCGDataUploadAdaptor::FPCGDataUploadAdaptor(UPCGDataBinding* InDataBinding, con
 		{
 			UniqueDataCollection = InputDataCollection;
 		}
-		else if (UniqueDataCollection != InputDataCollection || DataGPU->GetDataIndexInCollection() != InputDataIndex)
+		else if (UniqueDataCollection != InputDataCollection || DataGPU->GetDataIndexInCollection() != (InputDataIndex - FirstDataForThisPin))
 		{
 			// Multiple input buffers or out of order data items, no reuse for now.
 			// TODO: support indirection of data items.
