[PCG] Removed ensure on first point with the tangent align mode and fixed a crash when the tool data instance changes.

#rnx

diff --git a/Engine/Plugins/PCG/Source/PCGEditor/Private/EditorMode/Tools/Line/PCGDrawSplineTool.cpp b/Engine/Plugins/PCG/Source/PCGEditor/Private/EditorMode/Tools/Line/PCGDrawSplineTool.cpp
--- a/Engine/Plugins/PCG/Source/PCGEditor/Private/EditorMode/Tools/Line/PCGDrawSplineTool.cpp
+++ b/Engine/Plugins/PCG/Source/PCGEditor/Private/EditorMode/Tools/Line/PCGDrawSplineTool.cpp
@@ -393,9 +393,8 @@ FVector UPCGDrawSplineToolBase::GetNewPointNormal(const FVector& HitNormal) cons
 			{
 				return WorkingSpline->GetUpVectorAtSplinePoint(NumSplinePoints - 1, ESplineCoordinateSpace::World);
 			}
-			else // NumSplinePoints is zero or negative, which should never happen.
+			else // NumSplinePoints is zero or negative, which can happen, if there are no previous points.
 			{
-				ensure(false);
 				return HitNormal;
 			}
 		}
diff --git a/Engine/Plugins/PCG/Source/PCGEditor/Private/EditorMode/Tools/PCGInteractiveToolSettings.cpp b/Engine/Plugins/PCG/Source/PCGEditor/Private/EditorMode/Tools/PCGInteractiveToolSettings.cpp
--- a/Engine/Plugins/PCG/Source/PCGEditor/Private/EditorMode/Tools/PCGInteractiveToolSettings.cpp
+++ b/Engine/Plugins/PCG/Source/PCGEditor/Private/EditorMode/Tools/PCGInteractiveToolSettings.cpp
@@ -654,7 +654,18 @@ void UPCGInteractiveToolSettings::OnPropertyModified(UObject* Object, FProperty*
 			}
 		}
 	}
-	if(Object == this && Property->GetFName() == GET_MEMBER_NAME_CHECKED(UPCGInteractiveToolSettings, ActorClassToSpawn))
+	else if (Object == this && Property->GetFName() == GET_MEMBER_NAME_CHECKED(UPCGInteractiveToolSettings, DataInstance))
+	{
+		if (InitializeWorkingData(DataInstance))
+		{
+			// We generate again to make sure the initialized working data has a chance to register its tracking keys.
+			if (UPCGComponent* PCGComponent = GetWorkingPCGComponent())
+			{
+				PCGComponent->GenerateLocal(/*bForce=*/true);
+			}
+		}
+	}
+	else if(Object == this && Property->GetFName() == GET_MEMBER_NAME_CHECKED(UPCGInteractiveToolSettings, ActorClassToSpawn))
 	{
 		// @todo_pcg
 		// Delete previously spawned actor and set up the new one. Ideally recycling the existing PCG Component.
diff --git a/Engine/Plugins/PCG/Source/PCGEditor/Private/EditorMode/Tools/Paint/PCGPaintTool.cpp b/Engine/Plugins/PCG/Source/PCGEditor/Private/EditorMode/Tools/Paint/PCGPaintTool.cpp
--- a/Engine/Plugins/PCG/Source/PCGEditor/Private/EditorMode/Tools/Paint/PCGPaintTool.cpp
+++ b/Engine/Plugins/PCG/Source/PCGEditor/Private/EditorMode/Tools/Paint/PCGPaintTool.cpp
@@ -319,7 +319,9 @@ void UPCGPaintTool::NotifyDataChanged(UPCGData* InData)
 
 void UPCGPaintTool::AddSinglePointFromBrush()
 {
-	UPCGPointArrayData* PointData = ToolSettings->GetTypedWorkingData<FPCGInteractiveToolWorkingData_PointArrayData>(ToolSettings->DataInstance)->GetPointArrayData();
+	check(ToolSettings);
+	const FPCGInteractiveToolWorkingData_PointArrayData* WorkingData = ToolSettings->GetTypedWorkingData<FPCGInteractiveToolWorkingData_PointArrayData>(ToolSettings->DataInstance);
+	UPCGPointArrayData* PointData = WorkingData ? WorkingData->GetPointArrayData() : nullptr;
 
 	// If the current layer doesn't exist anymore (from the graph), it's possible the working data has been changed/removed.
 	// @todo_pcg: Preferably this should get caught when the graph changes.
