[PCG] PCGPointArray: Fix bad assignment on allocated value
- This was mostly avoided by freeing the property first but still very incorrect

#rnx

diff --git a/Engine/Plugins/PCG/Source/PCG/Public/PCGPointArray.h b/Engine/Plugins/PCG/Source/PCG/Public/PCGPointArray.h
--- a/Engine/Plugins/PCG/Source/PCG/Public/PCGPointArray.h
+++ b/Engine/Plugins/PCG/Source/PCG/Public/PCGPointArray.h
@@ -115,7 +115,7 @@ public:
 	{
 		for (int32 i = 0; i < Values.Num(); ++i)
 		{
-			Values[i] = Value;
+			Values[i] = InValue;
 		}
 		Value = InValue;
 	}
