[PCG] Add temporary logic to force tick the debug object tree view when it is not in foreground

* We rely on the object tree view to know which Stack to inspect, and the stack only update on tick, which does not happen if the widget is not in foreground.
* To be revisited for 5.8 when we properly remove this core logic from the widget tick

#rnx

diff --git a/Engine/Plugins/PCG/Source/PCGEditor/Private/Widgets/SPCGEditorGraphDebugObjectTree.cpp b/Engine/Plugins/PCG/Source/PCGEditor/Private/Widgets/SPCGEditorGraphDebugObjectTree.cpp
--- a/Engine/Plugins/PCG/Source/PCGEditor/Private/Widgets/SPCGEditorGraphDebugObjectTree.cpp
+++ b/Engine/Plugins/PCG/Source/PCGEditor/Private/Widgets/SPCGEditorGraphDebugObjectTree.cpp
@@ -571,8 +571,13 @@ void SPCGEditorGraphDebugObjectTree::Tick(const FGeometry& AllottedGeometry, con
 {
 	SCompoundWidget::Tick(AllottedGeometry, InCurrentTime, InDeltaTime);
 
+	OnTick(InCurrentTime);
+}
+
+void SPCGEditorGraphDebugObjectTree::OnTick(double InCurrentTime, bool bForceRefresh)
+{
 	bool bRefreshDone = false;
-	if (bNeedsRefresh && InCurrentTime >= NextRefreshTime)
+	if (bNeedsRefresh && (bForceRefresh || InCurrentTime >= NextRefreshTime))
 	{
 		// Updating the tree while the inspected component is generating can be bad as the selection can
 		// be lost. Don't change the tree if we're inspecting something that is generating.
@@ -619,9 +624,15 @@ void SPCGEditorGraphDebugObjectTree::Tick(const FGeometry& AllottedGeometry, con
 	}
 }
 
-void SPCGEditorGraphDebugObjectTree::RequestRefresh()
+void SPCGEditorGraphDebugObjectTree::RequestRefresh(bool bForce)
 {
 	bNeedsRefresh = true;
+
+	if (!IsOpen() && bForce)
+	{
+		// Force a refresh with a current time equal to next refresh time
+		OnTick(NextRefreshTime, /*bForceRefresh=*/true);
+	}
 }
 
 void SPCGEditorGraphDebugObjectTree::SetNodeBeingInspected(const UPCGNode* InPCGNode)
@@ -633,7 +644,7 @@ void SPCGEditorGraphDebugObjectTree::SetNodeBeingInspected(const UPCGNode* InPCG
 		bShouldSelectStackOnNextRefresh = true;
 	}
 
-	RequestRefresh();
+	RequestRefresh(/*bForce=*/true);
 }
 
 FReply SPCGEditorGraphDebugObjectTree::FocusOnDebugObject_OnClicked() const
@@ -791,6 +802,11 @@ bool SPCGEditorGraphDebugObjectTree::GetFirstStackFromSelection(const UPCGNode*
 	return false;
 }
 
+bool SPCGEditorGraphDebugObjectTree::IsOpen() const
+{
+	return PCGEditor.IsValid() && PCGEditor.Pin()->IsPanelCurrentlyForeground(EPCGEditorPanel::DebugObjectTree);
+}
+
 FPCGEditorGraphDebugObjectItemPtr SPCGEditorGraphDebugObjectTree::GetFirstDebugObjectFromSelection() const
 {
 	TRACE_CPUPROFILER_EVENT_SCOPE(SPCGEditorGraphDebugObjectTree::GetFirstDebugObjectFromSelection);
diff --git a/Engine/Plugins/PCG/Source/PCGEditor/Private/Widgets/SPCGEditorGraphDebugObjectTree.h b/Engine/Plugins/PCG/Source/PCGEditor/Private/Widgets/SPCGEditorGraphDebugObjectTree.h
--- a/Engine/Plugins/PCG/Source/PCGEditor/Private/Widgets/SPCGEditorGraphDebugObjectTree.h
+++ b/Engine/Plugins/PCG/Source/PCGEditor/Private/Widgets/SPCGEditorGraphDebugObjectTree.h
@@ -262,7 +262,8 @@ public:
 
 	virtual void Tick(const FGeometry& AllottedGeometry, const double InCurrentTime, const float InDeltaTime) override;
 
-	void RequestRefresh();
+	// Request a refresh of the widget. If bForce is true, will ask for refresh even if the widget is closed.
+	void RequestRefresh(bool bForce = false);
 
 	void SetNodeBeingInspected(const UPCGNode* InPCGNode);
 
@@ -273,6 +274,10 @@ public:
 	bool GetFirstStackFromSelection(const UPCGNode* InNode, const UPCGGraph* InGraph, FPCGStack& OutStack) const;
 
 private:
+	void OnTick(double InCurrentTime, bool bForceRefresh = false);
+	
+	bool IsOpen() const;
+	
 	FReply FocusOnDebugObject_OnClicked() const;
 	bool IsFocusOnDebugObjectButtonEnabled() const;
 
