[PCG] Prevented some cases that would yield graphs with the incorrect actor class. Added label option in the tool graphs. Fixed crash when changing tool data instances.

#rnx

diff --git a/Engine/Plugins/PCG/Source/PCG/Public/Data/PCGToolData.h b/Engine/Plugins/PCG/Source/PCG/Public/Data/PCGToolData.h
--- a/Engine/Plugins/PCG/Source/PCG/Public/Data/PCGToolData.h
+++ b/Engine/Plugins/PCG/Source/PCG/Public/Data/PCGToolData.h
@@ -48,6 +48,9 @@ struct FPCGGraphToolData
 
 	UPROPERTY(EditAnywhere, Category = Tool)
 	TSubclassOf<AActor> InitialActorClassToSpawn;
+
+	UPROPERTY(EditAnywhere, Category = Tool)
+	FName NewActorLabel;
 };
 
 USTRUCT(BlueprintType)
@@ -72,7 +75,7 @@ struct FPCGInteractiveToolDataContainer
 
 		if (const TInstancedStruct<FPCGInteractiveToolWorkingBaseData>* Data = FindToolDataStruct(WorkingDataIdentifier))
 		{
-			if (Data->GetScriptStruct()->IsChildOf(T::StaticStruct()) && Data->IsValid())
+			if (Data->IsValid() && Data->GetScriptStruct()->IsChildOf(T::StaticStruct()))
 			{
 				return Data->GetPtr<T>();
 			}
@@ -88,7 +91,7 @@ struct FPCGInteractiveToolDataContainer
 
 		if (TInstancedStruct<FPCGInteractiveToolWorkingBaseData>* Data = FindMutableToolDataStruct(WorkingDataIdentifier))
 		{
-			if (Data->GetScriptStruct()->IsChildOf(T::StaticStruct()) && Data->IsValid())
+			if (Data->IsValid() && Data->GetScriptStruct()->IsChildOf(T::StaticStruct()))
 			{
 				return Data->GetMutablePtr<T>();
 			}
diff --git a/Engine/Plugins/PCG/Source/PCGEditor/Private/EditorMode/Tools/Helpers/PCGEdModeEditorUtilities.cpp b/Engine/Plugins/PCG/Source/PCGEditor/Private/EditorMode/Tools/Helpers/PCGEdModeEditorUtilities.cpp
--- a/Engine/Plugins/PCG/Source/PCGEditor/Private/EditorMode/Tools/Helpers/PCGEdModeEditorUtilities.cpp
+++ b/Engine/Plugins/PCG/Source/PCGEditor/Private/EditorMode/Tools/Helpers/PCGEdModeEditorUtilities.cpp
@@ -1,9 +1,10 @@
-﻿// Copyright Epic Games, Inc. All Rights Reserved.
+// Copyright Epic Games, Inc. All Rights Reserved.
 
 #include "EditorMode/Tools/Helpers/PCGEdModeEditorUtilities.h"
 
 #include "PCGGraph.h"
 #include "AssetRegistry/AssetRegistryModule.h"
+#include "GameFramework/Actor.h"
 #include "Misc/OutputDeviceNull.h"
 #include "HAL/IConsoleManager.h"
 #include "Framework/Commands/InputBindingManager.h"
@@ -74,6 +75,28 @@ bool UE::PCG::EditorMode::Utility::DoesPCGGraphHaveToolTag(FName ToolTag, const
 	return false;
 }
 
+bool UE::PCG::EditorMode::Utility::IsPCGGraphCompatibleWithActor(AActor* Actor, const FAssetData& AssetData)
+{
+	if (!ensure(Actor))
+	{
+		return false;
+	}
+
+	FPCGGraphToolData ToolData;
+
+	FString ToolDataText;
+	if (AssetData.GetTagValue(GET_MEMBER_NAME_CHECKED(UPCGGraph, ToolData), ToolDataText))
+	{
+		FOutputDeviceNull NullOut;
+		StaticStruct<FPCGGraphToolData>()->ImportText(*ToolDataText, &ToolData, nullptr, PPF_None, &NullOut,
+			FPCGGraphToolData::StaticStruct()->GetName(), true);
+
+		return Actor->IsA(ToolData.InitialActorClassToSpawn);
+	}
+
+	return false;
+}
+
 TArray<TSharedPtr<FUICommandInfo>> UE::PCG::EditorMode::Utility::GetOrCreateUICommandsWithToolTag(TSharedRef<FBindingContext> BindingContext, FName ToolTag)
 {
 	static TMap<FName, TOptional<TArray<TSharedPtr<FUICommandInfo>>>> CachedGraphAssetCommands;
diff --git a/Engine/Plugins/PCG/Source/PCGEditor/Private/EditorMode/Tools/Helpers/PCGEdModeEditorUtilities.h b/Engine/Plugins/PCG/Source/PCGEditor/Private/EditorMode/Tools/Helpers/PCGEdModeEditorUtilities.h
--- a/Engine/Plugins/PCG/Source/PCGEditor/Private/EditorMode/Tools/Helpers/PCGEdModeEditorUtilities.h
+++ b/Engine/Plugins/PCG/Source/PCGEditor/Private/EditorMode/Tools/Helpers/PCGEdModeEditorUtilities.h
@@ -7,6 +7,7 @@
 #include "AssetRegistry/AssetData.h"
 #include "Templates/SharedPointer.h"
 
+class AActor;
 class FUICommandInfo;
 class FBindingContext;
  
@@ -20,5 +21,6 @@ namespace UE::PCG::EditorMode::Utility
 	
 	TArray<FPCGGraphToolEditorData> GetGraphToolsWithToolTag(FName ToolTag);
 	bool DoesPCGGraphHaveToolTag(FName ToolTag, const FAssetData& AssetData);
+	bool IsPCGGraphCompatibleWithActor(AActor* Actor, const FAssetData& AssetData);
 	TArray<TSharedPtr<FUICommandInfo>> GetOrCreateUICommandsWithToolTag(TSharedRef<FBindingContext> BindingContext, FName ToolTag);
 }
diff --git a/Engine/Plugins/PCG/Source/PCGEditor/Private/EditorMode/Tools/PCGInteractiveToolSettings.cpp b/Engine/Plugins/PCG/Source/PCGEditor/Private/EditorMode/Tools/PCGInteractiveToolSettings.cpp
--- a/Engine/Plugins/PCG/Source/PCGEditor/Private/EditorMode/Tools/PCGInteractiveToolSettings.cpp
+++ b/Engine/Plugins/PCG/Source/PCGEditor/Private/EditorMode/Tools/PCGInteractiveToolSettings.cpp
@@ -48,8 +48,7 @@ namespace UE::PCG::EditorMode::Tool
 }
 
 UPCGInteractiveToolSettings::UPCGInteractiveToolSettings()
-	: NewActorName(GetDefault<UPCGEditorModeSettings>()->DefaultNewActorName)
-	, NewPCGComponentName(GetDefault<UPCGEditorModeSettings>()->DefaultNewPCGComponentName)
+	: NewPCGComponentName(GetDefault<UPCGEditorModeSettings>()->DefaultNewPCGComponentName)
 {
 	
 }
@@ -57,12 +56,8 @@ UPCGInteractiveToolSettings::UPCGInteractiveToolSettings()
 TValueOrError<bool, FText> UPCGInteractiveToolSettings::Initialize(UInteractiveTool* InTool)
 {
 	using namespace UE::PCG::EditorMode;
-
-	UWorld* World = InTool->GetToolManager()->GetContextQueriesAPI()->GetCurrentEditingWorld();
-
-	FToolBuilderState State;
-	InTool->GetToolManager()->GetContextQueriesAPI()->GetCurrentSelectionState(State);
-	TArray<AActor*> SelectedActors = State.SelectedActors;
+	check(InTool);
+	Tool = InTool;
 
 	// In order to constrain component selection and actor creation, we need to make sure to setup the default settings first.
 	if (TValueOrError<bool, FText> Result = TryLoadDefaults(); Result.HasError() || (Result.HasValue() && Result.GetValue() == false))
@@ -70,11 +65,53 @@ TValueOrError<bool, FText> UPCGInteractiveToolSettings::Initialize(UInteractiveT
 		return Result;
 	}
 
+	OnModified.AddUObject(this, &UPCGInteractiveToolSettings::OnPropertyModified);
+
+	// Keyboard setup for switching between data instance identifiers
+	UKeyInputBehavior* KeyInputBehavior = NewObject<UKeyInputBehavior>();
+	KeyInputBehavior->Initialize(this,
+	{
+		EKeys::One, EKeys::Two, EKeys::Three, EKeys::Four, EKeys::Five,
+		EKeys::Six, EKeys::Seven, EKeys::Eight, EKeys::Nine, EKeys::Zero
+	});
+	KeyInputBehavior->bRequireAllKeys = false;
+	Tool->AddInputBehavior(KeyInputBehavior);
+
+	RegisterPropertyWatchers();
+
+	TValueOrError<bool, FText> InitState = InitActorAndComponent();
+	if (InitState.HasError() || InitState.GetValue() == false)
+	{
+		return InitState;
+	}
+
+	if (UPCGComponent* PCGComponent = GetWorkingPCGComponent())
+	{
+		UPCGGraphInstance* GraphInstance = PCGComponent->GetGraphInstance();
+		CacheGraphInstance(GraphInstance); // Guaranteed to be an instance.
+
+		// Finally, update the tool using the component graph. This will also update all tool data.
+		// Note: ToolGraph just tracks the graph and the graph parameters will be populated directly from the component's instance
+		SetToolGraph(GraphInstance->GetGraph());
+	}
+
+	return MakeValue(true);
+}
+
+TValueOrError<bool, FText> UPCGInteractiveToolSettings::InitActorAndComponent()
+{
+	check(Tool.IsValid());
+	FToolBuilderState State;
+	Tool->GetToolManager()->GetContextQueriesAPI()->GetCurrentSelectionState(State);
+	TArray<AActor*> SelectedActors = State.SelectedActors;
+
 	// If we have selected one actor, that's the actor we want to edit
 	// This actor might:
 	// - Not have a PCG Component
 	// - Have a PCG Component but no valid working data
 	// - Have a PCG Component with valid working data we want to edit
+	check(!EditActor.IsValid());
+
 	if (SelectedActors.Num() == 1 && SelectedActors[0] && IsWorkingActorCompatible(SelectedActors[0]))
 	{
 		EditActor = SelectedActors[0];
@@ -82,34 +119,38 @@ TValueOrError<bool, FText> UPCGInteractiveToolSettings::Initialize(UInteractiveT
 	// If we have no selected actor, we spawn a working actor that needs to fulfill the same actor requirements. This actor is fully instantiated when Applied.
 	else
 	{
-		AActor* SpawnedActor = Actor::SpawnWorking(World, FTransform::Identity, GetWorkingActorClass());
-		// Our spawned actor is also the actor we are editing
-		EditActor = SpawnedActor;
-		
-		SpawnedActor->SetActorLabel(NewActorName.ToString(), false);
+		check(GeneratedResources.SpawnedActor == nullptr);
 
-		FSelectedObjectsChangeList ChangeList;
-		ChangeList.ModificationType = ESelectedObjectsModificationType::Replace;
-		ChangeList.Actors = {SpawnedActor};
-		
-		InTool->GetToolManager()->RequestSelectionChange(ChangeList);
+		if (UWorld* World = Tool->GetToolManager()->GetContextQueriesAPI()->GetCurrentEditingWorld())
+		{
+			using namespace UE::PCG::EditorMode;
+			AActor* SpawnedActor = Actor::SpawnWorking(World, FTransform::Identity, GetWorkingActorClass());
+			// Our spawned actor is also the actor we are editing
+			EditActor = SpawnedActor;
+
+			SpawnedActor->SetActorLabel(NewActorName.ToString(), false);
+
+			FSelectedObjectsChangeList ChangeList;
+			ChangeList.ModificationType = ESelectedObjectsModificationType::Replace;
+			ChangeList.Actors = { SpawnedActor };
 
-		GeneratedResources.SpawnedActor = SpawnedActor;
+			Tool->GetToolManager()->RequestSelectionChange(ChangeList);
+
+			GeneratedResources.SpawnedActor = SpawnedActor;
+		}
 	}
 
-	if(!ensure(EditActor.IsValid()))
+	if (!ensure(EditActor.IsValid()))
 	{
 		return MakeError(LOCTEXT("InvalidEditActorDuringInit", "Invalid selected actor - either as a previously existing actor or a newly spawned actor."));
 	}
 
 	TWeakObjectPtr<UPCGComponent> PCGComponent = FindOrGeneratePCGComponent(*EditActor.Get());
-	if(!ensure(PCGComponent.IsValid()))
+	if (!ensure(PCGComponent.IsValid()))
 	{
 		return MakeError(LOCTEXT("InvalidPCGComponentDuringInit", "Invalid PCG Component - either as a previously existing or transient generated component."));
 	}
 
-	OnModified.AddUObject(this, &UPCGInteractiveToolSettings::OnPropertyModified);
-
 	// If we don't have a valid graph already, we try to assign the default graph to the newly created pcg component.
 	if (PCGComponent->GetGraph() == nullptr)
 	{
@@ -123,26 +164,20 @@ TValueOrError<bool, FText> UPCGInteractiveToolSettings::Initialize(UInteractiveT
 		bTryToGenerate = true;
 	}
 
-	UPCGGraphInstance* GraphInstance = PCGComponent->GetGraphInstance();
-	CacheGraphInstance(GraphInstance); // Guaranteed to be an instance.
+	return MakeValue(true);
+}
 
-	// Finally, update the tool using the component graph. This will also update all tool data.
-	// Note: ToolGraph just tracks the graph and the graph parameters will be populated directly from the component's instance
-	SetToolGraph(GraphInstance->GetGraph());
+bool UPCGInteractiveToolSettings::InitializeWorkingDataForGraph()
+{
+	TArray<FName> ValidDataInstanceNames = GetDataInstanceNamesForGraph();
 
-	// Keyboard setup for switching between data instance identifiers
-	UKeyInputBehavior* KeyInputBehavior = NewObject<UKeyInputBehavior>();
-	KeyInputBehavior->Initialize(this,
+	bool bAnyDataCreated = false;
+	for (const FName& DataInstanceIdentifier : ValidDataInstanceNames)
 	{
-		EKeys::One, EKeys::Two, EKeys::Three, EKeys::Four, EKeys::Five,
-		EKeys::Six, EKeys::Seven, EKeys::Eight, EKeys::Nine, EKeys::Zero
-	});
-	KeyInputBehavior->bRequireAllKeys = false;
-	InTool->AddInputBehavior(KeyInputBehavior);
-
-	RegisterPropertyWatchers();
+		bAnyDataCreated |= InitializeWorkingData(DataInstanceIdentifier);
+	}
 
-	return MakeValue(true);
+	return bAnyDataCreated;
 }
 
 void UPCGInteractiveToolSettings::Apply(UInteractiveTool* OwningTool)
@@ -364,6 +399,14 @@ void UPCGInteractiveToolSettings::SetActorClassToSpawn(TSubclassOf<AActor> Class
 	PostEditChangeProperty(Event);
 }
 
+void UPCGInteractiveToolSettings::SetActorLabel(FName InLabel)
+{
+	NewActorName = InLabel == NAME_None ? GetDefault<UPCGEditorModeSettings>()->DefaultNewActorName : InLabel;
+	// This will call OnPropertyModified, forwarding the change
+	FPropertyChangedEvent Event(FindFProperty<FProperty>(UPCGInteractiveToolSettings::StaticClass(), GET_MEMBER_NAME_CHECKED(UPCGInteractiveToolSettings, NewActorName)));
+	PostEditChangeProperty(Event);
+}
+
 UPCGComponent* UPCGInteractiveToolSettings::FindOrGeneratePCGComponent(AActor& OwningActor) const
 {
 	using namespace UE::PCG::EditorMode;
@@ -438,6 +481,7 @@ TValueOrError<bool, FText> UPCGInteractiveToolSettings::TryLoadDefaults()
 		SetToolGraph(DefaultGraphInterface);
 		SetActorClassToSpawn(DefaultGraphInterface->ToolData.InitialActorClassToSpawn);
 		SetDataInstance(GetDefaultDataInstanceName());
+		SetActorLabel(DefaultGraphInterface->ToolData.NewActorLabel);
 
 		return MakeValue(true);
 	}
@@ -616,7 +660,12 @@ void UPCGInteractiveToolSettings::PostWorkingDataInitialized(FPCGInteractiveTool
 
 void UPCGInteractiveToolSettings::OnPropertyModified(UObject* Object, FProperty* Property)
 {
-	if(Object == this && Property->GetFName() == GET_MEMBER_NAME_CHECKED(UPCGInteractiveToolSettings, ToolGraph))
+	if (Object != this)
+	{
+		return;
+	}
+
+	if(Property->GetFName() == GET_MEMBER_NAME_CHECKED(UPCGInteractiveToolSettings, ToolGraph))
 	{
 		if (UPCGComponent* WorkingComponent = GetWorkingPCGComponent())
 		{
@@ -628,8 +677,9 @@ void UPCGInteractiveToolSettings::OnPropertyModified(UObject* Object, FProperty*
 
 		if(ToolGraph)
 		{
-			// If the graph changes, we also set the default data instance
+			// If the graph changes, we also set the default data instance and the other graph-led values
 			SetActorClassToSpawn(ToolGraph->ToolData.InitialActorClassToSpawn);
+			SetActorLabel(ToolGraph->ToolData.NewActorLabel);
 		}
 
 		// We set the data instance to the first one provided by the graph. Uses the default as fallback.
@@ -637,15 +687,7 @@ void UPCGInteractiveToolSettings::OnPropertyModified(UObject* Object, FProperty*
 		SetDataInstance(DefaultDataInstanceName);
 		
 		// Update the working data to the available data instances ones for the current tool from the graph
-		TArray<FName> ValidDataInstanceNames = GetDataInstanceNamesForGraph();
-
-		bool bAnyDataCreated = false;
-		for (const FName& DataInstanceIdentifier : ValidDataInstanceNames)
-		{
-			bAnyDataCreated |= InitializeWorkingData(DataInstanceIdentifier);
-		}
-
-		if(bAnyDataCreated)
+		if(InitializeWorkingDataForGraph())
 		{
 			// We generate again to make sure the initialized working data has a chance to register its tracking keys.
 			if (UPCGComponent* PCGComponent = GetWorkingPCGComponent())
@@ -654,7 +696,7 @@ void UPCGInteractiveToolSettings::OnPropertyModified(UObject* Object, FProperty*
 			}
 		}
 	}
-	else if (Object == this && Property->GetFName() == GET_MEMBER_NAME_CHECKED(UPCGInteractiveToolSettings, DataInstance))
+	else if (Property->GetFName() == GET_MEMBER_NAME_CHECKED(UPCGInteractiveToolSettings, DataInstance))
 	{
 		if (InitializeWorkingData(DataInstance))
 		{
@@ -665,10 +707,39 @@ void UPCGInteractiveToolSettings::OnPropertyModified(UObject* Object, FProperty*
 			}
 		}
 	}
-	else if(Object == this && Property->GetFName() == GET_MEMBER_NAME_CHECKED(UPCGInteractiveToolSettings, ActorClassToSpawn))
+	else if(Property->GetFName() == GET_MEMBER_NAME_CHECKED(UPCGInteractiveToolSettings, ActorClassToSpawn))
 	{
-		// @todo_pcg
 		// Delete previously spawned actor and set up the new one. Ideally recycling the existing PCG Component.
+		if (HasSpawnedActor() && EditActor->GetClass() != ActorClassToSpawn.Get())
+		{
+			AActor* PreviouslySpawnedActor = EditActor.Get();
+			UPCGComponent* PreviousComponent = GetWorkingPCGComponent();
+			EditActor = nullptr;
+			GeneratedResources.SpawnedActor = nullptr;
+			GeneratedResources.GeneratedPCGComponent = nullptr;
+			WorkingComponentCache.Reset();
+
+			// Unselect the previously edited actor
+			FSelectedObjectsChangeList ChangeList;
+			ChangeList.ModificationType = ESelectedObjectsModificationType::Clear;
+			Tool->GetToolManager()->RequestSelectionChange(ChangeList);
+
+			InitActorAndComponent();
+			if (UPCGComponent* NewComponent = GetWorkingPCGComponent(); NewComponent && PreviousComponent && NewComponent->GetGraph() == PreviousComponent->GetGraph())
+			{
+				// Copy parameters from previous component to new component
+				NewComponent->GetGraphInstance()->ParametersOverrides = PreviousComponent->GetGraphInstance()->ParametersOverrides;
+				NewComponent->GetGraphInstance()->OnGraphParametersChanged(EPCGGraphParameterEvent::UndoRedo, NAME_None);
+			}
+
+			// Update working data otherwise we'll have dangling pointers to components in the data, and the tool might close.
+			InitializeWorkingDataForGraph();
+
+			// @todo_pcg - copy over working data from previous to new on a per-name/matching class basis
+
+			// Finally, get rid of the previous generated actor
+			PreviouslySpawnedActor->Destroy();
+		}
 	}
 }
 
@@ -703,12 +774,26 @@ bool UPCGInteractiveToolSettings::GraphAssetFilter(const FAssetData& AssetData)
 {
 	using namespace UE::PCG::EditorMode;
 	
+	// Asset is of the proper type
 	if(AssetData.GetClass()->IsChildOf<UPCGGraphInterface>() == false)
 	{
 		return true;
 	}
+
+	// Asset has the compatible tool tag.
+	if (!Utility::DoesPCGGraphHaveToolTag(GetToolTag(), AssetData))
+	{
+		return true;
+	}
+	
+	// There is no current selected/created actor, or it is a transient actor.
+	if (!EditActor.IsValid() || GeneratedResources.SpawnedActor == EditActor.Get())
+	{
+		return false;
+	}
 	
-	return !Utility::DoesPCGGraphHaveToolTag(GetToolTag(), AssetData);
+	// Otherwise, the edit actor class must be compatible with the actor class prescribed by the asset.
+	return !Utility::IsPCGGraphCompatibleWithActor(EditActor.Get(), AssetData);
 }
 
 UPCGComponent* UPCGInteractiveToolSettings::GetWorkingPCGComponent() const
diff --git a/Engine/Plugins/PCG/Source/PCGEditor/Private/EditorMode/Tools/PCGInteractiveToolSettings.h b/Engine/Plugins/PCG/Source/PCGEditor/Private/EditorMode/Tools/PCGInteractiveToolSettings.h
--- a/Engine/Plugins/PCG/Source/PCGEditor/Private/EditorMode/Tools/PCGInteractiveToolSettings.h
+++ b/Engine/Plugins/PCG/Source/PCGEditor/Private/EditorMode/Tools/PCGInteractiveToolSettings.h
@@ -74,6 +74,8 @@ public:
 
 	virtual void SetActorClassToSpawn(TSubclassOf<AActor> Class);
 
+	void SetActorLabel(FName InLabel);
+
 	/** Tools can have a requirement on the working actor */
 	virtual bool IsWorkingActorCompatible(const AActor* InActor) const { return InActor != nullptr; }
 
@@ -179,6 +181,8 @@ public:
 protected:
 	/** Try to load the defaults (graph, actor class, ...) into the tool settings (ToolGraph, etc.). */
 	virtual TValueOrError<bool, FText> TryLoadDefaults();
+	virtual TValueOrError<bool, FText> InitActorAndComponent();
+	virtual bool InitializeWorkingDataForGraph();
 	
 	virtual void Tick(float DeltaTime) override;
 	virtual bool IsTickable() const override;
@@ -224,7 +228,7 @@ public:
 	UPROPERTY(EditAnywhere, DisplayName="Component Name", Category = "PCG", meta = (TransientToolProperty, EditCondition="HasGeneratedPCGComponent()", EditConditionHides))
 	FName NewPCGComponentName;
 
-	UPROPERTY(EditAnywhere, Category = "PCG", meta=(TransientToolProperty, EditConditionHides))
+	UPROPERTY(EditAnywhere, Category = "PCG", meta=(TransientToolProperty, EditCondition="HasSpawnedActor()", EditConditionHides))
 	TSubclassOf<AActor> ActorClassToSpawn;
 
 	/** The current time left before a graph refresh (in seconds). */
@@ -257,4 +261,6 @@ private:
 
 	UPROPERTY(meta = (TransientToolProperty))
 	TObjectPtr<UPCGGraphInstance> GraphInstanceCache;
+
+	TWeakObjectPtr<UInteractiveTool> Tool = nullptr;
 };
