[PCG] Fix subgraph overrides on the node not used when having subgraph overrides

* We were using the Subgraph Override parameters directly, discarding the original overrides on the node
* We now have a special path where if we have a subgraph override but also overrides on the node, to propagate them
* Do 2 copies which is wasteful but the property bag API doesn't provide a way to merge 2 different bags.
* Add `Detach` functionality to the Instanced Property Bag to benefit from the logic from there for migration of bags, but taking ownership of the internal value.

#rnx

diff --git a/Engine/Plugins/PCG/Source/PCG/Private/PCGSubgraph.cpp b/Engine/Plugins/PCG/Source/PCG/Private/PCGSubgraph.cpp
--- a/Engine/Plugins/PCG/Source/PCG/Private/PCGSubgraph.cpp
+++ b/Engine/Plugins/PCG/Source/PCG/Private/PCGSubgraph.cpp
@@ -555,8 +555,36 @@ void FPCGSubgraphContext::InitializeUserParametersStruct()
 	const FInstancedPropertyBag* UserParameters = Graph ? Graph->GetUserParametersStruct() : nullptr;
 	FConstStructView UserParametersView = UserParameters ? UserParameters->GetValue() : FConstStructView{};
 
-	if (!OverridableParams.IsEmpty() && UserParametersView.IsValid())
+	if (!OverridableParams.IsEmpty() && UserParameters)
 	{
+		// In the case where we have a subgraph override, we'll need to make sure the override properties are also taken into account when we apply the overrides
+		// and copy the overriden values.
+		const UPCGGraphInstance* OriginalGraphInstance = Cast<UPCGGraphInstance>(Settings->GetOriginalSubgraphInterface());
+		if (OriginalGraphInstance && Graph != OriginalGraphInstance && !OriginalGraphInstance->ParametersOverrides.PropertiesIDsOverridden.IsEmpty())
+		{
+			// Copy the subgraph override parameters first and copy to the subgraph override parameters, copying over compatible parameters
+			FInstancedPropertyBag TempGraphInstanceParametersOverride = *UserParameters; // Full copy - Also set the layout of the property bag
+			TArray<FName> PropertyNames;
+			PropertyNames.Reserve(OriginalGraphInstance->ParametersOverrides.PropertiesIDsOverridden.Num());
+			const UPropertyBag* OriginalBag = OriginalGraphInstance->ParametersOverrides.Parameters.GetPropertyBagStruct();
+			if (ensure(OriginalBag))
+			{
+				Algo::Transform(OriginalGraphInstance->ParametersOverrides.PropertiesIDsOverridden, PropertyNames,
+				[OriginalBag](const FGuid& InId) -> FName
+					{
+						// If the property is overridden the desc must exist.
+						const FPropertyBagPropertyDesc* Desc = OriginalBag->FindPropertyDescByID(InId);
+						return ensure(Desc) ? Desc->Name : NAME_None;
+					});
+
+				TempGraphInstanceParametersOverride.CopyMatchingValuesByName(OriginalGraphInstance->ParametersOverrides.Parameters, MakeConstArrayView(PropertyNames)); // Subset copy
+				GraphInstanceParametersOverride = TempGraphInstanceParametersOverride.Detach(); // Move
+			}
+
+			return;
+		}
+
+		// Otherwise look for any overridden param
 		bool bHasParamConnected = !InputData.GetParamsByPin(PCGPinConstants::DefaultParamsLabel).IsEmpty();
 
 		int32 Index = 0;
diff --git a/Engine/Plugins/PCG/Source/PCG/Public/PCGSubgraph.h b/Engine/Plugins/PCG/Source/PCG/Public/PCGSubgraph.h
--- a/Engine/Plugins/PCG/Source/PCG/Public/PCGSubgraph.h
+++ b/Engine/Plugins/PCG/Source/PCG/Public/PCGSubgraph.h
@@ -24,6 +24,10 @@ class UPCGBaseSubgraphSettings : public UPCGSettings
 public:
 	UE_API UPCGGraph* GetSubgraph() const;
 	virtual UPCGGraphInterface* GetSubgraphInterface() const { return nullptr; }
+	
+	// When we have nodes that supports subgraph overrides, we need to check if some parameters are overriden there.
+	// By default, it should be the same as SubgraphInterface
+	virtual UPCGGraphInterface* GetOriginalSubgraphInterface() const { return GetSubgraphInterface(); }
 
 	/** Returns true if the subgraphs nodes were not inlined into the parent graphs tasks during compilation. */
 	virtual bool IsDynamicGraph() const { return false; }
@@ -125,6 +129,7 @@ protected:
 public:
 	UE_API virtual UPCGGraphInterface* GetSubgraphInterface() const override;
 	UE_API virtual bool IsDynamicGraph() const override;
+	virtual UPCGGraphInterface* GetOriginalSubgraphInterface() const { return SubgraphInstance; }
 
 	/** Used to filter the subgraph list based on the graph configuration. */
 	UFUNCTION()
diff --git a/Engine/Source/Runtime/CoreUObject/Private/StructUtils/PropertyBag.cpp b/Engine/Source/Runtime/CoreUObject/Private/StructUtils/PropertyBag.cpp
--- a/Engine/Source/Runtime/CoreUObject/Private/StructUtils/PropertyBag.cpp
+++ b/Engine/Source/Runtime/CoreUObject/Private/StructUtils/PropertyBag.cpp
@@ -1922,19 +1922,33 @@ void FInstancedPropertyBag::InitializeFromBagStruct(const UPropertyBag* NewBagSt
 	Value.InitializeAs(NewBagStruct);
 }
 
-void FInstancedPropertyBag::CopyMatchingValuesByID(const FInstancedPropertyBag& Other)
+void FInstancedPropertyBag::CopyMatchingValuesByID(const FInstancedPropertyBag& Other, TOptional<TConstArrayView<FGuid>> OptionalPropertyIdsSubset)
 {
-	UE::StructUtils::Private::CopyMatchingValuesByID(Other.Value, Value);
+	if (OptionalPropertyIdsSubset.IsSet())
+	{
+		UE::StructUtils::Private::CopyMatchingValuesByIDWithOverrides(Other.Value, Value, OptionalPropertyIdsSubset.GetValue());
+	}
+	else
+	{
+		UE::StructUtils::Private::CopyMatchingValuesByID(Other.Value, Value);
+	}
 }
 
-void FInstancedPropertyBag::CopyMatchingValuesByName(const FInstancedPropertyBag& Other)
+void FInstancedPropertyBag::CopyMatchingValuesByName(const FInstancedPropertyBag& Other, TOptional<TConstArrayView<FName>> OptionalPropertyNamesSubset)
 {
 	UE::StructUtils::Private::CopyMatchingValues(
 		Other.Value,
 		Value,
-		[](TNotNull<const UPropertyBag*> TargetBagStruct, const FPropertyBagPropertyDesc& SourceDesc)
+		[&OptionalPropertyNamesSubset](TNotNull<const UPropertyBag*> TargetBagStruct, const FPropertyBagPropertyDesc& SourceDesc) -> const FPropertyBagPropertyDesc*
 		{
-			return TargetBagStruct->FindPropertyDescByName(SourceDesc.Name);
+			if (OptionalPropertyNamesSubset.IsSet() && !OptionalPropertyNamesSubset->Contains(SourceDesc.Name))
+			{
+				return nullptr;
+			}
+			else
+			{
+				return TargetBagStruct->FindPropertyDescByName(SourceDesc.Name);
+			}
 		});
 }
 
@@ -2256,6 +2270,11 @@ void* FInstancedPropertyBag::GetMutableValueAddress(const FPropertyBagPropertyDe
 	return Value.GetMutableMemory() + Desc->CachedProperty->GetOffset_ForInternal();
 }
 
+FInstancedStruct FInstancedPropertyBag::Detach()
+{
+	return FInstancedStruct(MoveTemp(Value));
+}
+
 TValueOrError<bool, EPropertyBagResult> FInstancedPropertyBag::GetValueBool(const FName Name) const
 {
 	const FPropertyBagPropertyDesc* Desc = FindPropertyDescByName(Name);
diff --git a/Engine/Source/Runtime/CoreUObject/Public/StructUtils/PropertyBag.h b/Engine/Source/Runtime/CoreUObject/Public/StructUtils/PropertyBag.h
--- a/Engine/Source/Runtime/CoreUObject/Public/StructUtils/PropertyBag.h
+++ b/Engine/Source/Runtime/CoreUObject/Public/StructUtils/PropertyBag.h
@@ -440,16 +440,18 @@ struct FInstancedPropertyBag
 	/**
 	 * Copies matching property values from another bag of potentially mismatching layout.
 	 * The properties are matched between the bags based on the property ID.
-	 * @param NewDescs Reference to the bag to copy the values from
+	 * @param Other Reference to the bag to copy the values from
+	 * @param OptionalPropertyIdsSubset Optional array of property ids subset to copy from NewDescs
 	 */
-	UE_API void CopyMatchingValuesByID(const FInstancedPropertyBag& NewDescs);
+	UE_API void CopyMatchingValuesByID(const FInstancedPropertyBag& Other, TOptional<TConstArrayView<FGuid>> OptionalPropertyIdsSubset = {});
 
 	/**
 	 * Copies matching property values from another bag of potentially mismatching layout.
 	 * The properties are matched between the bags based on the property name.
 	 * @param Other Reference to the bag to copy the values from
+	 * @param OptionalPropertyNamesSubset Optional array of property names subset to copy from Other 
 	 */
-	UE_API void CopyMatchingValuesByName(const FInstancedPropertyBag& Other);
+	UE_API void CopyMatchingValuesByName(const FInstancedPropertyBag& Other, TOptional<TConstArrayView<FName>> OptionalPropertyNamesSubset = {});
 
 	/** Returns number of the Properties in this Property Bag */
 	UE_API int32 GetNumPropertiesInBag() const;
@@ -605,7 +607,10 @@ struct FInstancedPropertyBag
 
 	/** @return const view to the struct that holds the values. NOTE: The returned value/view cannot be serialized, use this to access the struct only temporarily. */
 	FStructView GetMutableValue() { return Value; };
-	
+
+	/** Return Internal instanced struct to be able to take ownership on the internal memory of this instanced property bag. Internal value will be reset to invalid state. */
+	UE_API FInstancedStruct Detach();
+
 	/**
 	 * Getters
 	 * Numeric types (bool, (u)int32, (u)int64, float, double) support type conversion.
