[PCG] Fix Hash Attribute not working on anything else than [REDACTED_HANDLE]

* Also mark GetInputSource on UPCGMetadataSettingsBase PURE_VIRTUAL as it is mandatory for the child class to override it

#rnx

diff --git a/Engine/Plugins/PCG/Source/PCG/Private/Elements/Metadata/PCGHashAttribute.cpp b/Engine/Plugins/PCG/Source/PCG/Private/Elements/Metadata/PCGHashAttribute.cpp
--- a/Engine/Plugins/PCG/Source/PCG/Private/Elements/Metadata/PCGHashAttribute.cpp
+++ b/Engine/Plugins/PCG/Source/PCG/Private/Elements/Metadata/PCGHashAttribute.cpp
@@ -8,6 +8,12 @@
 
 #define LOCTEXT_NAMESPACE "PCGHashAttributeElement"
 
+FPCGAttributePropertyInputSelector UPCGHashAttributeSettings::GetInputSource(uint32 Index) const
+{
+	ensure(Index == 0);
+	return InputSource1;
+}
+
 FPCGElementPtr UPCGHashAttributeSettings::CreateElement() const
 {
 	return MakeShared<FPCGHashAttributeElement>();
diff --git a/Engine/Plugins/PCG/Source/PCG/Public/Elements/Metadata/PCGHashAttribute.h b/Engine/Plugins/PCG/Source/PCG/Public/Elements/Metadata/PCGHashAttribute.h
--- a/Engine/Plugins/PCG/Source/PCG/Public/Elements/Metadata/PCGHashAttribute.h
+++ b/Engine/Plugins/PCG/Source/PCG/Public/Elements/Metadata/PCGHashAttribute.h
@@ -15,6 +15,8 @@ class UPCGHashAttributeSettings : public UPCGMetadataSettingsBase
 
 public:
 	//~Begin UPCGMetadataSettingsBase
+	PCG_API virtual FPCGAttributePropertyInputSelector GetInputSource(uint32 Index) const override;
+
 	virtual FName GetInputPinLabel(uint32 Index) const override { return PCGPinConstants::DefaultInputLabel; }
 	virtual bool IsSupportedInputType(uint16 TypeId, uint32 InputIndex, bool& bHasSpecialRequirement) const override { return true; }
 	virtual uint32 GetOperandNum() const override { return 1; }
@@ -34,7 +36,7 @@ public:
 	virtual EPCGMetadataTypes GetPinInitialDefaultValueType(FName PinLabel) const override { return EPCGMetadataTypes::String; }
 	//~End IPCGSettingsDefaultValueProvider interface
 
-	UPROPERTY(BlueprintReadWrite, EditAnywhere, Category = Settings, meta = (PCG_Overridable))
+	UPROPERTY(BlueprintReadWrite, EditAnywhere, Category = Settings, meta = (PCG_Overridable, DisplayName="Input Source"))
 	FPCGAttributePropertyInputSelector InputSource1;
 };
 
diff --git a/Engine/Plugins/PCG/Source/PCG/Public/Elements/Metadata/PCGMetadataOpElementBase.h b/Engine/Plugins/PCG/Source/PCG/Public/Elements/Metadata/PCGMetadataOpElementBase.h
--- a/Engine/Plugins/PCG/Source/PCG/Public/Elements/Metadata/PCGMetadataOpElementBase.h
+++ b/Engine/Plugins/PCG/Source/PCG/Public/Elements/Metadata/PCGMetadataOpElementBase.h
@@ -155,7 +155,7 @@ protected:
 	//~End UPCGSettings interface
 
 public:
-	virtual FPCGAttributePropertyInputSelector GetInputSource(uint32 Index) const { return FPCGAttributePropertyInputSelector(); }
+	virtual FPCGAttributePropertyInputSelector GetInputSource(uint32 Index) const PURE_VIRTUAL(UPCGMetadataSettingsBase::GetInputSource, return FPCGAttributePropertyInputSelector(); )
 
 	virtual FName GetInputPinLabel(uint32 Index) const { return PCGPinConstants::DefaultInputLabel; }
 	virtual uint32 GetOperandNum() const { return 1; }
