- Moved Graph Tool Data from UPCGGraphInterface to UPCGGraph as other derived classes don't need the property itself.
- Added virtual GetGraphToolData function to UPCGGraphInterface so that UPCGGraph can return its property, while Instances return their graph's tool data.
--- Can be used in the future to provide overrides as well
- Instances now save their PCG Graph path to Asset Registry, allowing us to fetch tool data from the original graph from an instance without loading it
--- Override data will also need to be written to Asset Registry in the future

Also made working tool data editanywhere so it can be edited in the details panel

#rnx

diff --git a/Engine/Plugins/PCG/Source/PCG/Public/Data/PCGToolData.h b/Engine/Plugins/PCG/Source/PCG/Public/Data/PCGToolData.h
--- a/Engine/Plugins/PCG/Source/PCG/Public/Data/PCGToolData.h
+++ b/Engine/Plugins/PCG/Source/PCG/Public/Data/PCGToolData.h
@@ -100,6 +100,6 @@ struct FPCGInteractiveToolDataContainer
 		return nullptr;
 	}
 	
-	UPROPERTY(BlueprintReadOnly, VisibleAnywhere, Category = "Tool Data")
+	UPROPERTY(BlueprintReadOnly, EditAnywhere, Category = "Tool Data")
 	TArray<TInstancedStruct<FPCGInteractiveToolWorkingBaseData>> ToolData;
 };
\ No newline at end of file
diff --git a/Engine/Plugins/PCG/Source/PCG/Public/PCGGraph.h b/Engine/Plugins/PCG/Source/PCG/Public/PCGGraph.h
--- a/Engine/Plugins/PCG/Source/PCG/Public/PCGGraph.h
+++ b/Engine/Plugins/PCG/Source/PCG/Public/PCGGraph.h
@@ -138,6 +138,8 @@ public:
 
 	virtual TOptional<FText> GetTitleOverride() const { return bOverrideTitle ? Title : TOptional<FText>(); }
 	virtual TOptional<FLinearColor> GetColorOverride() const { return bOverrideColor ? Color : TOptional<FLinearColor>(); }
+	virtual TOptional<FPCGGraphToolData> GetGraphToolData() const { return {}; }
+
 protected:
 	/** By default export to library is visible only for graph that are assets, but can be enabled/disabled if needed. */
 	UFUNCTION()
@@ -179,10 +181,6 @@ public:
 	UPROPERTY(BlueprintReadWrite, EditAnywhere, Category = AssetInfo, AssetRegistrySearchable, meta = (EditCondition = "IsExportToLibraryEnabled", EditConditionHides, PCGNoHash))
 	bool bExposeToLibrary = false;
 
-	/** Contains the data relevant for PCG Editor Mode usage. */
-	UPROPERTY(BlueprintReadWrite, EditAnywhere, Category = AssetInfo, AssetRegistrySearchable, meta=(EditConditionHides, PCGNoHash))
-	FPCGGraphToolData ToolData;
-
 	/** When enabled, this graph can be executed outside of the world using an editor execution source. */
 	UPROPERTY(BlueprintReadWrite, EditAnywhere, Category = AssetInfo, AssetRegistrySearchable, meta = (PCGNoHash, EditCondition = "CanToggleStandaloneGraph()", EditConditionHides))
 	bool bIsStandaloneGraph = false;
@@ -296,6 +294,8 @@ public:
 	UE_API virtual void PreEditChange(FProperty* InProperty) override;
 	UE_API virtual void PostEditChangeChainProperty(FPropertyChangedChainEvent& PropertyChangedEvent);
 	UE_API virtual void PostEditUndo() override;
+	
+	virtual TOptional<FPCGGraphToolData> GetGraphToolData() const override { return ToolData; }
 #endif
 	/** ~End UObject interface */
 
@@ -333,6 +333,10 @@ public:
 
 	UPROPERTY(EditAnywhere, Category = Customization, meta=(PCGNoHash, EditCondition="ShowGraphCustomization()", EditConditionHides))
 	FPCGGraphEditorCustomization GraphCustomization;
+
+	/** Contains the data relevant for PCG Editor Mode usage. */
+	UPROPERTY(BlueprintReadWrite, EditAnywhere, Category = AssetInfo, AssetRegistrySearchable, meta=(PCGNoHash))
+	FPCGGraphToolData ToolData;
 #endif
 
 	UPROPERTY(EditAnywhere, Category = "Settings|Advanced", meta=(EditCondition="!bIsStandaloneGraph", EditConditionHides))
@@ -688,6 +692,8 @@ public:
 
 	UE_API void SetupCallbacks();
 	UE_API void TeardownCallbacks();
+
+	virtual TOptional<FPCGGraphToolData> GetGraphToolData() const override { return GetGraph() ? GetGraph()->GetGraphToolData() : TOptional<FPCGGraphToolData>(); }
 #endif
 
 	UE_API virtual void OnGraphParametersChanged(EPCGGraphParameterEvent InChangeType, FName InChangedPropertyName) override;
diff --git a/Engine/Plugins/PCG/Source/PCGEditor/Private/EditorMode/Tools/Helpers/PCGEdModeEditorUtilities.cpp b/Engine/Plugins/PCG/Source/PCGEditor/Private/EditorMode/Tools/Helpers/PCGEdModeEditorUtilities.cpp
--- a/Engine/Plugins/PCG/Source/PCGEditor/Private/EditorMode/Tools/Helpers/PCGEdModeEditorUtilities.cpp
+++ b/Engine/Plugins/PCG/Source/PCGEditor/Private/EditorMode/Tools/Helpers/PCGEdModeEditorUtilities.cpp
@@ -57,18 +57,54 @@ TArray<UE::PCG::EditorMode::Utility::FPCGGraphToolEditorData> UE::PCG::EditorMod
 
 bool UE::PCG::EditorMode::Utility::DoesPCGGraphHaveToolTag(FName ToolTag, const FAssetData& AssetData)
 {
-	FPCGGraphToolData ToolData;
-
-	FString ToolDataText;
-	if(AssetData.GetTagValue(GET_MEMBER_NAME_CHECKED(UPCGGraph, ToolData), ToolDataText))
+	if (AssetData.GetClass()->IsChildOf(UPCGGraph::StaticClass()))
 	{
-		FOutputDeviceNull NullOut;
-		StaticStruct<FPCGGraphToolData>()->ImportText(*ToolDataText, &ToolData, nullptr, PPF_None, &NullOut,
-			FPCGGraphToolData::StaticStruct()->GetName(), true);
+		FPCGGraphToolData ToolData;
+		FString ToolDataText;
+		if (AssetData.GetTagValue(GET_MEMBER_NAME_CHECKED(UPCGGraph, ToolData), ToolDataText))
+		{
+			FOutputDeviceNull NullOut;
+			StaticStruct<FPCGGraphToolData>()->ImportText(*ToolDataText, &ToolData, nullptr, PPF_None, &NullOut,
+				FPCGGraphToolData::StaticStruct()->GetName(), true);
 
-		if(ToolData.CompatibleToolTags.Contains(ToolTag))
+			if (ToolData.CompatibleToolTags.Contains(ToolTag))
+			{
+				return true;
+			}
+		}
+	}
+
+	return false;
+}
+
+bool UE::PCG::EditorMode::Utility::DoesPCGGraphInterfaceHaveToolTag(FName ToolTag, const FAssetData& AssetData)
+{	
+	if (!AssetData.IsValid() || !AssetData.GetClass()->IsChildOf(UPCGGraphInterface::StaticClass()))
+	{
+		return false;
+	}
+	
+	if (AssetData.GetClass()->IsChildOf(UPCGGraph::StaticClass()))
+	{
+		return DoesPCGGraphHaveToolTag(ToolTag, AssetData);
+	}
+	else if (AssetData.GetClass()->IsChildOf(UPCGGraphInstance::StaticClass()))
+	{
+		FString PCGGraphText;
+		if (AssetData.GetTagValue("Graph", PCGGraphText))
 		{
-			return true;
+			FSoftObjectPath PCGGraphPath(PCGGraphText);
+
+			if(PCGGraphPath.IsAsset())
+			{
+				FAssetRegistryModule& AssetRegistryModule = FModuleManager::Get().LoadModuleChecked<FAssetRegistryModule>("AssetRegistry");
+				FAssetData LinkedPCGGraphAssetData = AssetRegistryModule.Get().GetAssetByObjectPath(PCGGraphPath);
+
+				if (LinkedPCGGraphAssetData.IsValid() && LinkedPCGGraphAssetData.GetClass()->IsChildOf(UPCGGraph::StaticClass()))
+				{
+					return DoesPCGGraphHaveToolTag(ToolTag, LinkedPCGGraphAssetData);
+				}
+			}
 		}
 	}
 
@@ -77,23 +113,59 @@ bool UE::PCG::EditorMode::Utility::DoesPCGGraphHaveToolTag(FName ToolTag, const
 
 bool UE::PCG::EditorMode::Utility::IsPCGGraphCompatibleWithActor(AActor* Actor, const FAssetData& AssetData)
 {
-	if (!ensure(Actor))
+	if (AssetData.GetClass()->IsChildOf(UPCGGraph::StaticClass()))
 	{
-		return false;
+		FPCGGraphToolData ToolData;
+		FString ToolDataText;
+		if (AssetData.GetTagValue(GET_MEMBER_NAME_CHECKED(UPCGGraph, ToolData), ToolDataText))
+		{
+			FOutputDeviceNull NullOut;
+			StaticStruct<FPCGGraphToolData>()->ImportText(*ToolDataText, &ToolData, nullptr, PPF_None, &NullOut,
+				FPCGGraphToolData::StaticStruct()->GetName(), true);
+
+			return Actor->IsA(ToolData.InitialActorClassToSpawn);
+		}
 	}
 
-	FPCGGraphToolData ToolData;
+	return false;
+}
 
-	FString ToolDataText;
-	if (AssetData.GetTagValue(GET_MEMBER_NAME_CHECKED(UPCGGraph, ToolData), ToolDataText))
+bool UE::PCG::EditorMode::Utility::IsPCGGraphInterfaceCompatibleWithActor(AActor* Actor, const FAssetData& AssetData)
+{
+	if (!ensure(Actor))
 	{
-		FOutputDeviceNull NullOut;
-		StaticStruct<FPCGGraphToolData>()->ImportText(*ToolDataText, &ToolData, nullptr, PPF_None, &NullOut,
-			FPCGGraphToolData::StaticStruct()->GetName(), true);
-
-		return Actor->IsA(ToolData.InitialActorClassToSpawn);
+		return false;
 	}
 
+	if (!AssetData.IsValid() || !AssetData.GetClass()->IsChildOf(UPCGGraphInterface::StaticClass()))
+	{
+		return false;
+	}
+	
+	if (AssetData.GetClass()->IsChildOf(UPCGGraph::StaticClass()))
+	{
+		return IsPCGGraphCompatibleWithActor(Actor, AssetData);
+	}
+	else if (AssetData.GetClass()->IsChildOf(UPCGGraphInstance::StaticClass()))
+    {
+    	FString PCGGraphText;
+    	if (AssetData.GetTagValue("Graph", PCGGraphText))
+    	{
+    		FSoftObjectPath PCGGraphPath(PCGGraphText);
+
+    		if(PCGGraphPath.IsAsset())
+    		{
+    			FAssetRegistryModule& AssetRegistryModule = FModuleManager::Get().LoadModuleChecked<FAssetRegistryModule>("AssetRegistry");
+    			FAssetData LinkedPCGGraphAssetData = AssetRegistryModule.Get().GetAssetByObjectPath(PCGGraphPath);
+
+    			if (LinkedPCGGraphAssetData.IsValid() && LinkedPCGGraphAssetData.GetClass()->IsChildOf(UPCGGraph::StaticClass()))
+    			{
+    				return IsPCGGraphCompatibleWithActor(Actor, LinkedPCGGraphAssetData);
+    			}
+    		}
+    	}
+    }
+	
 	return false;
 }
 
diff --git a/Engine/Plugins/PCG/Source/PCGEditor/Private/EditorMode/Tools/Helpers/PCGEdModeEditorUtilities.h b/Engine/Plugins/PCG/Source/PCGEditor/Private/EditorMode/Tools/Helpers/PCGEdModeEditorUtilities.h
--- a/Engine/Plugins/PCG/Source/PCGEditor/Private/EditorMode/Tools/Helpers/PCGEdModeEditorUtilities.h
+++ b/Engine/Plugins/PCG/Source/PCGEditor/Private/EditorMode/Tools/Helpers/PCGEdModeEditorUtilities.h
@@ -21,6 +21,8 @@ namespace UE::PCG::EditorMode::Utility
 	
 	TArray<FPCGGraphToolEditorData> GetGraphToolsWithToolTag(FName ToolTag);
 	bool DoesPCGGraphHaveToolTag(FName ToolTag, const FAssetData& AssetData);
+	bool DoesPCGGraphInterfaceHaveToolTag(FName ToolTag, const FAssetData& AssetData);
 	bool IsPCGGraphCompatibleWithActor(AActor* Actor, const FAssetData& AssetData);
+	bool IsPCGGraphInterfaceCompatibleWithActor(AActor* Actor, const FAssetData& AssetData);
 	TArray<TSharedPtr<FUICommandInfo>> GetOrCreateUICommandsWithToolTag(TSharedRef<FBindingContext> BindingContext, FName ToolTag);
 }
diff --git a/Engine/Plugins/PCG/Source/PCGEditor/Private/EditorMode/Tools/PCGInteractiveToolSettings.cpp b/Engine/Plugins/PCG/Source/PCGEditor/Private/EditorMode/Tools/PCGInteractiveToolSettings.cpp
--- a/Engine/Plugins/PCG/Source/PCGEditor/Private/EditorMode/Tools/PCGInteractiveToolSettings.cpp
+++ b/Engine/Plugins/PCG/Source/PCGEditor/Private/EditorMode/Tools/PCGInteractiveToolSettings.cpp
@@ -472,16 +472,17 @@ TValueOrError<bool, FText> UPCGInteractiveToolSettings::TryLoadDefaults()
 	check(!HasAnyFlags(RF_ClassDefaultObject));
 	if (UPCGGraphInterface* DefaultGraphInterface = CastChecked<UPCGGraphInterface>(GetDefaultGraph().TryLoad(), ECastCheckedType::NullAllowed))
 	{
-		if (DefaultGraphInterface->ToolData.CompatibleToolTags.Contains(GetToolTag()) == false)
+		TOptional<FPCGGraphToolData> GraphToolData = DefaultGraphInterface->GetGraphToolData();
+		if (GraphToolData.IsSet() == false || GraphToolData.GetValue().CompatibleToolTags.Contains(GetToolTag()) == false)
 		{
 			return MakeError(FText::FormatOrdered(LOCTEXT("DefaultGraph_IncompatibleToolTag", "Graph {0} does not have compatible tool tag {1}"), FText::FromString(GetPathNameSafe(DefaultGraphInterface)), FText::FromName(GetToolTag())));
 		}
 		
 		// At this point, the delegates (OnPropertyModified) aren't setup, so we need to set things up manually.
 		SetToolGraph(DefaultGraphInterface);
-		SetActorClassToSpawn(DefaultGraphInterface->ToolData.InitialActorClassToSpawn);
+		SetActorClassToSpawn(GraphToolData.GetValue().InitialActorClassToSpawn);
 		SetDataInstance(GetDefaultDataInstanceName());
-		SetActorLabel(DefaultGraphInterface->ToolData.NewActorLabel);
+		SetActorLabel(GraphToolData.GetValue().NewActorLabel);
 
 		return MakeValue(true);
 	}
@@ -677,9 +678,14 @@ void UPCGInteractiveToolSettings::OnPropertyModified(UObject* Object, FProperty*
 
 		if(ToolGraph)
 		{
-			// If the graph changes, we also set the default data instance and the other graph-led values
-			SetActorClassToSpawn(ToolGraph->ToolData.InitialActorClassToSpawn);
-			SetActorLabel(ToolGraph->ToolData.NewActorLabel);
+			TOptional<FPCGGraphToolData> GraphToolData = ToolGraph->GetGraphToolData();
+
+			if(GraphToolData.IsSet())
+			{
+				// If the graph changes, we also set the actor class to spawn
+				SetActorClassToSpawn(GraphToolData.GetValue().InitialActorClassToSpawn);
+				SetActorLabel(GraphToolData.GetValue().NewActorLabel);
+			}
 		}
 
 		// We set the data instance to the first one provided by the graph. Uses the default as fallback.
@@ -781,7 +787,7 @@ bool UPCGInteractiveToolSettings::GraphAssetFilter(const FAssetData& AssetData)
 	}
 
 	// Asset has the compatible tool tag.
-	if (!Utility::DoesPCGGraphHaveToolTag(GetToolTag(), AssetData))
+	if (!Utility::DoesPCGGraphInterfaceHaveToolTag(GetToolTag(), AssetData))
 	{
 		return true;
 	}
@@ -793,7 +799,7 @@ bool UPCGInteractiveToolSettings::GraphAssetFilter(const FAssetData& AssetData)
 	}
 	
 	// Otherwise, the edit actor class must be compatible with the actor class prescribed by the asset.
-	return !Utility::IsPCGGraphCompatibleWithActor(EditActor.Get(), AssetData);
+	return !Utility::IsPCGGraphInterfaceCompatibleWithActor(EditActor.Get(), AssetData);
 }
 
 UPCGComponent* UPCGInteractiveToolSettings::GetWorkingPCGComponent() const
