[PCG] Crash: Fix unloading of Original PCGComponent causing RemoveGraphInstance calls on Local PCGComponents (while garbage collection is in progress)
- Keep track if the original component was unregistered in a loading scope
- Fix nullptr searchs on TObjectKey which were never hitting
- Also add some extra security in FPCGManagedActorLoadingScope to never load when garbage collecting

diff --git a/Engine/Plugins/PCG/Source/PCG/Private/Grid/PCGPartitionActor.cpp b/Engine/Plugins/PCG/Source/PCG/Private/Grid/PCGPartitionActor.cpp
--- a/Engine/Plugins/PCG/Source/PCG/Private/Grid/PCGPartitionActor.cpp
+++ b/Engine/Plugins/PCG/Source/PCG/Private/Grid/PCGPartitionActor.cpp
@@ -416,10 +416,26 @@ UPCGComponent* APCGPartitionActor::GetLocalComponent(const UPCGComponent* Origin
 	{
 		return *LocalComponent;
 	}
-	else if (bRebuildMappingOnNullEntries && (OriginalToLocal.Contains(nullptr) || (OriginalToLocal.Num() != LocalToOriginal.Num())))
+	else if (bRebuildMappingOnNullEntries)
 	{
-		const_cast<APCGPartitionActor*>(this)->RebuildOriginalToLocal();
-		return GetLocalComponent(OriginalComponent, /*bRebuildMappingOnNullEntries=*/false);
+		bool bNeedsRebuild = OriginalToLocal.Num() != LocalToOriginal.Num();
+		if (!bNeedsRebuild)
+		{
+			for (const auto& OriginalToLocalItem : OriginalToLocal)
+			{
+				if (!OriginalToLocalItem.Key.ResolveObjectPtr())
+				{
+					bNeedsRebuild = true;
+					break;
+				}
+			}
+		}
+		
+		if(bNeedsRebuild)
+		{
+			const_cast<APCGPartitionActor*>(this)->RebuildOriginalToLocal();
+			return GetLocalComponent(OriginalComponent, /*bRebuildMappingOnNullEntries=*/false);
+		}
 	}
 
 	return nullptr;
@@ -506,7 +522,13 @@ bool APCGPartitionActor::CleanupDeadGraphInstancesInternal(bool bCleanupTest, TA
 		}
 
 		// Remove all dead entries
-		OriginalToLocal.Remove(nullptr);
+		for (auto It = OriginalToLocal.CreateIterator(); It; ++It)
+		{
+			if (!It.Key().ResolveObjectPtr())
+			{
+				It.RemoveCurrent();
+			}
+		}
 	}
 
 	return bModified;
diff --git a/Engine/Plugins/PCG/Source/PCG/Private/PCGActorAndComponentMapping.cpp b/Engine/Plugins/PCG/Source/PCG/Private/PCGActorAndComponentMapping.cpp
--- a/Engine/Plugins/PCG/Source/PCG/Private/PCGActorAndComponentMapping.cpp
+++ b/Engine/Plugins/PCG/Source/PCG/Private/PCGActorAndComponentMapping.cpp
@@ -565,9 +565,15 @@ void FPCGActorAndComponentMapping::UnregisterPartitionedPCGComponent(UPCGCompone
 		}
 	}
 
-	for (APCGPartitionActor* Actor : PartitionActorsToCleanUp)
+#if WITH_EDITOR
+	// Avoid removing graph instances when we are being unloaded (non-destructive)
+	if(!InComponent->bUnregisteredThroughLoading)
+#endif
 	{
-		Actor->RemoveGraphInstance(InComponent);
+		for (APCGPartitionActor* Actor : PartitionActorsToCleanUp)
+		{
+			Actor->RemoveGraphInstance(InComponent);
+		}
 	}
 }
 
diff --git a/Engine/Plugins/PCG/Source/PCG/Private/PCGComponent.cpp b/Engine/Plugins/PCG/Source/PCG/Private/PCGComponent.cpp
--- a/Engine/Plugins/PCG/Source/PCG/Private/PCGComponent.cpp
+++ b/Engine/Plugins/PCG/Source/PCG/Private/PCGComponent.cpp
@@ -102,14 +102,14 @@ namespace PCGComponent
 		FPCGManagedActorLoadingScope()
 		{
 #if WITH_EDITOR
-			bShouldLoadActors = CVarLoadManagedActorsForCleanup.GetValueOnAnyThread();
+			bShouldLoadActors = ensure(!IsGarbageCollecting()) && CVarLoadManagedActorsForCleanup.GetValueOnAnyThread();
 #endif
 		}
 
 		FPCGManagedActorLoadingScope(bool bInShouldLoadActors)
 		{
 #if WITH_EDITOR
-			bShouldLoadActors = bInShouldLoadActors;
+			bShouldLoadActors = ensure(!IsGarbageCollecting()) && bInShouldLoadActors;
 #endif
 		}
 
@@ -1946,6 +1946,9 @@ void UPCGComponent::EndPlay(const EEndPlayReason::Type EndPlayReason)
 void UPCGComponent::OnUnregister()
 {
 #if WITH_EDITOR
+	// Track that we were unregistered through loading to avoid clean up
+	bUnregisteredThroughLoading = UE::GetIsEditorLoadingPackage();
+
 	if (UPCGSubsystem* Subsystem = GetSubsystem())
 	{
 		if (!PCGHelpers::IsRuntimeOrPIE())
@@ -2319,6 +2322,8 @@ void UPCGComponent::OnRegister()
 	Super::OnRegister();
 
 #if WITH_EDITOR
+	bUnregisteredThroughLoading = false;
+
 	// Disable auto-refreshing on preview actors until we have something more robust on the execution side.
 	if (GetOwner() && GetOwner()->bIsEditorPreviewActor)
 	{
diff --git a/Engine/Plugins/PCG/Source/PCG/Public/PCGComponent.h b/Engine/Plugins/PCG/Source/PCG/Public/PCGComponent.h
--- a/Engine/Plugins/PCG/Source/PCG/Public/PCGComponent.h
+++ b/Engine/Plugins/PCG/Source/PCG/Public/PCGComponent.h
@@ -662,6 +662,9 @@ protected:
 private:
 	/** Track if component should disable 'bIsComponentPartitioned'. Used to deprecate partitioned components in non-WP levels before support for partitioning in non-WP levels. */
 	bool bDisableIsComponentPartitionedOnLoad = false;
+
+	/** Track if component was unregistered while in a loading scope. */
+	bool bUnregisteredThroughLoading = false;
 #endif // WITH_EDITORONLY_DATA
 	
 	UPROPERTY()
