[PCG] Allow entering paint mode without having a selected actor. Also fixed a crash related to selection of a garbage collected actor.

#rnx

diff --git a/Engine/Plugins/PCG/Source/PCG/Internal/Data/Tool/PCGToolPointData.h b/Engine/Plugins/PCG/Source/PCG/Internal/Data/Tool/PCGToolPointData.h
--- a/Engine/Plugins/PCG/Source/PCG/Internal/Data/Tool/PCGToolPointData.h
+++ b/Engine/Plugins/PCG/Source/PCG/Internal/Data/Tool/PCGToolPointData.h
@@ -4,6 +4,8 @@
 
 #include "Data/Tool/PCGToolBaseData.h"
 
+#include "Components/BoxComponent.h"
+
 #include "PCGToolPointData.generated.h"
 
 /** [EXPERIMENTAL]
@@ -25,6 +27,7 @@ public:
 #if WITH_EDITORONLY_DATA
 	PCG_API virtual void InitializeInternal(const FPCGInteractiveToolWorkingDataContext& Context) override;
 	PCG_API virtual void OnToolStart(const FPCGInteractiveToolWorkingDataContext& Context) override;
+	PCG_API virtual void OnToolApply(const FPCGInteractiveToolWorkingDataContext& Context) override;
 	PCG_API virtual void OnToolCancel(const FPCGInteractiveToolWorkingDataContext& Context) override;
 	PCG_API virtual void OnResetToolDataRequested(const FPCGInteractiveToolWorkingDataContext& Context) override;
 protected:
@@ -39,4 +42,8 @@ protected:
 	/** A transient copy of the point array data, created on tool start/data access. We revert back to these points on Cancel. */
 	UPROPERTY(Transient)
 	TObjectPtr<UPCGPointArrayData> OnToolStartPointArray;
+
+	/** Box component optionally created when the actor doesn't have valid bounds. */
+	UPROPERTY(VisibleAnywhere, Category = "Working Data")
+	TWeakObjectPtr<UBoxComponent> BoxComponent;
 };
\ No newline at end of file
diff --git a/Engine/Plugins/PCG/Source/PCG/Private/Data/Tool/PCGToolPointData.cpp b/Engine/Plugins/PCG/Source/PCG/Private/Data/Tool/PCGToolPointData.cpp
--- a/Engine/Plugins/PCG/Source/PCG/Private/Data/Tool/PCGToolPointData.cpp
+++ b/Engine/Plugins/PCG/Source/PCG/Private/Data/Tool/PCGToolPointData.cpp
@@ -4,6 +4,7 @@
 
 #include "PCGComponent.h"
 #include "PCGContext.h"
+#include "Helpers/PCGHelpers.h"
 #include "Data/PCGPointArrayData.h"
 #include "Helpers/PCGDynamicTrackingHelpers.h"
 
@@ -15,6 +16,42 @@
 void FPCGInteractiveToolWorkingData_PointArrayData::InitializeInternal(const FPCGInteractiveToolWorkingDataContext& Context)
 {
 	GeneratedPointData = NewObject<UPCGPointArrayData>(Context.OwningPCGComponent.Get());
+
+	if (Context.OwningActor.IsValid() && !PCGHelpers::GetActorBounds(Context.OwningActor.Get()).IsValid)
+	{
+		// Create a dummy box component so we can have proper bounds on the actor (otherwise PCG will fail it).
+		BoxComponent = NewObject<UBoxComponent>(Context.OwningActor.Get());
+
+		// Transient until Applied, but always transactional
+		BoxComponent->SetFlags(RF_Transient | RF_Transactional);
+
+		if (!Context.DataInstanceIdentifier.IsNone())
+		{
+			BoxComponent->ComponentTags.Add(Context.DataInstanceIdentifier);
+		}
+
+		BoxComponent->SetCollisionObjectType(ECC_WorldStatic);
+		BoxComponent->SetCollisionResponseToAllChannels(ECR_Ignore);
+		BoxComponent->SetCollisionEnabled(ECollisionEnabled::NoCollision);
+		BoxComponent->SetGenerateOverlapEvents(false);
+		BoxComponent->SetBoxExtent(FVector::OneVector);
+		BoxComponent->SetSimulatePhysics(false);
+		BoxComponent->SetMobility(Context.OwningActor->GetRootComponent()->GetMobility());
+
+		if (!Context.OwningActor->GetRootComponent())
+		{
+			Context.OwningActor->SetRootComponent(BoxComponent.Get());
+		}
+		else
+		{
+			BoxComponent->AttachToComponent(Context.OwningActor->GetRootComponent(), FAttachmentTransformRules(EAttachmentRule::KeepRelative, EAttachmentRule::KeepWorld, EAttachmentRule::KeepWorld, false));
+		}
+
+		BoxComponent->RegisterComponent();
+		Context.OwningActor->AddInstanceComponent(BoxComponent.Get());
+
+		GeneratedResources.GeneratedComponents.Add(BoxComponent);
+	}
 }
 
 void FPCGInteractiveToolWorkingData_PointArrayData::OnToolStart(const FPCGInteractiveToolWorkingDataContext& Context)
@@ -27,6 +64,24 @@ void FPCGInteractiveToolWorkingData_PointArrayData::OnToolStart(const FPCGIntera
 	Super::OnToolStart(Context);
 }
 
+void FPCGInteractiveToolWorkingData_PointArrayData::OnToolApply(const FPCGInteractiveToolWorkingDataContext& Context)
+{
+	// Update bounds if needed
+	if (BoxComponent.IsValid() && GeneratedPointData)
+	{
+		FBox PointBounds = GeneratedPointData->GetBounds();
+		if (PointBounds.IsValid)
+		{
+			BoxComponent->SetRelativeLocation(PointBounds.GetCenter());
+
+			FVector NewExtent = BoxComponent->GetUnscaledBoxExtent().ComponentMax(PointBounds.GetExtent());
+			BoxComponent->SetBoxExtent(NewExtent);
+		}
+	}
+
+	Super::OnToolApply(Context);
+}
+
 void FPCGInteractiveToolWorkingData_PointArrayData::OnToolCancel(const FPCGInteractiveToolWorkingDataContext& Context)
 {
 	if (OnToolStartPointArray != nullptr)
diff --git a/Engine/Plugins/PCG/Source/PCGEditor/Private/EditorMode/Tools/PCGInteractiveToolSettings.cpp b/Engine/Plugins/PCG/Source/PCGEditor/Private/EditorMode/Tools/PCGInteractiveToolSettings.cpp
--- a/Engine/Plugins/PCG/Source/PCGEditor/Private/EditorMode/Tools/PCGInteractiveToolSettings.cpp
+++ b/Engine/Plugins/PCG/Source/PCGEditor/Private/EditorMode/Tools/PCGInteractiveToolSettings.cpp
@@ -427,6 +427,12 @@ void UPCGInteractiveToolSettings::DeleteGeneratedResources() const
 {
 	if(GeneratedResources.SpawnedActor.IsValid())
 	{
+		// Unselect the actor if it was previously selected
+		FSelectedObjectsChangeList ChangeList;
+		ChangeList.ModificationType = ESelectedObjectsModificationType::Remove;
+		ChangeList.Actors.Add(GeneratedResources.SpawnedActor.Get());
+		Tool->GetToolManager()->RequestSelectionChange(ChangeList);
+
 		GeneratedResources.SpawnedActor->Destroy();
 		GeneratedResources.SpawnedActor.Reset();
 	}
diff --git a/Engine/Plugins/PCG/Source/PCGEditor/Private/EditorMode/Tools/Paint/PCGPaintTool.cpp b/Engine/Plugins/PCG/Source/PCGEditor/Private/EditorMode/Tools/Paint/PCGPaintTool.cpp
--- a/Engine/Plugins/PCG/Source/PCGEditor/Private/EditorMode/Tools/Paint/PCGPaintTool.cpp
+++ b/Engine/Plugins/PCG/Source/PCGEditor/Private/EditorMode/Tools/Paint/PCGPaintTool.cpp
@@ -51,6 +51,12 @@ FName UPCGInteractiveToolSettings_PaintTool::GetToolTag() const
 	return StaticGetToolTag();
 }
 
+bool UPCGInteractiveToolSettings_PaintTool::IsWorkingActorCompatible(const AActor* InActor) const
+{
+	// Working actor should have a root component, we'll create one only for new actors.
+	return Super::IsWorkingActorCompatible(InActor) && InActor && InActor->GetRootComponent();
+}
+
 const UScriptStruct* UPCGInteractiveToolSettings_PaintTool::GetWorkingDataType() const
 {
 	return FPCGInteractiveToolWorkingData_PointArrayData::StaticStruct();
@@ -80,13 +86,14 @@ void UPCGPaintTool::Setup()
 	
 	AddToolPropertySource(ToolSettings);
 
-	AActor* TargetActor = Cast<UPrimitiveComponentToolTarget>(GetTarget())->GetOwnerActor();
+	UPrimitiveComponentToolTarget* PrimitiveTarget = Cast<UPrimitiveComponentToolTarget>(GetTarget());
+	AActor* TargetActor = PrimitiveTarget ? PrimitiveTarget->GetOwnerActor() : nullptr;
 
-	if(ToolSettings->GeneratedResources.SpawnedActor.IsValid())
+	if(TargetActor && ToolSettings->GeneratedResources.SpawnedActor.IsValid())
 	{
 		ToolSettings->GeneratedResources.SpawnedActor->SetActorLocation(TargetActor->GetActorLocation());
 	}
-	
+
 	CachedRayParams =
 	{
 		.World = GetWorld(),
@@ -408,7 +415,11 @@ void UPCGPaintTool::TryDeletePointsFromBrush()
 
 bool UPCGPaintToolBuilder::CanBuildTool(const FToolBuilderState& SceneState) const
 {
-	return SceneState.TargetManager->CountSelectedAndTargetable(SceneState, GetTargetRequirements()) > 0;
+	const UPCGInteractiveToolSettings_PaintTool* SettingsCDO = GetDefault<UPCGInteractiveToolSettings_PaintTool>();
+	check(SettingsCDO);
+
+	return SceneState.SelectedActors.IsEmpty()
+		|| SceneState.SelectedActors.ContainsByPredicate([SettingsCDO](const AActor* InActor) { return SettingsCDO->IsWorkingActorCompatible(InActor); });
 }
 
 UInteractiveTool* UPCGPaintToolBuilder::BuildTool(const FToolBuilderState& SceneState) const
@@ -422,7 +433,7 @@ UInteractiveTool* UPCGPaintToolBuilder::BuildTool(const FToolBuilderState& Scene
 		.ToolIcon = FPCGEditorModeStyle::Get().GetBrush("PCGEditorMode.Tools.Paint")
 	});
 
-	UToolTarget* const Target = SceneState.TargetManager->BuildFirstSelectedTargetable(SceneState, GetTargetRequirements());
+	UToolTarget* const Target = SceneState.TargetManager->BuildFirstSelectedTargetable(SceneState, FToolTargetTypeRequirements{});
 	NewTool->SetTarget(Target);
 	NewTool->SetWorld(SceneState.World);
 
@@ -434,10 +445,4 @@ UInteractiveTool* UPCGPaintToolBuilder::BuildTool(const FToolBuilderState& Scene
 	return NewTool;
 }
 
-const FToolTargetTypeRequirements& UPCGPaintToolBuilder::GetTargetRequirements() const
-{
-	static FToolTargetTypeRequirements TypeRequirements(UPrimitiveComponentBackedTarget::StaticClass());
-	return TypeRequirements;
-}
-
 #undef LOCTEXT_NAMESPACE
diff --git a/Engine/Plugins/PCG/Source/PCGEditor/Private/EditorMode/Tools/Paint/PCGPaintTool.h b/Engine/Plugins/PCG/Source/PCGEditor/Private/EditorMode/Tools/Paint/PCGPaintTool.h
--- a/Engine/Plugins/PCG/Source/PCGEditor/Private/EditorMode/Tools/Paint/PCGPaintTool.h
+++ b/Engine/Plugins/PCG/Source/PCGEditor/Private/EditorMode/Tools/Paint/PCGPaintTool.h
@@ -26,6 +26,7 @@ public:
 	virtual FName GetToolTag() const override;
 	virtual const UScriptStruct* GetWorkingDataType() const override;
 	virtual void RegisterPropertyWatchers() override;
+	virtual bool IsWorkingActorCompatible(const AActor* InActor) const override final;
 	// ~End UPCGInteractiveToolSettings interface
 
 	/** The brush shape to use. */
@@ -109,12 +110,11 @@ private:
 };
 
 UCLASS(Transient)
-class UPCGPaintToolBuilder : public UPCGInteractiveToolWithToolTargetsBuilder
+class UPCGPaintToolBuilder : public UPCGInteractiveToolBuilder
 {
 	GENERATED_BODY()
 
 public:
 	virtual bool CanBuildTool(const FToolBuilderState& SceneState) const override;
 	virtual UInteractiveTool* BuildTool(const FToolBuilderState& SceneState) const override;
-	virtual const FToolTargetTypeRequirements& GetTargetRequirements() const override;
 };
