Added property to stop PCG Tool Data from calling "OnToolStart" more than once per tool instantiation.
If the graph changes, yet tool data is still valid and has been started, calling start again would be wrong

diff --git a/Engine/Plugins/PCG/Source/PCGEditor/Private/EditorMode/Tools/PCGInteractiveToolSettings.cpp b/Engine/Plugins/PCG/Source/PCGEditor/Private/EditorMode/Tools/PCGInteractiveToolSettings.cpp
--- a/Engine/Plugins/PCG/Source/PCGEditor/Private/EditorMode/Tools/PCGInteractiveToolSettings.cpp
+++ b/Engine/Plugins/PCG/Source/PCGEditor/Private/EditorMode/Tools/PCGInteractiveToolSettings.cpp
@@ -322,6 +322,9 @@ bool UPCGInteractiveToolSettings::InitializeWorkingData(FName DataInstanceIdenti
 	Context.InteractiveTool = GetTypedOuter<UInteractiveTool>();
 	Context.PCGSettings = this;
 	Context.WorkingDataIdentifier = GetWorkingDataIdentifier(DataInstanceIdentifier);
+
+	// It is possible this is called multiple times for existing tool data; we return whether it has been created in this call.
+	bool bWorkingDataCreated = false;
 	
 	if (EditActor.IsValid() && PCGComponent)
 	{
@@ -336,12 +339,14 @@ bool UPCGInteractiveToolSettings::InitializeWorkingData(FName DataInstanceIdenti
 			
 			PCGComponent->ToolDataContainer.ToolData.Add(MoveTemp(NewToolWorkingData));
 			GeneratedResources.AddedWorkingDataIdentifiers.Add(Context.WorkingDataIdentifier);
+
+			bWorkingDataCreated = true;
 		}
 	}
 
 	FPCGInteractiveToolWorkingData* WorkingData = GetMutableTypedWorkingData(DataInstanceIdentifier);
 	
-	if(!ensureMsgf(WorkingData, TEXT("Working Data has to be allocated at this point!.")))
+	if (!ensureMsgf(WorkingData, TEXT("Working Data has to be allocated at this point!.")))
 	{
 		return false;
 	}
@@ -349,11 +354,15 @@ bool UPCGInteractiveToolSettings::InitializeWorkingData(FName DataInstanceIdenti
 	// The working data needs to be initialized at this point
 	ensureMsgf(WorkingData->IsInitialized(), TEXT("Working Data isn't initialized!. Have you forgotten to call the parent Initialize function?"));
 	
-	// If the data is valid, we call start. If not, we delete it.
-	if(WorkingData->IsValid())
+	if (WorkingData->IsValid())
 	{
-		// Notify the working data that we have begun accessing it. This allows to make a 'last state' that the tool can revert to on Cancel.
-		WorkingData->OnToolStart(Context);
+		// If the data is valid and hasn't gotten OnToolStart called yet, we do so
+		if (WorkingDataToolStartSet.Contains(Context.WorkingDataIdentifier) == false)
+		{
+			// Notify the working data that we have begun accessing it. This allows to make a 'last state' that the tool can revert to on Cancel.
+			WorkingData->OnToolStart(Context);
+			WorkingDataToolStartSet.Add(Context.WorkingDataIdentifier);
+		}
 	}
 	else
 	{
@@ -369,7 +378,7 @@ bool UPCGInteractiveToolSettings::InitializeWorkingData(FName DataInstanceIdenti
 		return InitializeWorkingData(DataInstanceIdentifier);
 	}
 
-	return GeneratedResources.AddedWorkingDataIdentifiers.Contains(Context.WorkingDataIdentifier);
+	return bWorkingDataCreated;
 }
 
 void UPCGInteractiveToolSettings::SetToolGraph(UPCGGraphInterface* InGraph)
diff --git a/Engine/Plugins/PCG/Source/PCGEditor/Private/EditorMode/Tools/PCGInteractiveToolSettings.h b/Engine/Plugins/PCG/Source/PCGEditor/Private/EditorMode/Tools/PCGInteractiveToolSettings.h
--- a/Engine/Plugins/PCG/Source/PCGEditor/Private/EditorMode/Tools/PCGInteractiveToolSettings.h
+++ b/Engine/Plugins/PCG/Source/PCGEditor/Private/EditorMode/Tools/PCGInteractiveToolSettings.h
@@ -259,8 +259,11 @@ private:
 	UFUNCTION()
 	bool HasGeneratedPCGComponent() const;
 
-	UPROPERTY(meta = (TransientToolProperty))
+	UPROPERTY(Transient, meta = (TransientToolProperty))
 	TObjectPtr<UPCGGraphInstance> GraphInstanceCache;
 
+	UPROPERTY(Transient, meta = (TransientToolProperty))
+	TSet<FName> WorkingDataToolStartSet;
+
 	TWeakObjectPtr<UInteractiveTool> Tool = nullptr;
 };
