[PCG] Move bStandaloneGraph in UPCGGraph (down from UPCGGraphInterface)

#rnx

diff --git a/Engine/Plugins/PCG/Source/PCG/Private/PCGGraph.cpp b/Engine/Plugins/PCG/Source/PCG/Private/PCGGraph.cpp
--- a/Engine/Plugins/PCG/Source/PCG/Private/PCGGraph.cpp
+++ b/Engine/Plugins/PCG/Source/PCG/Private/PCGGraph.cpp
@@ -26,6 +26,7 @@
 #if WITH_EDITOR
 #include "CoreGlobals.h"
 #include "Editor.h"
+#include "Engine/AssetManager.h"
 #include "Dialogs/Dialogs.h"
 #include "EdGraph/EdGraphPin.h"
 #include "GameFramework/Actor.h"
@@ -146,6 +147,38 @@ namespace PCGGraphUtils
 * UPCGGraphInterface
 ****************************/
 
+#if WITH_EDITOR
+bool UPCGGraphInterface::IsStandaloneGraph() const
+{
+	return GetGraph() && GetGraph()->bIsStandaloneGraph;
+}
+
+bool UPCGGraphInterface::IsStandaloneGraphAsset(const FAssetData& InAssetData)
+{
+	if (UPCGGraphInterface* GraphInterface = Cast<UPCGGraphInterface>(InAssetData.FastGetAsset(/*bLoad=*/false)))
+	{
+		return GraphInterface->IsStandaloneGraph();
+	}
+	else if (InAssetData.IsInstanceOf<UPCGGraph>())
+	{
+		bool bTagValue = false;
+		return InAssetData.GetTagValue(GET_MEMBER_NAME_CHECKED(UPCGGraph, bIsStandaloneGraph), bTagValue) && bTagValue;
+	}
+	else if (InAssetData.IsInstanceOf<UPCGGraphInstance>())
+	{
+		FSoftObjectPath GraphPath = FSoftObjectPath(InAssetData.GetTagValueRef<FString>(GET_MEMBER_NAME_CHECKED(UPCGGraphInstance, Graph)));
+		UAssetManager& AssetManager = UAssetManager::Get();
+		FAssetData GraphAssetData;
+		if (AssetManager.GetAssetDataForPath(GraphPath, GraphAssetData))
+		{
+			return IsStandaloneGraphAsset(GraphAssetData);
+		}
+	}
+
+	return false;
+}
+#endif
+
 EPropertyBagResult UPCGGraphInterface::SetGraphParameter(const FName PropertyName, const uint64 Value, const UEnum* Enum)
 {
 	FInstancedPropertyBag* UserParameters = GetMutableUserParametersStruct();
@@ -1965,7 +1998,7 @@ void UPCGGraph::PostEditChangeChainProperty(FPropertyChangedChainEvent& Property
 		|| PropertyName == GET_MEMBER_NAME_CHECKED(UPCGGraphInterface, bOverrideTitle)
 		|| PropertyName == GET_MEMBER_NAME_CHECKED(UPCGGraphInterface, Color)
 		|| PropertyName == GET_MEMBER_NAME_CHECKED(UPCGGraphInterface, bOverrideColor)
-		|| PropertyName == GET_MEMBER_NAME_CHECKED(UPCGGraphInterface, bIsStandaloneGraph))
+		|| PropertyName == GET_MEMBER_NAME_CHECKED(UPCGGraph, bIsStandaloneGraph))
 	{
 		NotifyGraphChanged(EPCGChangeType::Cosmetic);
 	}
@@ -2519,8 +2552,6 @@ void UPCGGraphInstance::SetGraph(UPCGGraphInterface* InGraph)
 	Graph = InGraph;
 
 #if WITH_EDITOR
-	bIsStandaloneGraph = InGraph && InGraph->bIsStandaloneGraph;
-
 	SetupCallbacks();
 #endif // WITH_EDITOR
 
diff --git a/Engine/Plugins/PCG/Source/PCG/Private/Subsystems/PCGEngineSubsystem.cpp b/Engine/Plugins/PCG/Source/PCG/Private/Subsystems/PCGEngineSubsystem.cpp
--- a/Engine/Plugins/PCG/Source/PCG/Private/Subsystems/PCGEngineSubsystem.cpp
+++ b/Engine/Plugins/PCG/Source/PCG/Private/Subsystems/PCGEngineSubsystem.cpp
@@ -111,7 +111,7 @@ void UPCGEngineSubsystem::GenerateGraph(const FPCGGenerateGraphParams& InParams)
 		return;
 	}
 
-	if (!InParams.Graph->bIsStandaloneGraph)
+	if (!InParams.Graph->IsStandaloneGraph())
 	{
 		UE_LOG(LogPCG, Error, TEXT("GenerateGraph called with non standalone graph"));
 		InParams.GenerationCallback.ExecuteIfBound(nullptr, EPCGGenerationStatus::Aborted);
diff --git a/Engine/Plugins/PCG/Source/PCG/Public/PCGGraph.h b/Engine/Plugins/PCG/Source/PCG/Public/PCGGraph.h
--- a/Engine/Plugins/PCG/Source/PCG/Public/PCGGraph.h
+++ b/Engine/Plugins/PCG/Source/PCG/Public/PCGGraph.h
@@ -140,6 +140,11 @@ public:
 	virtual TOptional<FLinearColor> GetColorOverride() const { return bOverrideColor ? Color : TOptional<FLinearColor>(); }
 	virtual TOptional<FPCGGraphToolData> GetGraphToolData() const { return {}; }
 
+	UE_API static bool IsStandaloneGraphAsset(const FAssetData& InAssetData);
+
+	UFUNCTION()
+	UE_API bool IsStandaloneGraph() const;
+
 protected:
 	/** By default export to library is visible only for graph that are assets, but can be enabled/disabled if needed. */
 	UFUNCTION()
@@ -151,9 +156,6 @@ protected:
 	UFUNCTION()
 	virtual bool IsTemplatePropertyEnabled() const { return false; }
 
-	UFUNCTION()
-	virtual bool CanToggleStandaloneGraph() const { return true; }
-
 	UE_API bool VerifyAndUpdateIfGraphParameterValueChanged(const FPropertyChangedChainEvent& PropertyChangedEvent);
 	UE_API bool VerifyIfGraphCustomizationChanged(const FPropertyChangedChainEvent& PropertyChangedEvent);
 #endif // WITH_EDITOR
@@ -181,10 +183,6 @@ public:
 	UPROPERTY(BlueprintReadWrite, EditAnywhere, Category = AssetInfo, AssetRegistrySearchable, meta = (EditCondition = "IsExportToLibraryEnabled", EditConditionHides, PCGNoHash))
 	bool bExposeToLibrary = false;
 
-	/** When enabled, this graph can be executed outside of the world using an editor execution source. */
-	UPROPERTY(BlueprintReadWrite, EditAnywhere, Category = AssetInfo, AssetRegistrySearchable, meta = (PCGNoHash, EditCondition = "CanToggleStandaloneGraph()", EditConditionHides))
-	bool bIsStandaloneGraph = false;
-
 	/** When enabled and standalone, it adds the action in the asset explorer to trigger a generation. Invisible, to be set in the CDO of child class to disable. */
 	UPROPERTY(AssetRegistrySearchable, meta = (PCGNoHash))
 	bool bExposeGenerationInAssetExplorer = true;
@@ -328,7 +326,7 @@ public:
 	FText Description;
 
 	/** Marks the graph to be not refreshed automatically when the landscape changes, even if it is used. */
-	UPROPERTY(BlueprintReadWrite, EditAnywhere, Category = "Settings|Advanced", meta=(PCGNoHash, EditCondition="!bIsStandaloneGraph", EditConditionHides))
+	UPROPERTY(BlueprintReadWrite, EditAnywhere, Category = "Settings|Advanced", meta=(PCGNoHash, EditCondition="!IsStandaloneGraph()", EditConditionHides))
 	bool bIgnoreLandscapeTracking = false;
 
 	UPROPERTY(EditAnywhere, Category = Customization, meta=(PCGNoHash, EditCondition="ShowGraphCustomization()", EditConditionHides))
@@ -337,9 +335,13 @@ public:
 	/** Contains the data relevant for PCG Editor Mode usage. */
 	UPROPERTY(BlueprintReadWrite, EditAnywhere, Category = AssetInfo, AssetRegistrySearchable, meta=(PCGNoHash))
 	FPCGGraphToolData ToolData;
+
+	/** When enabled, this graph can be executed outside of the world using an editor execution source. */
+	UPROPERTY(BlueprintReadWrite, EditAnywhere, Category = AssetInfo, AssetRegistrySearchable, meta = (PCGNoHash, EditCondition = "CanToggleStandaloneGraph()", EditConditionHides))
+	bool bIsStandaloneGraph = false;
 #endif
 
-	UPROPERTY(EditAnywhere, Category = "Settings|Advanced", meta=(EditCondition="!bIsStandaloneGraph", EditConditionHides))
+	UPROPERTY(EditAnywhere, Category = "Settings|Advanced", meta=(EditCondition="!IsStandaloneGraph()", EditConditionHides))
 	bool bLandscapeUsesMetadata = true;
 
 	/** Creates a node using the given settings interface. Does not manage ownership - done outside of this method. */
@@ -488,7 +490,10 @@ public:
 protected:
 #if WITH_EDITOR
 	UFUNCTION()
-	virtual bool SupportHierarchicalGeneration() const { return !bIsStandaloneGraph; }
+	virtual bool CanToggleStandaloneGraph() const { return true; }
+
+	UFUNCTION()
+	virtual bool SupportHierarchicalGeneration() const { return !IsStandaloneGraph(); }
 
 	UFUNCTION()
 	virtual bool ShowGraphCustomization() const { return true; }
@@ -574,7 +579,7 @@ protected:
 	bool bUse2DGrid = true;
 
 	/** Controls whether the default pins (In, ...) are provided data on top-level graphs automatically. Note this will eventually change to be turned off by default. */
-	UPROPERTY(BlueprintReadWrite, EditAnywhere, Category = Settings, meta=(EditCondition="!bIsStandaloneGraph", EditConditionHides))
+	UPROPERTY(BlueprintReadWrite, EditAnywhere, Category = Settings, meta=(EditCondition="!IsStandaloneGraph()", EditConditionHides))
 	bool bHasDefaultConstructedInputs = true;
 
 	/** Execution grid size for nodes. */
@@ -597,14 +602,14 @@ protected:
 	mutable FCriticalSection CachedGridInfoLock;
 
 	/** Sets whether this graph is marked as editor-only; note that the IsEditorOnly call depends on the local graph value and the value in all subgraphs, recursively. */
-	UPROPERTY(EditAnywhere, BlueprintReadWrite, Category = Cooking, meta=(EditCondition="!bIsStandaloneGraph", EditConditionHides))
+	UPROPERTY(EditAnywhere, BlueprintReadWrite, Category = Cooking, meta=(EditCondition="!IsStandaloneGraph()", EditConditionHides))
 	bool bIsEditorOnly = false;
 
 #if WITH_EDITORONLY_DATA
 	/** When true the Debug flag in the graph editor will display debug information contextually for the selected debug object. Otherwise
 	* debug information is displayed for all components using a graph (requires regenerate).
 	*/
-	UPROPERTY(EditAnywhere, Category = Debug, meta=(EditCondition="!bIsStandaloneGraph", EditConditionHides))
+	UPROPERTY(EditAnywhere, Category = Debug, meta=(EditCondition="!IsStandaloneGraph()", EditConditionHides))
 	bool bDebugFlagAppliesToIndividualComponents = true;
 #endif // WITH_EDITORONLY_DATA
 
@@ -624,7 +629,7 @@ public:
 	// Best used on a brand new PCG Graph.
 	UE_API EPropertyBagAlterationResult AddUserParameters(const TArray<FPropertyBagPropertyDesc>& InDescs, const UPCGGraph* InOptionalOriginalGraph = nullptr);
 
-	UPROPERTY(BlueprintReadWrite, EditAnywhere, Category = "Runtime Generation", meta=(EditCondition = "!bIsStandaloneGraph", EditConditionHides))
+	UPROPERTY(BlueprintReadWrite, EditAnywhere, Category = "Runtime Generation", meta=(EditCondition = "!IsStandaloneGraph()", EditConditionHides))
 	FPCGRuntimeGenerationRadii GenerationRadii;
 
 	UE_API virtual void OnGraphParametersChanged(EPCGGraphParameterEvent InChangeType, FName InChangedPropertyName) override;
diff --git a/Engine/Plugins/PCG/Source/PCGEditor/Private/AssetDefinitions/AssetDefinition_PCGGraphInterface.cpp b/Engine/Plugins/PCG/Source/PCGEditor/Private/AssetDefinitions/AssetDefinition_PCGGraphInterface.cpp
--- a/Engine/Plugins/PCG/Source/PCGEditor/Private/AssetDefinitions/AssetDefinition_PCGGraphInterface.cpp
+++ b/Engine/Plugins/PCG/Source/PCGEditor/Private/AssetDefinitions/AssetDefinition_PCGGraphInterface.cpp
@@ -87,13 +87,7 @@ namespace MenuExtension_PCGGraphInterface
 
 				const bool bStandaloneGraphSelected = CBContext->SelectedAssets.ContainsByPredicate([](const FAssetData& InAsset)
 				{
-					if (!InAsset.IsInstanceOf<UPCGGraphInterface>())
-					{
-						return false;
-					}
-
-					bool bTagValue = false;
-					return InAsset.GetTagValue(GET_MEMBER_NAME_CHECKED(UPCGGraphInterface, bIsStandaloneGraph), bTagValue) && bTagValue;
+					return UPCGGraphInterface::IsStandaloneGraphAsset(InAsset);
 				});
 
 				if (bStandaloneGraphSelected)
@@ -112,7 +106,7 @@ namespace MenuExtension_PCGGraphInterface
 							for (const FAssetData& AssetData : CBContext->SelectedAssets)
 							{
 								bool bTagValue = false;
-								const bool bIsStandaloneGraph = AssetData.IsInstanceOf<UPCGGraphInterface>() && AssetData.GetTagValue(GET_MEMBER_NAME_CHECKED(UPCGGraphInterface, bIsStandaloneGraph), bTagValue) && bTagValue;
+								const bool bIsStandaloneGraph = UPCGGraphInterface::IsStandaloneGraphAsset(AssetData);
 								const bool bExportAssetActionAvailable = AssetData.IsInstanceOf<UPCGGraphInterface>() && AssetData.GetTagValue(GET_MEMBER_NAME_CHECKED(UPCGGraphInterface, bExposeGenerationInAssetExplorer), bTagValue) && bTagValue;
 								if (!bIsStandaloneGraph || !bExportAssetActionAvailable)
 								{
diff --git a/Engine/Plugins/PCG/Source/PCGEditor/Private/Details/PCGGraphDetails.cpp b/Engine/Plugins/PCG/Source/PCGEditor/Private/Details/PCGGraphDetails.cpp
--- a/Engine/Plugins/PCG/Source/PCGEditor/Private/Details/PCGGraphDetails.cpp
+++ b/Engine/Plugins/PCG/Source/PCGEditor/Private/Details/PCGGraphDetails.cpp
@@ -96,7 +96,7 @@ void FPCGGraphDetails::CustomizeDetails(IDetailLayoutBuilder& DetailBuilder)
 			];
 
 		// Add a "Run Graph Determinism Test" button in the debug category, if we're not a standalone graph
-		if (!Algo::AnyOf(SelectedGraphs, [](const TWeakObjectPtr<UPCGGraph> Graph) { return !Graph.IsValid() || Graph->bIsStandaloneGraph; }))
+		if (!Algo::AnyOf(SelectedGraphs, [](const TWeakObjectPtr<UPCGGraph> Graph) { return !Graph.IsValid() || Graph->IsStandaloneGraph(); }))
 		{
 			IDetailCategoryBuilder& DebugCategory = DetailBuilder.EditCategory(TEXT("Debug"));
 
diff --git a/Engine/Plugins/PCG/Source/PCGEditor/Private/PCGEditor.cpp b/Engine/Plugins/PCG/Source/PCGEditor/Private/PCGEditor.cpp
--- a/Engine/Plugins/PCG/Source/PCGEditor/Private/PCGEditor.cpp
+++ b/Engine/Plugins/PCG/Source/PCGEditor/Private/PCGEditor.cpp
@@ -3323,7 +3323,7 @@ void FPCGEditor::RefreshViews()
 
 void FPCGEditor::UpdateDefaultExecutionSource()
 {
-	if (PCGGraphBeingEdited->bIsStandaloneGraph && !PCGDefaultExecutionSource)
+	if (PCGGraphBeingEdited->IsStandaloneGraph() && !PCGDefaultExecutionSource)
 	{
 		if (UPCGEngineSubsystem* EngineSubsystem = UPCGEngineSubsystem::Get())
 		{
@@ -3331,7 +3331,7 @@ void FPCGEditor::UpdateDefaultExecutionSource()
 			RefreshViews();
 		}
 	}
-	else if (!PCGGraphBeingEdited->bIsStandaloneGraph && PCGDefaultExecutionSource)
+	else if (!PCGGraphBeingEdited->IsStandaloneGraph() && PCGDefaultExecutionSource)
 	{
 		if (IPCGEditorModule* PCGEditorModule = IPCGEditorModule::Get())
 		{
