[PCG] Fix Filter Elements by Index not working with consecutives indices from input

#rnx

diff --git a/Engine/Plugins/PCG/Source/PCG/Private/Elements/PCGFilterElementsByIndex.cpp b/Engine/Plugins/PCG/Source/PCG/Private/Elements/PCGFilterElementsByIndex.cpp
--- a/Engine/Plugins/PCG/Source/PCG/Private/Elements/PCGFilterElementsByIndex.cpp
+++ b/Engine/Plugins/PCG/Source/PCG/Private/Elements/PCGFilterElementsByIndex.cpp
@@ -236,7 +236,8 @@ bool FPCGFilterElementsByIndexElement::ExecuteInternal(FPCGContext* InContext) c
 			{
 				if (IndexArray[i] != RangeEndIndex + 1)
 				{
-					FilteredIndices.AddRange(RangeStartIndex, RangeEndIndex);
+					// Ranges are exclusive on max, so +1
+					FilteredIndices.AddRange(RangeStartIndex, RangeEndIndex + 1);
 					RangeStartIndex = IndexArray[i];
 					RangeEndIndex = RangeStartIndex;
 				}
