Small fix to OnToolStart call procedure.
Wiping working data identifier cache in case the actor changes during tool usage.
Also making sure OnToolStart is called if the working data just got created regardless of the cache as that will always be correct

diff --git a/Engine/Plugins/PCG/Source/PCGEditor/Private/EditorMode/Tools/PCGInteractiveToolSettings.cpp b/Engine/Plugins/PCG/Source/PCGEditor/Private/EditorMode/Tools/PCGInteractiveToolSettings.cpp
--- a/Engine/Plugins/PCG/Source/PCGEditor/Private/EditorMode/Tools/PCGInteractiveToolSettings.cpp
+++ b/Engine/Plugins/PCG/Source/PCGEditor/Private/EditorMode/Tools/PCGInteractiveToolSettings.cpp
@@ -357,7 +357,9 @@ bool UPCGInteractiveToolSettings::InitializeWorkingData(FName DataInstanceIdenti
 	if (WorkingData->IsValid())
 	{
 		// If the data is valid and hasn't gotten OnToolStart called yet, we do so
-		if (WorkingDataToolStartSet.Contains(Context.WorkingDataIdentifier) == false)
+		// It's also possible working data got removed (the actor respawns, component changes etc.) during tool usage.
+		// In that case, since the working data just got created, we also call OnToolStart
+		if (WorkingDataToolStartSet.Contains(Context.WorkingDataIdentifier) == false || bWorkingDataCreated)
 		{
 			// Notify the working data that we have begun accessing it. This allows to make a 'last state' that the tool can revert to on Cancel.
 			WorkingData->OnToolStart(Context);
@@ -739,6 +741,7 @@ void UPCGInteractiveToolSettings::OnPropertyModified(UObject* Object, FProperty*
 			GeneratedResources.SpawnedActor = nullptr;
 			GeneratedResources.GeneratedPCGComponent = nullptr;
 			WorkingComponentCache.Reset();
+			WorkingDataToolStartSet.Reset();
 
 			// Unselect the previously edited actor
 			FSelectedObjectsChangeList ChangeList;
