[PCG] Add box support to volume tool. Fixed risky use of root component in paint tool data. Replaced usage of weak pointers in tool data because we would lose references on blueprint actors.

#rnx

diff --git a/Engine/Plugins/PCG/Source/PCG/Internal/Data/Tool/PCGToolPointData.h b/Engine/Plugins/PCG/Source/PCG/Internal/Data/Tool/PCGToolPointData.h
--- a/Engine/Plugins/PCG/Source/PCG/Internal/Data/Tool/PCGToolPointData.h
+++ b/Engine/Plugins/PCG/Source/PCG/Internal/Data/Tool/PCGToolPointData.h
@@ -24,7 +24,7 @@ public:
 
 	PCG_API UPCGPointArrayData* GetPointArrayData() const;
 
-#if WITH_EDITORONLY_DATA
+#if WITH_EDITOR
 	PCG_API virtual void InitializeInternal(const FPCGInteractiveToolWorkingDataContext& Context) override;
 	PCG_API virtual void OnToolStart(const FPCGInteractiveToolWorkingDataContext& Context) override;
 	PCG_API virtual void OnToolApply(const FPCGInteractiveToolWorkingDataContext& Context) override;
@@ -32,18 +32,20 @@ public:
 	PCG_API virtual void OnResetToolDataRequested(const FPCGInteractiveToolWorkingDataContext& Context) override;
 protected:
 	PCG_API virtual void OnToolShutdown(const FPCGInteractiveToolWorkingDataContext& Context) override;
-#endif
+#endif // WITH_EDITOR
 
 protected:
 	/** Contents of the tool data. */
 	UPROPERTY(VisibleAnywhere, Category = "Working Data")
 	TObjectPtr<UPCGPointArrayData> GeneratedPointData;
 
+#if WITH_EDITORONLY_DATA
 	/** A transient copy of the point array data, created on tool start/data access. We revert back to these points on Cancel. */
 	UPROPERTY(Transient)
 	TObjectPtr<UPCGPointArrayData> OnToolStartPointArray;
 
 	/** Box component optionally created when the actor doesn't have valid bounds. */
 	UPROPERTY(VisibleAnywhere, Category = "Working Data")
-	TWeakObjectPtr<UBoxComponent> BoxComponent;
+	TObjectPtr<UBoxComponent> BoxComponent;
+#endif // WITH_EDITORONLY_DATA
 };
\ No newline at end of file
diff --git a/Engine/Plugins/PCG/Source/PCG/Internal/Data/Tool/PCGToolSplineData.h b/Engine/Plugins/PCG/Source/PCG/Internal/Data/Tool/PCGToolSplineData.h
--- a/Engine/Plugins/PCG/Source/PCG/Internal/Data/Tool/PCGToolSplineData.h
+++ b/Engine/Plugins/PCG/Source/PCG/Internal/Data/Tool/PCGToolSplineData.h
@@ -44,11 +44,13 @@ private:
 protected:
 	/** Spline this working data interacts with. */
 	UPROPERTY(VisibleAnywhere, Category = "Working Data")
-	TWeakObjectPtr<USplineComponent> SplineComponent;
+	TObjectPtr<USplineComponent> SplineComponent;
 
+#if WITH_EDITORONLY_DATA
 	/** Copy of the spline at the beginning of the interaction, for rollback purposes. */
 	UPROPERTY(Transient)
 	TObjectPtr<USplineComponent> OnToolStartSplineComponent;
+#endif // WITH_EDITORONLY_DATA
 };
 
 USTRUCT(DisplayName = "Spline Surface Tool Data")
diff --git a/Engine/Plugins/PCG/Source/PCG/Internal/Data/Tool/PCGToolVolumeData.h b/Engine/Plugins/PCG/Source/PCG/Internal/Data/Tool/PCGToolVolumeData.h
--- a/Engine/Plugins/PCG/Source/PCG/Internal/Data/Tool/PCGToolVolumeData.h
+++ b/Engine/Plugins/PCG/Source/PCG/Internal/Data/Tool/PCGToolVolumeData.h
@@ -7,6 +7,7 @@
 #include "PCGToolVolumeData.generated.h"
 
 class AVolume;
+class UBoxComponent;
 class UBrushComponent;
 class UBrushBuilder;
 
@@ -26,6 +27,15 @@ public:
 	UPROPERTY(VisibleAnywhere, Category = "Working Data")
 	TWeakObjectPtr<UBrushComponent> BrushComponent = nullptr;
 
+	/** Box component this working data is affecting. */
+	UPROPERTY(VisibleAnywhere, Category = "Working Data")
+	TWeakObjectPtr<UBoxComponent> BoxComponent = nullptr;
+
+	//~Begin FPCGInteractiveToolWorkingData interface
+public:
+	virtual bool IsValid() const override;
+	virtual void InitializeRuntimeElementData(FPCGContext* InContext) const override;
+
 #if WITH_EDITOR
 public:
 	PCG_API virtual void OnToolStart(const FPCGInteractiveToolWorkingDataContext& Context) override;
@@ -35,10 +45,14 @@ protected:
 
 private:
 	virtual void InitializeInternal(const FPCGInteractiveToolWorkingDataContext& Context) override;
+#endif //WITH_EDITOR
+	//~End FPCGInteractiveToolWorkingData interface
 
+#if WITH_EDITOR
 	UBrushComponent* FindOrCreateBrushComponent(const FPCGInteractiveToolWorkingDataContext& Context);
 	virtual UBrushComponent* FindMatchingBrushComponent(const FPCGInteractiveToolWorkingDataContext& Context);
 	virtual UBrushComponent* GenerateMatchingBrushComponent(const FPCGInteractiveToolWorkingDataContext& Context);
+	UBoxComponent* FindBoxComponent(const FPCGInteractiveToolWorkingDataContext& Context);
 #endif
 
 protected:
@@ -47,11 +61,16 @@ protected:
 	UPROPERTY(Transient)
 	TObjectPtr<UBrushComponent> OnToolStartBrushComponent = nullptr;
 
+	/** Copy of the original brush builder on tool initialization, used for rollback purposes. */
 	UPROPERTY(Transient)
 	TObjectPtr<UBrushBuilder> OnToolStartBrushBuilder = nullptr;
 
 	/** Original transform of the volume actor, used for rollback. */
 	UPROPERTY(Transient)
 	FTransform OnToolStartVolumeTransform;
+
+	/** Copy of the original box component on tool initialization, used for rollback purposes. */
+	UPROPERTY(Transient)
+	TObjectPtr<UBoxComponent> OnToolStartBoxComponent = nullptr;
 #endif
 };
\ No newline at end of file
diff --git a/Engine/Plugins/PCG/Source/PCG/Private/Data/Tool/PCGToolPointData.cpp b/Engine/Plugins/PCG/Source/PCG/Private/Data/Tool/PCGToolPointData.cpp
--- a/Engine/Plugins/PCG/Source/PCG/Private/Data/Tool/PCGToolPointData.cpp
+++ b/Engine/Plugins/PCG/Source/PCG/Private/Data/Tool/PCGToolPointData.cpp
@@ -36,14 +36,15 @@ void FPCGInteractiveToolWorkingData_PointArrayData::InitializeInternal(const FPC
 		BoxComponent->SetGenerateOverlapEvents(false);
 		BoxComponent->SetBoxExtent(FVector::OneVector);
 		BoxComponent->SetSimulatePhysics(false);
-		BoxComponent->SetMobility(Context.OwningActor->GetRootComponent()->GetMobility());
 
 		if (!Context.OwningActor->GetRootComponent())
 		{
+			BoxComponent->SetMobility(EComponentMobility::Static);
 			Context.OwningActor->SetRootComponent(BoxComponent.Get());
 		}
 		else
 		{
+			BoxComponent->SetMobility(Context.OwningActor->GetRootComponent()->GetMobility());
 			BoxComponent->AttachToComponent(Context.OwningActor->GetRootComponent(), FAttachmentTransformRules(EAttachmentRule::KeepRelative, EAttachmentRule::KeepWorld, EAttachmentRule::KeepWorld, false));
 		}
 
@@ -67,7 +68,7 @@ void FPCGInteractiveToolWorkingData_PointArrayData::OnToolStart(const FPCGIntera
 void FPCGInteractiveToolWorkingData_PointArrayData::OnToolApply(const FPCGInteractiveToolWorkingDataContext& Context)
 {
 	// Update bounds if needed
-	if (BoxComponent.IsValid() && GeneratedPointData)
+	if (BoxComponent && GeneratedPointData)
 	{
 		FBox PointBounds = GeneratedPointData->GetBounds();
 		if (PointBounds.IsValid)
diff --git a/Engine/Plugins/PCG/Source/PCG/Private/Data/Tool/PCGToolSplineData.cpp b/Engine/Plugins/PCG/Source/PCG/Private/Data/Tool/PCGToolSplineData.cpp
--- a/Engine/Plugins/PCG/Source/PCG/Private/Data/Tool/PCGToolSplineData.cpp
+++ b/Engine/Plugins/PCG/Source/PCG/Private/Data/Tool/PCGToolSplineData.cpp
@@ -27,7 +27,7 @@ void FPCGInteractiveToolWorkingData_Spline::OnToolStart(const FPCGInteractiveToo
 {
 	Super::OnToolStart(Context);
 
-	if (SplineComponent.IsValid() && !OnToolStartSplineComponent)
+	if (SplineComponent && !OnToolStartSplineComponent)
 	{
 		OnToolStartSplineComponent = Cast<USplineComponent>(StaticDuplicateObject(SplineComponent.Get(), GetTransientPackage()));
 	}
@@ -36,12 +36,12 @@ void FPCGInteractiveToolWorkingData_Spline::OnToolStart(const FPCGInteractiveToo
 void FPCGInteractiveToolWorkingData_Spline::OnToolCancel(const FPCGInteractiveToolWorkingDataContext& Context)
 {
 	// If the Spline Component is valid and isn't a generated component, we revert its state.
-	if (SplineComponent.IsValid() && GeneratedResources.GeneratedComponents.Contains(SplineComponent) == false)
+	if (SplineComponent && GeneratedResources.GeneratedComponents.Contains(SplineComponent) == false)
 	{
 		if (OnToolStartSplineComponent)
 		{
 			// We restore the spline component to its previous state
-			UEngine::CopyPropertiesForUnrelatedObjects(OnToolStartSplineComponent, SplineComponent.Get());
+			UEngine::CopyPropertiesForUnrelatedObjects(OnToolStartSplineComponent, SplineComponent);
 
 			SplineComponent->UpdateSpline();
 
@@ -67,12 +67,12 @@ void FPCGInteractiveToolWorkingData_Spline::OnToolShutdown(const FPCGInteractive
 
 void FPCGInteractiveToolWorkingData_Spline::OnResetToolDataRequested(const FPCGInteractiveToolWorkingDataContext& Context)
 {
-	if (SplineComponent.IsValid())
+	if (SplineComponent)
 	{
 		SplineComponent->ClearSplinePoints(/*bUpdateSpline=*/true);
 
 		FPropertyChangedEvent EmptyPropertyChangedEvent(nullptr);
-		FCoreUObjectDelegates::OnObjectPropertyChanged.Broadcast(SplineComponent.Get(), EmptyPropertyChangedEvent);
+		FCoreUObjectDelegates::OnObjectPropertyChanged.Broadcast(SplineComponent, EmptyPropertyChangedEvent);
 	}
 }
 
@@ -155,7 +155,7 @@ void FPCGInteractiveToolWorkingData_Spline::InitializeRuntimeElementData(FPCGCon
 
 	TArray<FPCGTaggedData>& Outputs = Context->OutputData.TaggedData;
 
-	if (SplineComponent.IsValid())
+	if (SplineComponent)
 	{
 		FPCGTaggedData& Output = Outputs.Emplace_GetRef();
 		// We duplicate it so that PCG can take 'ownership'.
@@ -164,10 +164,10 @@ void FPCGInteractiveToolWorkingData_Spline::InitializeRuntimeElementData(FPCGCon
 		// And if we recreate it every time it gets properly managed
 		// If it gets dropped from cache, it's expected to be invalid
 		UPCGSplineData* SplineData = FPCGContext::NewObject_AnyThread<UPCGSplineData>(Context);
-		SplineData->Initialize(SplineComponent.Get());
+		SplineData->Initialize(SplineComponent);
 		Output.Data = SplineData;
 #if WITH_EDITOR
-		FPCGSelectionKey Key = FPCGSelectionKey::CreateFromObjectPtr(SplineComponent.Get());
+		FPCGSelectionKey Key = FPCGSelectionKey::CreateFromObjectPtr(SplineComponent);
 		FPCGDynamicTrackingHelper::AddSingleDynamicTrackingKey(Context, MoveTemp(Key), false);
 #endif
 	}
@@ -175,7 +175,7 @@ void FPCGInteractiveToolWorkingData_Spline::InitializeRuntimeElementData(FPCGCon
 
 bool FPCGInteractiveToolWorkingData_Spline::IsValid() const
 {
-	return SplineComponent.IsValid();
+	return SplineComponent != nullptr;
 }
 
 #if WITH_EDITOR
diff --git a/Engine/Plugins/PCG/Source/PCG/Private/Data/Tool/PCGToolVolumeData.cpp b/Engine/Plugins/PCG/Source/PCG/Private/Data/Tool/PCGToolVolumeData.cpp
--- a/Engine/Plugins/PCG/Source/PCG/Private/Data/Tool/PCGToolVolumeData.cpp
+++ b/Engine/Plugins/PCG/Source/PCG/Private/Data/Tool/PCGToolVolumeData.cpp
@@ -4,7 +4,11 @@
 
 #include "PCGComponent.h"
 #include "PCGContext.h"
+#include "Data/PCGCollisionShapeData.h"
+#include "Data/PCGVolumeData.h"
+#include "Helpers/PCGDynamicTrackingHelpers.h"
 
+#include "Components/BoxComponent.h"
 #include "Components/BrushComponent.h"
 #include "Engine/Engine.h"
 #include "Engine/BrushBuilder.h"
@@ -15,39 +19,106 @@
 #include "ActorFactories/ActorFactory.h"
 #endif
 
+bool FPCGInteractiveToolWorkingData_Volume::IsValid() const
+{
+	return BrushComponent.IsValid() || BoxComponent.IsValid();
+}
+
+void FPCGInteractiveToolWorkingData_Volume::InitializeRuntimeElementData(FPCGContext* Context) const
+{
+	Super::InitializeRuntimeElementData(Context);
+
+	TArray<FPCGTaggedData>& Outputs = Context->OutputData.TaggedData;
+
+	if (!Volume.IsValid() && !BoxComponent.IsValid())
+	{
+		return;
+	}
+
+	FPCGTaggedData& Output = Outputs.Emplace_GetRef();
+
+	if (Volume.IsValid())
+	{
+		UPCGVolumeData* VolumeData = FPCGContext::NewObject_AnyThread<UPCGVolumeData>(Context);
+		VolumeData->Initialize(Volume.Get());
+		Output.Data = VolumeData;
+
+#if WITH_EDITOR
+		FPCGSelectionKey Key = FPCGSelectionKey::CreateFromObjectPtr(Volume.Get());
+		FPCGDynamicTrackingHelper::AddSingleDynamicTrackingKey(Context, MoveTemp(Key), false);
+#endif
+	}
+	else if (BoxComponent.IsValid())
+	{
+		UPCGCollisionShapeData* BoxData = FPCGContext::NewObject_AnyThread<UPCGCollisionShapeData>(Context);
+		BoxData->Initialize(BoxComponent.Get());
+		Output.Data = BoxData;
+
+#if WITH_EDITOR
+		FPCGSelectionKey Key = FPCGSelectionKey::CreateFromObjectPtr(BoxComponent.Get());
+		FPCGDynamicTrackingHelper::AddSingleDynamicTrackingKey(Context, MoveTemp(Key), false);
+#endif
+	}
+}
+
 #if WITH_EDITOR
 void FPCGInteractiveToolWorkingData_Volume::InitializeInternal(const FPCGInteractiveToolWorkingDataContext& Context)
 {
-	Volume = Cast<AVolume>(Context.OwningActor.Get());
-	BrushComponent = FindOrCreateBrushComponent(Context);
+	if (AVolume* ActorAsVolume = Cast<AVolume>(Context.OwningActor.Get()))
+	{
+		Volume = ActorAsVolume;
+		BrushComponent = FindOrCreateBrushComponent(Context);
+	}
+	else
+	{
+		BoxComponent = FindBoxComponent(Context);
+	}
 }
 
 void FPCGInteractiveToolWorkingData_Volume::OnToolStart(const FPCGInteractiveToolWorkingDataContext& Context)
 {
 	Super::OnToolStart(Context);
 
-	bool bUpdateTransform = (!OnToolStartBrushComponent || !OnToolStartBrushBuilder);
+	bool bUpdateTransform = false;
+
+	if (BoxComponent.IsValid() && !OnToolStartBoxComponent)
+	{
+		OnToolStartBoxComponent = Cast<UBoxComponent>(StaticDuplicateObject(BoxComponent.Get(), GetTransientPackage()));
+		bUpdateTransform = true;
+	}
 
 	if (BrushComponent.IsValid() && !OnToolStartBrushComponent)
 	{
 		OnToolStartBrushComponent = Cast<UBrushComponent>(StaticDuplicateObject(BrushComponent.Get(), GetTransientPackage()));
+		bUpdateTransform = true;
 	}
 
 	if (Volume.IsValid() && Volume->BrushBuilder && !OnToolStartBrushBuilder)
 	{
 		OnToolStartBrushBuilder = DuplicateObject<UBrushBuilder>(Volume->BrushBuilder, GetTransientPackage());
+		bUpdateTransform = true;
 	}
 
-	if (Volume.IsValid() && bUpdateTransform)
+	if (bUpdateTransform)
 	{
-		OnToolStartVolumeTransform = Volume->GetActorTransform();
+		if (Volume.IsValid())
+		{
+			OnToolStartVolumeTransform = Volume->GetActorTransform();
+		}
+		else if(BoxComponent.IsValid())
+		{
+			OnToolStartVolumeTransform = BoxComponent->GetComponentTransform();
+		}
 	}
 }
 
 void FPCGInteractiveToolWorkingData_Volume::OnToolCancel(const FPCGInteractiveToolWorkingDataContext& Context)
 {
-	if (Volume.IsValid() && BrushComponent.IsValid() && GeneratedResources.GeneratedComponents.Contains(BrushComponent) == false && OnToolStartBrushComponent)
+	if (Volume.IsValid() && BrushComponent.IsValid() && !GeneratedResources.GeneratedComponents.Contains(BrushComponent) && OnToolStartBrushComponent)
 	{
+		// Restore the brush component to its previous state
+		UEngine::CopyPropertiesForUnrelatedObjects(OnToolStartBrushComponent, BrushComponent.Get());
+
 		Volume->SetActorTransform(OnToolStartVolumeTransform);
 
 		if (OnToolStartBrushBuilder)
@@ -56,12 +127,25 @@ void FPCGInteractiveToolWorkingData_Volume::OnToolCancel(const FPCGInteractiveTo
 			UActorFactory::CreateBrushForVolumeActor(Volume.Get(), OnToolStartBrushBuilder);
 		}
 
-		// Restore the brush component to its previous state
-		UEngine::CopyPropertiesForUnrelatedObjects(OnToolStartBrushComponent, BrushComponent.Get());
-
 		FPropertyChangedEvent EmptyPropertyChangedEvent(nullptr);
 		FCoreUObjectDelegates::OnObjectPropertyChanged.Broadcast(BrushComponent.Get(), EmptyPropertyChangedEvent);
 	}
+	else if (BoxComponent.IsValid() && !GeneratedResources.GeneratedComponents.Contains(BoxComponent) && OnToolStartBoxComponent)
+	{
+		// Restore the box component to its previous state
+		UEngine::CopyPropertiesForUnrelatedObjects(OnToolStartBoxComponent, BoxComponent.Get());
+
+		BoxComponent->SetWorldTransform(OnToolStartVolumeTransform);
+
+		FPropertyChangedEvent EmptyPropertyChangedEvent(nullptr);
+		FCoreUObjectDelegates::OnObjectPropertyChanged.Broadcast(BoxComponent.Get(), EmptyPropertyChangedEvent);
+
+		// Since we've changed the transform and the box might have been the root, we need to trigger a change on the actor.
+		if (AActor* Owner = BoxComponent->GetOwner())
+		{
+			Owner->PostEditChange();
+		}
+	}
 
 	Super::OnToolCancel(Context);
 }
@@ -82,6 +166,12 @@ void FPCGInteractiveToolWorkingData_Volume::OnToolShutdown(const FPCGInteractive
 
 	OnToolStartVolumeTransform = FTransform::Identity;
 
+	if (OnToolStartBoxComponent)
+	{
+		OnToolStartBoxComponent->MarkAsGarbage();
+		OnToolStartBoxComponent = nullptr;
+	}
+
 	Super::OnToolShutdown(Context);
 }
 
@@ -112,9 +202,9 @@ UBrushComponent* FPCGInteractiveToolWorkingData_Volume::FindMatchingBrushCompone
 	if (!Context.DataInstanceIdentifier.IsNone())
 	{
 		AllBrushComponents.RemoveAll([DataInstanceIdentifier = Context.DataInstanceIdentifier](UBrushComponent* Candidate)
-			{
-				return !Candidate || !Candidate->ComponentHasTag(DataInstanceIdentifier);
-			});
+		{
+			return !Candidate || !Candidate->ComponentHasTag(DataInstanceIdentifier);
+		});
 	}
 
 	MatchingComponent = AllBrushComponents.IsEmpty() ? nullptr : AllBrushComponents[0];
@@ -159,4 +249,33 @@ UBrushComponent* FPCGInteractiveToolWorkingData_Volume::GenerateMatchingBrushCom
 
 	return MatchingBrushComponent;
 }
+
+UBoxComponent* FPCGInteractiveToolWorkingData_Volume::FindBoxComponent(const FPCGInteractiveToolWorkingDataContext& Context)
+{
+	TArray<UBoxComponent*> AllBoxComponents;
+	if (Context.OwningActor.IsValid())
+	{
+		Context.OwningActor->GetComponents<UBoxComponent>(AllBoxComponents, false);
+	}
+
+	UBoxComponent* MatchingComponent = nullptr;
+
+	if (!Context.DataInstanceIdentifier.IsNone())
+	{
+		AllBoxComponents.RemoveAll([DataInstanceIdentifier = Context.DataInstanceIdentifier](UBoxComponent* Candidate)
+		{
+			return !Candidate || !Candidate->ComponentHasTag(DataInstanceIdentifier);
+		});
+	}
+
+	MatchingComponent = AllBoxComponents.IsEmpty() ? nullptr : AllBoxComponents[0];
+
+	if (AllBoxComponents.Num() > 1 && MatchingComponent != nullptr)
+	{
+		UE_LOGFMT(LogPCG, Warning, "More than one box component found. Choosing the first available: {0}", MatchingComponent->GetName());
+	}
+
+	return MatchingComponent;
+}
+
 #endif // WITH_EDITOR
\ No newline at end of file
diff --git a/Engine/Plugins/PCG/Source/PCGEditor/Private/EditorMode/Tools/Volume/PCGVolumeTool.cpp b/Engine/Plugins/PCG/Source/PCGEditor/Private/EditorMode/Tools/Volume/PCGVolumeTool.cpp
--- a/Engine/Plugins/PCG/Source/PCGEditor/Private/EditorMode/Tools/Volume/PCGVolumeTool.cpp
+++ b/Engine/Plugins/PCG/Source/PCGEditor/Private/EditorMode/Tools/Volume/PCGVolumeTool.cpp
@@ -13,6 +13,7 @@
 #include "ActorFactories/ActorFactory.h"
 #include "BaseBehaviors/SingleClickOrDragBehavior.h"
 #include "BaseBehaviors/MouseHoverBehavior.h"
+#include "Components/BoxComponent.h"
 #include "Builders/CubeBuilder.h"
 #include "GameFramework/Volume.h"
 #include "ToolTargets/ToolTarget.h"
@@ -54,8 +55,8 @@ const UScriptStruct* UPCGInteractiveToolSettings_Volume::GetWorkingDataType() co
 
 bool UPCGInteractiveToolSettings_Volume::IsWorkingActorCompatible(const AActor* InActor) const
 {
-	// Working actor needs to be a volume
-	return Super::IsWorkingActorCompatible(InActor) && InActor && InActor->IsA<AVolume>();
+	// Working actor needs to be a volume or have a box component
+	return Super::IsWorkingActorCompatible(InActor) && InActor && (InActor->IsA<AVolume>() || InActor->GetComponentByClass<UBoxComponent>() != nullptr);
 }
 
 AVolume* UPCGInteractiveToolSettings_Volume::GetWorkingVolume(FName DataInstanceIdentifier) const
@@ -68,11 +69,35 @@ AVolume* UPCGInteractiveToolSettings_Volume::GetWorkingVolume(FName DataInstance
 	return nullptr;
 }
 
+UBoxComponent* UPCGInteractiveToolSettings_Volume::GetWorkingBoxComponent(FName DataInstanceIdentifier) const
+{
+	if (const FPCGInteractiveToolWorkingData_Volume* ToolData = GetTypedWorkingData<FPCGInteractiveToolWorkingData_Volume>(DataInstanceIdentifier))
+	{
+		return ToolData->BoxComponent.Get();
+	}
+
+	return nullptr;
+}
+
 void UPCGInteractiveToolSettings_Volume::SetActorClassToSpawn(TSubclassOf<AActor> InClass)
 {
 	TSubclassOf<AActor> Class = InClass;
 
-	if (!Class || !Class->IsChildOf(AVolume::StaticClass()))
+	bool bClassIsValid = (Class != nullptr);
+
+	// Accept volumes or actors that have a box component
+	if (bClassIsValid)
+	{
+		if (!Class->IsChildOf(AVolume::StaticClass()))
+		{
+			if (const AActor* ActorCDO = Class->GetDefaultObject<AActor>())
+			{
+				bClassIsValid &= ActorCDO->GetComponentByClass<UBoxComponent>() != nullptr;
+			}
+		}
+	}
+
+	if(!bClassIsValid)
 	{
 		UE_LOG(LogPCGEditor, Error, TEXT("Invalid class used in volume tool. Will use PCGVolume instead."));
 		Class = APCGVolume::StaticClass();
@@ -137,26 +162,26 @@ void UPCGVolumeTool::OnTick(const float DeltaTime)
 
 void UPCGVolumeTool::UpdateVolumeFromBox()
 {
-	if (AVolume* WorkingVolume = ToolSettings ? ToolSettings->GetWorkingVolume(ToolSettings->DataInstance) : nullptr)
+	const FVector BoxCenter = LastBox.GetCenter();
+	FVector Extent = 2.0 * LastBox.GetExtent();
+
+	// if the user is using snapping, we still need to make sure that the final location would cover the full extent.
+	// This is because the brush builder will enforce snapping.
+	if (FSnappingUtils::IsSnapToGridEnabled())
 	{
-		const FVector BoxCenter = LastBox.GetCenter();
-		FVector Extent = 2.0 * LastBox.GetExtent();
+		FVector SnappedCenter = BoxCenter;
+		FRotator DummyRotator = FRotator::ZeroRotator;
+		FSnappingUtils::SnapToBSPVertex(SnappedCenter, FVector::ZeroVector, DummyRotator);
+
+		// We will not change the X/Y coordinates, but will adjust the Z coordinate so that we're always covering the bottom of the original box
+		Extent.Z += 2.0 * FMath::Abs(SnappedCenter.Z - BoxCenter.Z);
+	}
 
+	if (AVolume* WorkingVolume = ToolSettings ? ToolSettings->GetWorkingVolume(ToolSettings->DataInstance) : nullptr)
+	{
 		// Move actor to bounds center
 		WorkingVolume->SetActorTransform(FTransform(BoxCenter));
 
-		// if the user is using snapping, we still need to make sure that the final location would cover the full extent.
-		// This is because the brush builder will enforce snapping.
-		if (FSnappingUtils::IsSnapToGridEnabled())
-		{
-			FVector SnappedCenter = BoxCenter;
-			FRotator DummyRotator = FRotator::ZeroRotator;
-			FSnappingUtils::SnapToBSPVertex(SnappedCenter, FVector::ZeroVector, DummyRotator);
-
-			// We will not change the X/Y coordinates, but will adjust the Z coordinate so that we're always covering the bottom of the original box
-			Extent.Z += 2.0 * FMath::Abs(SnappedCenter.Z - BoxCenter.Z);
-		}
-
 		// setup cube builder
 		UCubeBuilder* Builder = NewObject<UCubeBuilder>();
 		Builder->X = Extent.X;
@@ -164,6 +189,22 @@ void UPCGVolumeTool::UpdateVolumeFromBox()
 		Builder->Z = Extent.Z;
 		UActorFactory::CreateBrushForVolumeActor(WorkingVolume, Builder);
 	}
+	else if (UBoxComponent* WorkingBoxComponent = ToolSettings ? ToolSettings->GetWorkingBoxComponent(ToolSettings->DataInstance) : nullptr)
+	{
+		FTransform PreviousBoxTransform = WorkingBoxComponent->GetComponentTransform();
+		WorkingBoxComponent->SetWorldTransform(FTransform(BoxCenter));
+
+		const FVector PreviousBoxScale = PreviousBoxTransform.GetScale3D();
+		const FVector InvScaleFactor(
+			(FMath::Abs(PreviousBoxScale.X) > UE_DOUBLE_SMALL_NUMBER ? (1.0 / PreviousBoxScale.X) : 0.0),
+			(FMath::Abs(PreviousBoxScale.Y) > UE_DOUBLE_SMALL_NUMBER ? (1.0 / PreviousBoxScale.Y) : 0.0),
+			(FMath::Abs(PreviousBoxScale.Z) > UE_DOUBLE_SMALL_NUMBER ? (1.0 / PreviousBoxScale.Z) : 0.0));
+
+		WorkingBoxComponent->SetBoxExtent(0.5 * Extent * InvScaleFactor);
+
+		FPropertyChangedEvent EmptyPropertyChangedEvent(nullptr);
+		FCoreUObjectDelegates::OnObjectPropertyChanged.Broadcast(WorkingBoxComponent, EmptyPropertyChangedEvent);
+	}
 }
 
 FBox UPCGVolumeTool::GetCurrentBox() const
@@ -223,9 +264,9 @@ FBox UPCGVolumeTool::GetCurrentBox() const
 
 bool UPCGVolumeTool::CanAccept() const
 {
-	AVolume* WorkingVolume = ToolSettings ? ToolSettings->GetWorkingVolume(ToolSettings->DataInstance) : nullptr;
 	// @todo_pcg also make sure we're in a correct state (e.g. not being interacted with at this time).
-	return WorkingVolume != nullptr;
+	return ToolSettings &&
+		(ToolSettings->GetWorkingVolume(ToolSettings->DataInstance) != nullptr || ToolSettings->GetWorkingBoxComponent(ToolSettings->DataInstance) != nullptr);
 }
 
 FInputRayHit UPCGVolumeTool::IsHitByClick(const FInputDeviceRay& ClickPosition)
diff --git a/Engine/Plugins/PCG/Source/PCGEditor/Private/EditorMode/Tools/Volume/PCGVolumeTool.h b/Engine/Plugins/PCG/Source/PCGEditor/Private/EditorMode/Tools/Volume/PCGVolumeTool.h
--- a/Engine/Plugins/PCG/Source/PCGEditor/Private/EditorMode/Tools/Volume/PCGVolumeTool.h
+++ b/Engine/Plugins/PCG/Source/PCGEditor/Private/EditorMode/Tools/Volume/PCGVolumeTool.h
@@ -10,6 +10,7 @@
 #include "PCGVolumeTool.generated.h"
 
 class AVolume;
+class UBoxComponent;
 class USingleClickOrDragInputBehavior;
 class UMouseHoverBehavior;
 
@@ -36,6 +37,7 @@ public:
 	// ~End UPCGInteractiveToolSettings interface
 
 	AVolume* GetWorkingVolume(FName DataInstanceIdentifier) const;
+	UBoxComponent* GetWorkingBoxComponent(FName DataInstanceIdentifier) const;
 	// @todo_pcg: Origin modes: Center, Corner
 	// @todo_pcg: Snap to grid
 
