[PCG] Fix issue where RG Change Detector could collect stale detection inputs in its StateCache, causing it to avoid re-scanning cells when input values are reused.

Properly re-scan cells for generation when generation radii change

#rnx

diff --git a/Engine/Plugins/PCG/Source/PCG/Internal/RuntimeGen/PCGRuntimeGenChangeDetection.h b/Engine/Plugins/PCG/Source/PCG/Internal/RuntimeGen/PCGRuntimeGenChangeDetection.h
--- a/Engine/Plugins/PCG/Source/PCG/Internal/RuntimeGen/PCGRuntimeGenChangeDetection.h
+++ b/Engine/Plugins/PCG/Source/PCG/Internal/RuntimeGen/PCGRuntimeGenChangeDetection.h
@@ -12,20 +12,21 @@ class UPCGComponent;
 
 namespace PCGRuntimeGenChangeDetection
 {
-	using FDetectionInputsBase = TTuple<TObjectKey<UPCGComponent>, const IPCGGenSourceBase*, /*RadiusMultiplier*/float, /*UseGenSourceDirection*/bool, /*MinGridSize*/uint32>;
+	using FDetectionInputsBase = TTuple<TObjectKey<UPCGComponent>, const IPCGGenSourceBase*, FPCGRuntimeGenerationRadii, /*RadiusMultiplier*/float, /*UseGenSourceDirection*/bool, /*MinGridSize*/uint32>;
 
 	/** The key used for the state cache. */
 	struct FDetectionInputs : FDetectionInputsBase
 	{
-		FDetectionInputs(const UPCGComponent* InOriginalComponent, const IPCGGenSourceBase* InGenSource, float InRadiusMultiplier, bool bInUseGenSourceDirection, uint32 InMinGridSize)
-			: FDetectionInputsBase(InOriginalComponent, InGenSource, InRadiusMultiplier, bInUseGenSourceDirection, InMinGridSize)
+		FDetectionInputs(const UPCGComponent* InOriginalComponent, const IPCGGenSourceBase* InGenSource, const FPCGRuntimeGenerationRadii& InGenerationRadii, float InRadiusMultiplier, bool bInUseGenSourceDirection, uint32 InMinGridSize)
+			: FDetectionInputsBase(InOriginalComponent, InGenSource, InGenerationRadii, InRadiusMultiplier, bInUseGenSourceDirection, InMinGridSize)
 		{}
 
 		const UPCGComponent* GetOriginalComponent() const { return Get<0>().ResolveObjectPtr(); }
 		const IPCGGenSourceBase* GetGenSource() const { return Get<1>(); }
-		float GetGenerationRadiusMultiplier() const { return Get<2>(); }
-		bool UseGenSourceDirection() const { return Get<3>(); }
-		uint32 GetMinGridSize() const { return Get<4>(); }
+		const FPCGRuntimeGenerationRadii& GetGenerationRadii() const { return Get<2>(); }
+		float GetGenerationRadiusMultiplier() const { return Get<3>(); }
+		bool UseGenSourceDirection() const { return Get<4>(); }
+		uint32 GetMinGridSize() const { return Get<5>(); }
 	};
 
 	/** Used to detect when world state (generation sources, cvars etc) have changed enough to warrant rescanning generation cells. */
@@ -43,6 +44,9 @@ namespace PCGRuntimeGenChangeDetection
 			float GenerationRadiusMultiplier = 1.0f;
 		};
 
+		void PreTick();
+		void PostTick();
+
 		/** Compare current state against cached state to determine if a rescan of cells is required. */
 		bool IsCellScanRequired(const FDetectionInputs& InInputs);
 
@@ -54,5 +58,8 @@ namespace PCGRuntimeGenChangeDetection
 
 	private:
 		TMap<FDetectionInputs, FCachedState> StateCache;
+
+		/** Any inputs in this array will be removed from the StateCache during PostTick(). */
+		TArray<FDetectionInputs> InputsPendingDeletion;
 	};
 }
diff --git a/Engine/Plugins/PCG/Source/PCG/Private/PCGCommon.cpp b/Engine/Plugins/PCG/Source/PCG/Private/PCGCommon.cpp
--- a/Engine/Plugins/PCG/Source/PCG/Private/PCGCommon.cpp
+++ b/Engine/Plugins/PCG/Source/PCG/Private/PCGCommon.cpp
@@ -228,6 +228,42 @@ double FPCGRuntimeGenerationRadii::GetCleanupRadiusFromGrid(EPCGHiGenGrid Grid)
 	return GetGenerationRadiusFromGrid(Grid) * CleanupRadiusMultiplier;
 }
 
+bool FPCGRuntimeGenerationRadii::operator==(const FPCGRuntimeGenerationRadii& Other) const
+{
+	return GenerationRadius == Other.GenerationRadius
+		&& GenerationRadius400 == Other.GenerationRadius400
+		&& GenerationRadius800 == Other.GenerationRadius800
+		&& GenerationRadius1600 == Other.GenerationRadius1600
+		&& GenerationRadius3200 == Other.GenerationRadius3200
+		&& GenerationRadius6400 == Other.GenerationRadius6400
+		&& GenerationRadius12800 == Other.GenerationRadius12800
+		&& GenerationRadius25600 == Other.GenerationRadius25600
+		&& GenerationRadius51200 == Other.GenerationRadius51200
+		&& GenerationRadius102400 == Other.GenerationRadius102400
+		&& GenerationRadius204800 == Other.GenerationRadius204800
+		&& CleanupRadiusMultiplier == Other.CleanupRadiusMultiplier;
+}
+
+uint32 GetTypeHash(const FPCGRuntimeGenerationRadii& InGenerationRadii)
+{
+	uint32 Hash = 0;
+
+	Hash = HashCombine(Hash, InGenerationRadii.GenerationRadius);
+	Hash = HashCombine(Hash, InGenerationRadii.GenerationRadius400);
+	Hash = HashCombine(Hash, InGenerationRadii.GenerationRadius800);
+	Hash = HashCombine(Hash, InGenerationRadii.GenerationRadius1600);
+	Hash = HashCombine(Hash, InGenerationRadii.GenerationRadius3200);
+	Hash = HashCombine(Hash, InGenerationRadii.GenerationRadius6400);
+	Hash = HashCombine(Hash, InGenerationRadii.GenerationRadius12800);
+	Hash = HashCombine(Hash, InGenerationRadii.GenerationRadius25600);
+	Hash = HashCombine(Hash, InGenerationRadii.GenerationRadius51200);
+	Hash = HashCombine(Hash, InGenerationRadii.GenerationRadius102400);
+	Hash = HashCombine(Hash, InGenerationRadii.GenerationRadius204800);
+	Hash = HashCombine(Hash, InGenerationRadii.CleanupRadiusMultiplier);
+
+	return Hash;
+}
+
 namespace PCGDelegates
 {
 #if WITH_EDITOR
diff --git a/Engine/Plugins/PCG/Source/PCG/Private/PCGComponent.cpp b/Engine/Plugins/PCG/Source/PCG/Private/PCGComponent.cpp
--- a/Engine/Plugins/PCG/Source/PCG/Private/PCGComponent.cpp
+++ b/Engine/Plugins/PCG/Source/PCG/Private/PCGComponent.cpp
@@ -1348,6 +1348,11 @@ void UPCGComponent::SetSchedulingPolicyClass(TSubclassOf<UPCGSchedulingPolicyBas
 	RefreshSchedulingPolicy();
 }
 
+const FPCGRuntimeGenerationRadii& UPCGComponent::GetGenerationRadii() const
+{
+	return (bOverrideGenerationRadii || !GetGraph()) ? GenerationRadii : GetGraph()->GenerationRadii;
+}
+
 double UPCGComponent::GetGenerationRadiusFromGrid(EPCGHiGenGrid Grid) const
 {
 	const double Multiplier = FMath::Max(PCGRuntimeGenSchedulerHelpers::CVarRuntimeGenerationRadiusMultiplier.GetValueOnAnyThread(), 0.0f);
@@ -2509,7 +2514,7 @@ void UPCGComponent::PostEditChangeProperty(FPropertyChangedEvent& PropertyChange
 
 	const FName MemberName = PropertyChangedEvent.MemberProperty->GetFName();
 
-	if (MemberName == GET_MEMBER_NAME_CHECKED(UPCGComponent, GenerationRadii))
+	if (MemberName == GET_MEMBER_NAME_CHECKED(UPCGComponent, GenerationRadii) || MemberName == GET_MEMBER_NAME_CHECKED(UPCGComponent, bOverrideGenerationRadii))
 	{
 		// RuntimeGen will automatically pick up any changes to generation radii, we don't need to do any work here.
 		return;
diff --git a/Engine/Plugins/PCG/Source/PCG/Private/RuntimeGen/PCGRuntimeGenChangeDetection.cpp b/Engine/Plugins/PCG/Source/PCG/Private/RuntimeGen/PCGRuntimeGenChangeDetection.cpp
--- a/Engine/Plugins/PCG/Source/PCG/Private/RuntimeGen/PCGRuntimeGenChangeDetection.cpp
+++ b/Engine/Plugins/PCG/Source/PCG/Private/RuntimeGen/PCGRuntimeGenChangeDetection.cpp
@@ -82,8 +82,28 @@ namespace PCGRuntimeGenChangeDetection
 		return false;
 	}
 
+	void FDetector::PreTick()
+	{
+		// Mark all inputs in the StateCache as pending deletion. They will be cleaned out unless directly queried by the RGS.
+		StateCache.GetKeys(InputsPendingDeletion);
+	}
+
+	void FDetector::PostTick()
+	{
+		// Clear the cache state for any inputs that were not refreshed during this tick.
+		for (const FDetectionInputs& DetectionInput : InputsPendingDeletion)
+		{
+			StateCache.Remove(DetectionInput);
+		}
+
+		InputsPendingDeletion.Reset();
+	}
+
 	bool FDetector::IsCellScanRequired(const FDetectionInputs& InInputs)
 	{
+		// These inputs are active, therefore they are no longer pending deletion at the end of this tick.
+		InputsPendingDeletion.Remove(InInputs);
+		
 		// If no cached state was present it will be added here and ShouldScanCells will return true in this case.
 		return StateCache.FindOrAdd(InInputs).ShouldScanCells(InInputs);
 	}
diff --git a/Engine/Plugins/PCG/Source/PCG/Private/RuntimeGen/PCGRuntimeGenScheduler.cpp b/Engine/Plugins/PCG/Source/PCG/Private/RuntimeGen/PCGRuntimeGenScheduler.cpp
--- a/Engine/Plugins/PCG/Source/PCG/Private/RuntimeGen/PCGRuntimeGenScheduler.cpp
+++ b/Engine/Plugins/PCG/Source/PCG/Private/RuntimeGen/PCGRuntimeGenScheduler.cpp
@@ -316,7 +316,9 @@ void FPCGRuntimeGenScheduler::Tick(APCGWorldActor* InPCGWorldActor, double InEnd
 		Inputs.AllNonPartitionedComponents = ActorAndComponentMapping->GetAllRegisteredNonPartitionedComponents();
 		Inputs.GeneratedComponents = GeneratedComponents;
 
+		ChangeDetector.PreTick();
 		TickQueueComponentsForGeneration(Inputs, ComponentsToGenerate);
+		ChangeDetector.PostTick();
 	}
 
 	// 2. Schedule cleanup on components that become out of range.
@@ -557,6 +559,7 @@ void FPCGRuntimeGenScheduler::TickQueueComponentsForGeneration(
 				{
 					OriginalComponent,
 					GenSource,
+					OriginalComponent->GetGenerationRadii(),
 					PCGRuntimeGenSchedulerHelpers::CVarRuntimeGenerationRadiusMultiplier.GetValueOnGameThread(),
 					Policy->CullsBasedOnDirection(),
 					GridSizes.IsEmpty() ? PCGHiGenGrid::UnboundedGridSize() : GridSizes.Last(),
diff --git a/Engine/Plugins/PCG/Source/PCG/Public/PCGCommon.h b/Engine/Plugins/PCG/Source/PCG/Public/PCGCommon.h
--- a/Engine/Plugins/PCG/Source/PCG/Public/PCGCommon.h
+++ b/Engine/Plugins/PCG/Source/PCG/Public/PCGCommon.h
@@ -446,6 +446,8 @@ public:
 	/** Compute the runtime cleanup radius for the given grid size. */
 	UE_API double GetCleanupRadiusFromGrid(EPCGHiGenGrid Grid) const;
 
+	UE_API bool operator==(const FPCGRuntimeGenerationRadii& Other) const;
+
 	static constexpr double DefaultGenerationRadiusMultiplier = 2.0;
 	static constexpr double DefaultCleanupRadiusMultiplier = 1.1;
 
@@ -488,6 +490,8 @@ public:
 	double CleanupRadiusMultiplier = DefaultCleanupRadiusMultiplier;
 };
 
+UE_API uint32 GetTypeHash(const FPCGRuntimeGenerationRadii& InGenerationRadii);
+
 struct FInstancedPropertyBag;
 
 namespace PCGDelegates
diff --git a/Engine/Plugins/PCG/Source/PCG/Public/PCGComponent.h b/Engine/Plugins/PCG/Source/PCG/Public/PCGComponent.h
--- a/Engine/Plugins/PCG/Source/PCG/Public/PCGComponent.h
+++ b/Engine/Plugins/PCG/Source/PCG/Public/PCGComponent.h
@@ -317,6 +317,9 @@ public:
 	/** Set the runtime generation scheduling policy type. */
 	UE_API void SetSchedulingPolicyClass(TSubclassOf<UPCGSchedulingPolicyBase> InSchedulingPolicyClass);
 
+	/** Get the generation radii that are currently active for this component. */
+	UE_API const FPCGRuntimeGenerationRadii& GetGenerationRadii() const;
+
 	/** Get the runtime generation radius for the given grid size. */
 	UE_API double GetGenerationRadiusFromGrid(EPCGHiGenGrid Grid) const;
 
