[PCG] Soften crash in FPCGGraphCompiler::PostCullStackCleanup

#rnx

diff --git a/Engine/Plugins/PCG/Source/PCG/Private/Graph/PCGGraphCompiler.cpp b/Engine/Plugins/PCG/Source/PCG/Private/Graph/PCGGraphCompiler.cpp
--- a/Engine/Plugins/PCG/Source/PCG/Private/Graph/PCGGraphCompiler.cpp
+++ b/Engine/Plugins/PCG/Source/PCG/Private/Graph/PCGGraphCompiler.cpp
@@ -1098,7 +1098,10 @@ void FPCGGraphCompiler::PostCullStackCleanup(TArray<FPCGGraphTask>& InCompiledTa
 		// Remap stack IDs.
 		for (FPCGGraphTask& Task : InCompiledTasks)
 		{
-			Task.StackIndex -= RemovedCounts[Task.StackIndex];
+			if (ensure(RemovedCounts.IsValidIndex(Task.StackIndex)))
+			{
+				Task.StackIndex -= RemovedCounts[Task.StackIndex];
+			}
 		}
 	}
 }
