[PCG] Soften crash in PCGDataCollectionPackingHelpers::ComputePackedSize()

Fix data sanitization in UPCGDataBinding initialization, which was allowing invalid types to flow through compute graphs

#rnx

diff --git a/Engine/Plugins/PCG/Source/PCG/Private/Compute/PCGComputeGraph.h b/Engine/Plugins/PCG/Source/PCG/Private/Compute/PCGComputeGraph.h
--- a/Engine/Plugins/PCG/Source/PCG/Private/Compute/PCGComputeGraph.h
+++ b/Engine/Plugins/PCG/Source/PCG/Private/Compute/PCGComputeGraph.h
@@ -54,6 +54,9 @@ public:
 	/** Get the virtual label associated with a binding index/kernel pin. This only exists if the kernel pin is an input pin on the border of the compute graph. */
 	const FName* GetVirtualPinLabelFromBindingIndex(int32 InBindingIndex) const { return CPUDataBindingToVirtualPinLabel.Find(InBindingIndex); }
 
+	/** Get the PCG data type associated with the virtual input pin label. */
+	const FPCGDataTypeIdentifier* GetVirtualInputPinType(FName InVirtualInputPinLabel) const { return VirtualInputPinToType.Find(InVirtualInputPinLabel); }
+
 	/** Populates an array with all of the kernel pins in the compute graph. */
 	void GetKernelPins(TArray<FPCGKernelPin>& OutKernelPins) const;
 
@@ -129,5 +132,9 @@ protected:
 	UPROPERTY()
 	TMap</*BindingIndex*/int32, /*VirtualPinLabel*/FName> CPUDataBindingToVirtualPinLabel;
 
+	/** Maps virtual input pins to their PCG data type. */
+	UPROPERTY()
+	TMap</*VirtualPinLabel*/FName, FPCGDataTypeIdentifier> VirtualInputPinToType;
+
 	friend class FPCGGraphCompilerGPU;
 };
diff --git a/Engine/Plugins/PCG/Source/PCG/Private/Compute/PCGDataBinding.cpp b/Engine/Plugins/PCG/Source/PCG/Private/Compute/PCGDataBinding.cpp
--- a/Engine/Plugins/PCG/Source/PCG/Private/Compute/PCGDataBinding.cpp
+++ b/Engine/Plugins/PCG/Source/PCG/Private/Compute/PCGDataBinding.cpp
@@ -33,14 +33,27 @@ void UPCGDataBinding::Initialize(const UPCGComputeGraph* InComputeGraph, FPCGCon
 	InputDataCollection.TaggedData.Reserve(InContext->InputData.TaggedData.Num());
 	for (const FPCGTaggedData& InputData : InContext->InputData.TaggedData)
 	{
-		if (InputData.Data && PCGComputeHelpers::IsTypeAllowedAsInput(FPCGDataTypeIdentifier{InputData.Data->GetDataTypeId()}))
+		if (!InputData.Data)
 		{
-			InputDataCollection.TaggedData.Add(InputData);
+			continue;
 		}
-		else if (InputData.Data)
+		
+		if (!PCGComputeHelpers::IsTypeAllowedAsInput(FPCGDataTypeIdentifier{InputData.Data->GetDataTypeId()}))
 		{
-			UE_LOG(LogPCG, Warning, TEXT("Stripped input data that is not currently supported by GPU execution: %s"), *InputData.Data->GetName());
+			UE_LOG(LogPCG, Warning, TEXT("Stripped input data '%s' that is not currently supported by GPU execution."), *InputData.Data->GetName());
+			continue;
 		}
+
+		const FPCGDataTypeIdentifier* VirtualInputPinType = Graph->GetVirtualInputPinType(InputData.Pin);
+		const FPCGDataTypeIdentifier InputDataType = InputData.Data->GetUnderlyingDataTypeId();
+
+		if (!VirtualInputPinType || !InputDataType.IsChildOf(*VirtualInputPinType))
+		{
+			UE_LOG(LogPCG, Warning, TEXT("Stripped input data '%s' that is not the correct type for pin '%s'."), *InputData.Data->GetName(), *InputData.Pin.ToString());
+			continue;
+		}
+
+		InputDataCollection.TaggedData.Add(InputData);
 	}
 }
 
diff --git a/Engine/Plugins/PCG/Source/PCG/Private/Compute/Packing/PCGDataCollectionPacking.cpp b/Engine/Plugins/PCG/Source/PCG/Private/Compute/Packing/PCGDataCollectionPacking.cpp
--- a/Engine/Plugins/PCG/Source/PCG/Private/Compute/Packing/PCGDataCollectionPacking.cpp
+++ b/Engine/Plugins/PCG/Source/PCG/Private/Compute/Packing/PCGDataCollectionPacking.cpp
@@ -488,11 +488,18 @@ namespace PCGDataCollectionPackingHelpers
 
 	uint64 ComputePackedSizeBytes(const FPCGDataDesc& InDataDesc)
 	{
-		ensure(InDataDesc.GetElementDimension() == EPCGElementDimension::One);
+		if (!ensure(InDataDesc.GetElementDimension() == EPCGElementDimension::One))
+		{
+			return 0u;
+		}
 
 		const FPCGDataTypeIdentifier& Type = InDataDesc.GetType();
 
-		check(PCGComputeHelpers::IsTypeAllowedInDataCollection(Type));
+		if (!ensure(PCGComputeHelpers::IsTypeAllowedInDataCollection(Type)))
+		{
+			return 0u;
+		}
+
 		uint64 DataSizeBytes = PCGDataCollectionPackingConstants::DATA_HEADER_SIZE_BYTES;
 
 		for (const FPCGKernelAttributeDesc& AttributeDesc : InDataDesc.GetAttributeDescriptions())
diff --git a/Engine/Plugins/PCG/Source/PCG/Private/Graph/PCGGraphCompilerGPU.cpp b/Engine/Plugins/PCG/Source/PCG/Private/Graph/PCGGraphCompilerGPU.cpp
--- a/Engine/Plugins/PCG/Source/PCG/Private/Graph/PCGGraphCompilerGPU.cpp
+++ b/Engine/Plugins/PCG/Source/PCG/Private/Graph/PCGGraphCompilerGPU.cpp
@@ -1338,6 +1338,7 @@ void FPCGGraphCompilerGPU::CompileComputeGraph(
 			if (const FName* VirtualPin = InOriginalToVirtualPin.Find(FNodePin(TaskId, Input.DownstreamPin->Label, /*IsInput*/true)))
 			{
 				OriginalToVirtualInputPin.Add(Input.DownstreamPin->Label, *VirtualPin);
+				InOutComputeGraph->VirtualInputPinToType.FindOrAdd(*VirtualPin) = Input.DownstreamPin->AllowedTypes;
 			}
 
 			CreateDataResolverDataInterface(InputDataInterfaces, Input.DownstreamPin->Label, /*bIsInput=*/true);
