[PCG] Fix Filter Elements by Index droping unallocated properties

#rnx

diff --git a/Engine/Plugins/PCG/Source/PCG/Private/Elements/PCGFilterElementsByIndex.cpp b/Engine/Plugins/PCG/Source/PCG/Private/Elements/PCGFilterElementsByIndex.cpp
--- a/Engine/Plugins/PCG/Source/PCG/Private/Elements/PCGFilterElementsByIndex.cpp
+++ b/Engine/Plugins/PCG/Source/PCG/Private/Elements/PCGFilterElementsByIndex.cpp
@@ -297,6 +297,7 @@ bool FPCGFilterElementsByIndexElement::ExecuteInternal(FPCGContext* InContext) c
 			const int32 NumInFilterPoints = Settings->bInvertFilter ? InputPointData->GetNumPoints() - AllFilteredIndices.GetTotalIndexCount() : AllFilteredIndices.GetTotalIndexCount();
 			InFilterPointData->SetNumPoints(NumInFilterPoints);
 			InFilterPointData->AllocateProperties(InputPointData->GetAllocatedProperties());
+			InFilterPointData->CopyUnallocatedPropertiesFrom(InputPointData);
 
 			AddNewOutput(InFilterPointData, Constants::InFilterOutputPinLabel);
 
