[PCG] Fix infinite loop when refreshing the debug object tree view with an aborted component

#rnx

diff --git a/Engine/Plugins/PCG/Source/PCGEditor/Private/PCGEditor.cpp b/Engine/Plugins/PCG/Source/PCGEditor/Private/PCGEditor.cpp
--- a/Engine/Plugins/PCG/Source/PCGEditor/Private/PCGEditor.cpp
+++ b/Engine/Plugins/PCG/Source/PCGEditor/Private/PCGEditor.cpp
@@ -314,13 +314,19 @@ void FPCGEditor::UpdateAfterInspectedStackChanged(const FPCGStack& FullStack)
 			// Implementation note: if we're inspecting and have not pre-run the graph, then it probably makes sense to enable inspection by default. 
 			// TODO This could be selected with a cvar though.
 			const bool bHasBeenGeneratedThisSession = Component && Component->bGenerated && Component->WasGeneratedThisSession();
+			const bool bWasAborted = LastExecutionStatus.IsSet() && LastExecutionStatus.GetValue().Key == Source && LastExecutionStatus.GetValue().Value == EPCGGenerationStatus::Aborted;
 			const bool bWasInspecting = Source->GetExecutionState().GetInspection().IsInspecting();
 			const bool bNeedsInspection = Algo::AnyOf(AttributesWidgets, [](const TSharedPtr<SPCGEditorGraphAttributeListView>& ALV) { return ALV->GetNodeBeingInspected() != nullptr; });
 
 			if (!bHasBeenGeneratedThisSession || (bNeedsInspection && !bWasInspecting))
 			{
 				Source->GetExecutionState().GetInspection().EnableInspection();
-				UpdateDebugAfterComponentSelection(Component, Component, true);
+
+				// Making sure to not re-trigger a new generation if the previous one was aborted, to avoid infinite loops
+				if (!bWasAborted)
+				{
+					UpdateDebugAfterComponentSelection(Component, Component, true);
+				}
 			}
 		}
 	}
@@ -3172,8 +3178,12 @@ void FPCGEditor::OnSourceGenerationDone(IPCGBaseSubsystem* Subsystem, IPCGGraphE
 		return;
 	}
 
+	LastExecutionStatus = {Source, Status};
+
 	OnSourceGenerated(Source);
 
+	LastExecutionStatus.Reset();
+
 	const bool CacheDebuggingEnabled = Subsystem && Subsystem->IsGraphCacheDebuggingEnabled();
 
 	// Refresh nodes to report any errors/warnings, and to display culling state after execution.
diff --git a/Engine/Plugins/PCG/Source/PCGEditor/Public/PCGEditor.h b/Engine/Plugins/PCG/Source/PCGEditor/Public/PCGEditor.h
--- a/Engine/Plugins/PCG/Source/PCGEditor/Public/PCGEditor.h
+++ b/Engine/Plugins/PCG/Source/PCGEditor/Public/PCGEditor.h
@@ -523,4 +523,8 @@ private:
 	UPCGEditorGraph* PCGEditorGraph = nullptr;
 
 	FPCGEditorInspectionDataManager InspectionDataManager;
+
+	// Keep track of the last execution status to be able to break infinite loop when a source is triggered to be generated by inspection
+	// aborted and re-triggered.
+	TOptional<TPair<const IPCGGraphExecutionSource*, EPCGGenerationStatus>> LastExecutionStatus;
 };
