PCG Editor Mode
Added Preset buttons that set the Tool Graph

- Small changes to toolkitbuilder and related classes to expose functionality
-- This is needed because the preset tool area takes up 50% of the space for the active tool name with no way around it and we need more space.
-- At the same time, we want to use most of the default UI for consistency

- Added dedicated Preset section widget
-- Filters for compatibility with existing actor
-- Filters for tool compatibility
- Refactored utility functions
-- If an asset is loaded, uses up2date data
-- If not, uses Asset Registry data
- Added PCG Toolkit Builder to change preset display
- Added PCGPalette; currently functions as a regular palette

- PCGGraph now has bIsPreset option
- PCGGraphInstance now has ToolDataOverride member to override bIsPreset

diff --git a/Engine/Plugins/PCG/Source/PCG/Private/PCGGraph.cpp b/Engine/Plugins/PCG/Source/PCG/Private/PCGGraph.cpp
--- a/Engine/Plugins/PCG/Source/PCG/Private/PCGGraph.cpp
+++ b/Engine/Plugins/PCG/Source/PCG/Private/PCGGraph.cpp
@@ -2516,6 +2516,18 @@ void UPCGGraphInstance::TeardownCallbacks()
 	}
 }
 
+TOptional<FPCGGraphToolData> UPCGGraphInstance::GetGraphToolData() const
+{
+	TOptional<FPCGGraphToolData> GraphToolData = GetGraph() ? GetGraph()->GetGraphToolData() : TOptional<FPCGGraphToolData>();
+
+	if (GraphToolData.IsSet())
+	{
+		GraphToolData->ApplyOverrides(ToolDataOverrides);
+	}
+	
+	return GraphToolData;
+}
+
 void UPCGGraphInstance::SetupCallbacks()
 {
 	if (FPCGContext::IsInitializingSettings())
diff --git a/Engine/Plugins/PCG/Source/PCG/Public/Data/PCGToolData.h b/Engine/Plugins/PCG/Source/PCG/Public/Data/PCGToolData.h
--- a/Engine/Plugins/PCG/Source/PCG/Public/Data/PCGToolData.h
+++ b/Engine/Plugins/PCG/Source/PCG/Public/Data/PCGToolData.h
@@ -31,6 +31,22 @@ protected:
 	FName WorkingDataIdentifier = NAME_None;
 };
 
+USTRUCT()
+struct FPCGGraphInstanceToolDataOverrides
+{
+	GENERATED_BODY()
+
+	UPROPERTY(EditAnywhere, Category = Tool)
+	FText DisplayName;
+
+	UPROPERTY(EditAnywhere, Category = Tool, meta=(MultiLine))
+	FText Tooltip;
+	
+	/** Whether the graph instance this belongs to is a tool preset. */
+	UPROPERTY(EditAnywhere, Category = Tool)
+	bool bIsPreset = false;
+};
+
 /** The tool data stored on the graph itself. Informs the UI. */
 USTRUCT(BlueprintType)
 struct FPCGGraphToolData
@@ -42,6 +58,9 @@ struct FPCGGraphToolData
 	
 	UPROPERTY(EditAnywhere, Category = Tool)
 	FText DisplayName;
+
+	UPROPERTY(EditAnywhere, Category = Tool, meta=(MultiLine))
+	FText Tooltip;
 	
 	UPROPERTY(EditAnywhere, Category = Tool)
 	TArray<FString> CompatibleToolTags;
@@ -51,6 +70,27 @@ struct FPCGGraphToolData
 
 	UPROPERTY(EditAnywhere, Category = Tool)
 	FName NewActorLabel;
+
+	UPROPERTY(EditAnywhere, Category = Tool)
+	bool bIsPreset = false;
+
+	void ApplyOverrides(const FPCGGraphInstanceToolDataOverrides& Overrides)
+	{
+		// Only override the display name if it was actually specified
+		if (Overrides.DisplayName.IsEmpty())
+		{
+			DisplayName = Overrides.DisplayName;
+		}
+
+		// Only override the tooltip if it was actually specified
+		if (Overrides.Tooltip.IsEmpty())
+		{
+			Tooltip = Overrides.Tooltip;
+		}
+
+		// Always override the preset value
+		bIsPreset = Overrides.bIsPreset;
+	}
 };
 
 USTRUCT(BlueprintType)
diff --git a/Engine/Plugins/PCG/Source/PCG/Public/PCGGraph.h b/Engine/Plugins/PCG/Source/PCG/Public/PCGGraph.h
--- a/Engine/Plugins/PCG/Source/PCG/Public/PCGGraph.h
+++ b/Engine/Plugins/PCG/Source/PCG/Public/PCGGraph.h
@@ -698,7 +698,7 @@ public:
 	UE_API void SetupCallbacks();
 	UE_API void TeardownCallbacks();
 
-	virtual TOptional<FPCGGraphToolData> GetGraphToolData() const override { return GetGraph() ? GetGraph()->GetGraphToolData() : TOptional<FPCGGraphToolData>(); }
+	virtual TOptional<FPCGGraphToolData> GetGraphToolData() const override;
 #endif
 
 	UE_API virtual void OnGraphParametersChanged(EPCGGraphParameterEvent InChangeType, FName InChangedPropertyName) override;
@@ -734,6 +734,9 @@ public:
 	  * the value for this instance is different than the value from the parent.
 	  */
 	UE_API FString GetDefaultPropertyValueForEditor(const FProperty* InProperty, bool& bIsDifferent) const;
+
+	UFUNCTION()
+	bool ShowToolDataOverrides() const { return IsAsset(); }
 #endif // WITH_EDITOR
 
 	/** Because subgraph nodes hold graph instances, delegation of the selection can't be done elsewhere, so this function gets called when the dialog is opened and will forward to the graph customization if filtering is applied. */
@@ -780,6 +783,10 @@ public:
 	/** Can override the category of this instance. */
 	UPROPERTY(BlueprintReadWrite, EditAnywhere, Category = AssetInfo, AssetRegistrySearchable, meta = (EditCondition = "AreOverridesEnabled && bOverrideCategory", EditConditionHides, DisplayAfter = ColorOverride, PCGNoHash))
 	FText Category;
+
+	/** Override information for tool usage. */
+	UPROPERTY(EditAnywhere, Category = AssetInfo, AssetRegistrySearchable, meta = (PCGNoHash, EditCondition = "ShowToolDataOverrides()", EditConditionHides))
+	FPCGGraphInstanceToolDataOverrides ToolDataOverrides;
 #endif // WITH_EDITORONLY_DATA
 };
 
diff --git a/Engine/Plugins/PCG/Source/PCGEditor/PCGEditor.Build.cs b/Engine/Plugins/PCG/Source/PCGEditor/PCGEditor.Build.cs
--- a/Engine/Plugins/PCG/Source/PCGEditor/PCGEditor.Build.cs
+++ b/Engine/Plugins/PCG/Source/PCGEditor/PCGEditor.Build.cs
@@ -76,11 +76,6 @@ namespace UnrealBuildTool.Rules
 					"GeometryCore",
 					"GeometryFramework"
 				});
-			
-			// @todo_pcg
-			// Set to 1 for development where each tool graph creates a button
-			// Remove again when we settle on a definitive UI
-			PrivateDefinitions.Add("USE_PCG_MODE_GRAPHS_AS_BUTTONS=0");
 		}
 	}
 }
diff --git a/Engine/Plugins/PCG/Source/PCGEditor/Private/EditorMode/PCGEdMode.cpp b/Engine/Plugins/PCG/Source/PCGEditor/Private/EditorMode/PCGEdMode.cpp
--- a/Engine/Plugins/PCG/Source/PCGEditor/Private/EditorMode/PCGEdMode.cpp
+++ b/Engine/Plugins/PCG/Source/PCGEditor/Private/EditorMode/PCGEdMode.cpp
@@ -169,53 +169,31 @@ void UPCGEditorMode::Enter()
 		}
 	});
 
-#if USE_PCG_MODE_GRAPHS_AS_BUTTONS
-	for(TSharedPtr<FUICommandInfo> ToolCommandInfo : FPCGEditorModeToolCommands::Get().GetLineContextToolCommands())
+	FPCGEditorModeToolCommands& ToolManagerCommands = FPCGEditorModeToolCommands::Get();
+	
 	{
-		UPCGDrawSplineToolBuilder* SplineToolBuilder = NewObject<UPCGDrawSplineToolBuilder>();
-		SplineToolBuilder->SetToolGraph(CastChecked<UPCGGraph>(FSoftObjectPath(ToolCommandInfo->GetCommandName().ToString()).TryLoad()));
-		SplineToolBuilder->SetIsSplineSurface(false);
-		RegisterTool(ToolCommandInfo, ToolCommandInfo->GetCommandName().ToString(), SplineToolBuilder);
+		UPCGDrawSplineToolBuilder* DrawSplineToolBuilder = NewObject<UPCGDrawSplineToolBuilder>();
+		DrawSplineToolBuilder->SetSplineToolClass(UPCGDrawSplineTool::StaticClass());
+		RegisterTool(ToolManagerCommands.EnableDrawSplineTool, UPCGInteractiveToolSettings_Spline::StaticGetToolTag().ToString(), DrawSplineToolBuilder);
+	}
+	
+	{
+		UPCGDrawSplineToolBuilder* DrawSurfaceToolBuilder = NewObject<UPCGDrawSplineToolBuilder>();
+		DrawSurfaceToolBuilder->SetSplineToolClass(UPCGDrawSplineSurfaceTool::StaticClass());
+		RegisterTool(ToolManagerCommands.EnableDrawSurfaceTool, UPCGInteractiveToolSettings_SplineSurface::StaticGetToolTag().ToString(), DrawSurfaceToolBuilder);
 	}
-
-	UPCGDrawSplineToolBuilder* SplineSurfaceToolBuilder = NewObject<UPCGDrawSplineToolBuilder>();
-	SplineSurfaceToolBuilder->SetIsSplineSurface(false);
-
-	// Add spline surface tools here
 	
-	for(TSharedPtr<FUICommandInfo> ToolCommandInfo : FPCGEditorModeToolCommands::Get().GetPaintContextToolCommands())
 	{
 		UPCGPaintToolBuilder* PaintToolBuilder = NewObject<UPCGPaintToolBuilder>();
-		PaintToolBuilder->SetToolGraph(CastChecked<UPCGGraph>(FSoftObjectPath(ToolCommandInfo->GetCommandName().ToString()).TryLoad()));
-		
-		RegisterTool(ToolCommandInfo, ToolCommandInfo->GetCommandName().ToString(), PaintToolBuilder);
+		RegisterTool(ToolManagerCommands.EnablePaintTool, UPCGInteractiveToolSettings_PaintTool::StaticGetToolTag().ToString(), PaintToolBuilder);
 	}
-
-	for(TSharedPtr<FUICommandInfo> ToolCommandInfo : FPCGEditorModeToolCommands::Get().GetVolumeContextToolCommands())
+	
 	{
 		UPCGVolumeToolBuilder* VolumeToolBuilder = NewObject<UPCGVolumeToolBuilder>();
-		VolumeToolBuilder->SetToolGraph(CastChecked<UPCGGraph>(FSoftObjectPath(ToolCommandInfo->GetCommandName().ToString()).TryLoad()));
-		RegisterTool(ToolCommandInfo, ToolCommandInfo->GetCommandName().ToString(), VolumeToolBuilder);
+		RegisterTool(ToolManagerCommands.EnableVolumeTool, UPCGInteractiveToolSettings_Volume::StaticGetToolTag().ToString(), VolumeToolBuilder);
 	}
-#else
-	const FPCGEditorModeToolCommands& ToolManagerCommands = FPCGEditorModeToolCommands::Get();
-
-	UPCGDrawSplineToolBuilder* DrawSplineToolBuilder = NewObject<UPCGDrawSplineToolBuilder>();
-	DrawSplineToolBuilder->SetSplineToolClass(UPCGDrawSplineTool::StaticClass());
-	RegisterTool(ToolManagerCommands.EnableDrawSplineTool, TEXT("DrawSpline"), DrawSplineToolBuilder);
-
-	UPCGDrawSplineToolBuilder* DrawSurfaceToolBuilder = NewObject<UPCGDrawSplineToolBuilder>();
-	DrawSurfaceToolBuilder->SetSplineToolClass(UPCGDrawSplineSurfaceTool::StaticClass());
-	RegisterTool(ToolManagerCommands.EnableDrawSurfaceTool, TEXT("DrawSurface"), DrawSurfaceToolBuilder);
-
-	UPCGPaintToolBuilder* PaintToolBuilder = NewObject<UPCGPaintToolBuilder>();
-	RegisterTool(ToolManagerCommands.EnablePaintTool, TEXT("Paint"), PaintToolBuilder);
-
-	UPCGVolumeToolBuilder* VolumeToolBuilder = NewObject<UPCGVolumeToolBuilder>();
-	RegisterTool(ToolManagerCommands.EnableVolumeTool, TEXT("Volume"), VolumeToolBuilder);
 
 	ToolManager->SelectActiveToolType(EToolSide::Left, "DrawSpline");
-#endif
 }
 
 void UPCGEditorMode::Exit()
@@ -415,7 +393,7 @@ void UPCGEditorMode::ConfigureRealTimeViewportsOverride(const bool bEnable)
 	}
 }
 
-const UPCGInteractiveToolSettings* UPCGEditorMode::GetCurrentToolSettings() const
+UPCGInteractiveToolSettings* UPCGEditorMode::GetCurrentToolSettings() const
 {
 	UInteractiveToolManager* ToolManager = GetToolManager();
 
diff --git a/Engine/Plugins/PCG/Source/PCGEditor/Private/EditorMode/PCGEdMode.h b/Engine/Plugins/PCG/Source/PCGEditor/Private/EditorMode/PCGEdMode.h
--- a/Engine/Plugins/PCG/Source/PCGEditor/Private/EditorMode/PCGEdMode.h
+++ b/Engine/Plugins/PCG/Source/PCGEditor/Private/EditorMode/PCGEdMode.h
@@ -41,6 +41,8 @@ public:
 	virtual bool IsSelectionAllowed(AActor* InActor, bool bInSelection) const override;
 	//~End UEdMode interface
 
+	UPCGInteractiveToolSettings* GetCurrentToolSettings() const;
+
 protected:
 	virtual void CreateToolkit() override;
 
@@ -64,7 +66,5 @@ private:
 	// 	const TFunction<bool(UInteractiveToolManager*, EToolSide)>& CanExecuteAction,
 	// 	const TFunction<bool(UInteractiveToolManager*, EToolSide)>& IsActionChecked) const;
 
-	const UPCGInteractiveToolSettings* GetCurrentToolSettings() const;
-
 	bool bIsToolActive = false;
 };
diff --git a/Engine/Plugins/PCG/Source/PCGEditor/Private/EditorMode/PCGEdModeCommands.cpp b/Engine/Plugins/PCG/Source/PCGEditor/Private/EditorMode/PCGEdModeCommands.cpp
--- a/Engine/Plugins/PCG/Source/PCGEditor/Private/EditorMode/PCGEdModeCommands.cpp
+++ b/Engine/Plugins/PCG/Source/PCGEditor/Private/EditorMode/PCGEdModeCommands.cpp
@@ -49,7 +49,7 @@ FUICommandInfo::MakeCommandInfo(AsShared()\
 void FPCGEditorModePaletteCommands::RegisterCommands()
 {
 	/*** PCG EDITOR CONTEXTS **/
-	PCG_CONTEXT_COMMAND_INFO(LoadLineContextPalette, "Line", "PCGEditorMode.Tools.DrawSpline");
+	PCG_CONTEXT_COMMAND_INFO(LoadSplineContextPalette, "Spline", "PCGEditorMode.Context.DrawSpline");
 	PCG_CONTEXT_COMMAND_INFO(LoadPaintContextPalette, "Paint", "PCGEditorMode.Context.Paint");
 	PCG_CONTEXT_COMMAND_INFO(LoadVolumeContextPalette, "Volume", "PCGEditorMode.Context.Volume");
 }
@@ -113,29 +113,17 @@ void FPCGEditorModeToolCommands::GetToolDefaultObjectList(TArray<UInteractiveToo
 
 TArray<TSharedPtr<FUICommandInfo>> FPCGEditorModeToolCommands::GetLineContextToolCommands()
 {
-#if USE_PCG_MODE_GRAPHS_AS_BUTTONS
-	return PCG::EditorMode::Utility::GetOrCreateUICommandsWithToolTag(AsShared(), "SplineTool");
-#else
 	return {EnableDrawSplineTool, EnableDrawSurfaceTool};
-#endif
 }
 
 TArray<TSharedPtr<FUICommandInfo>> FPCGEditorModeToolCommands::GetPaintContextToolCommands()
 {
-#if USE_PCG_MODE_GRAPHS_AS_BUTTONS
-	return PCG::EditorMode::Utility::GetOrCreateUICommandsWithToolTag(AsShared(), "PaintTool");
-#else
 	return {EnablePaintTool};
-#endif
 }
 
 TArray<TSharedPtr<FUICommandInfo>> FPCGEditorModeToolCommands::GetVolumeContextToolCommands()
 {
-#if USE_PCG_MODE_GRAPHS_AS_BUTTONS
-	return PCG::EditorMode::Utility::GetOrCreateUICommandsWithToolTag(AsShared(), "VolumeTool");
-#else
 	return {EnableVolumeTool};
-#endif
 }
 
 /**
diff --git a/Engine/Plugins/PCG/Source/PCGEditor/Private/EditorMode/PCGEdModeCommands.h b/Engine/Plugins/PCG/Source/PCGEditor/Private/EditorMode/PCGEdModeCommands.h
--- a/Engine/Plugins/PCG/Source/PCGEditor/Private/EditorMode/PCGEdModeCommands.h
+++ b/Engine/Plugins/PCG/Source/PCGEditor/Private/EditorMode/PCGEdModeCommands.h
@@ -23,7 +23,7 @@ public:
 	FPCGEditorModePaletteCommands();
 
 	/*** PCG EDITOR CONTEXTS **/
-	TSharedPtr<FUICommandInfo> LoadLineContextPalette;
+	TSharedPtr<FUICommandInfo> LoadSplineContextPalette;
 	TSharedPtr<FUICommandInfo> LoadPaintContextPalette;
 	TSharedPtr<FUICommandInfo> LoadVolumeContextPalette;
 
diff --git a/Engine/Plugins/PCG/Source/PCGEditor/Private/EditorMode/PCGEdModeStyle.cpp b/Engine/Plugins/PCG/Source/PCGEditor/Private/EditorMode/PCGEdModeStyle.cpp
--- a/Engine/Plugins/PCG/Source/PCGEditor/Private/EditorMode/PCGEdModeStyle.cpp
+++ b/Engine/Plugins/PCG/Source/PCGEditor/Private/EditorMode/PCGEdModeStyle.cpp
@@ -61,7 +61,7 @@ FPCGEditorModeStyle::FPCGEditorModeStyle()
 
 	// Editor Mode Contexts
 	{
-		Set("PCGEditorMode.Context.Line", new IMAGE_BRUSH_SVG("Icons/EdMode/DrawSpline", Icon20x20));
+		Set("PCGEditorMode.Context.DrawSpline", new IMAGE_BRUSH_SVG("Icons/EdMode/DrawSpline", Icon20x20));
 		Set("PCGEditorMode.Context.Paint", new IMAGE_BRUSH_SVG("Icons/EdMode/Paint", Icon20x20));
 		Set("PCGEditorMode.Context.Volume", new IMAGE_BRUSH_SVG("Icons/EdMode/Volume", Icon20x20));
 	}
diff --git a/Engine/Plugins/PCG/Source/PCGEditor/Private/EditorMode/PCGEdModeToolkit.cpp b/Engine/Plugins/PCG/Source/PCGEditor/Private/EditorMode/PCGEdModeToolkit.cpp
--- a/Engine/Plugins/PCG/Source/PCGEditor/Private/EditorMode/PCGEdModeToolkit.cpp
+++ b/Engine/Plugins/PCG/Source/PCGEditor/Private/EditorMode/PCGEdModeToolkit.cpp
@@ -2,9 +2,13 @@
 
 #include "EditorMode/PCGEdModeToolkit.h"
 
+#include "EditorMode/PCGEdMode.h"
 #include "EditorMode/PCGEdModeCommands.h"
 #include "EditorMode/PCGEdModeSettings.h"
+#include "EditorMode/Tools/PCGInteractiveToolSettings.h"
+#include "EditorMode/Tools/Helpers/PCGEdModeEditorUtilities.h"
 
+#include "DetailLayoutBuilder.h"
 #include "EditorModeManager.h"
 #include "LevelEditorViewport.h"
 #include "PropertyEditorModule.h"
@@ -14,10 +18,138 @@
 #include "Tools/UEdMode.h"
 
 #include "SPrimaryButton.h"
+#include "Layout/SeparatorTemplates.h"
+#include "Styling/SlateIconFinder.h"
+#include "Tools/Line/PCGDrawSplineTool.h"
+#include "Tools/Paint/PCGPaintTool.h"
+#include "Tools/Volume/PCGVolumeTool.h"
+#include "Widgets/SPCGToolPresetSection.h"
 #include "Widgets/Input/STextComboBox.h"
 
 #define LOCTEXT_NAMESPACE "PCGEditorModeToolkit"
 
+FPCGEditorModeToolkitBuilder::FPCGEditorModeToolkitBuilder(FToolkitBuilderArgs& Args) : FToolkitBuilder(Args)
+{
+}
+
+void FPCGEditorModeToolkitBuilder::UpdateContentForCategory(FName InActiveCategoryName, FText InActiveCategoryText)
+{
+	if ( !MainContentVerticalBox.IsValid() )
+	{
+		return;
+	}
+	
+	if (ToolkitSections->ModeWarningArea)
+	{
+		MainContentVerticalBox->AddSlot()
+		.AutoHeight()
+		.HAlign(HAlign_Fill)
+		.Padding(5)
+		[
+			ToolkitSections->ModeWarningArea->AsShared()
+		];
+	}
+
+	if (ToolkitSections->Header)
+	{
+		MainContentVerticalBox->AddSlot()
+		.AutoHeight()
+		.HAlign(HAlign_Fill)
+		.Padding(0)
+		[
+			ToolkitSections->Header->AsShared()
+		];
+	}
+
+	MainContentVerticalBox->AddSlot()
+	.AutoHeight()
+	.HAlign(HAlign_Fill)
+	.Padding(0)
+	[
+		GetToolPaletteWidget()
+	];
+
+	if (ToolkitSections->ToolPresetArea)
+	{
+		MainContentVerticalBox->AddSlot()
+		.AutoHeight()
+		.HAlign(HAlign_Fill)
+		.Padding(0)
+		[
+			SNew(SBorder)
+			.HAlign(HAlign_Fill)
+			.Padding(Style.ActiveToolTitleBorderPadding)
+			[
+				ToolkitSections->ToolPresetArea->AsShared()
+			]
+		];
+	}
+	
+	MainContentVerticalBox->AddSlot()
+	.AutoHeight()
+	[
+		*FSeparatorTemplates::SmallHorizontalBackgroundNoBorder()
+		.BindVisibility(TAttribute<EVisibility>::CreateLambda([this]()
+		 { 
+		 	return ActivePaletteButtonVisibility; 
+		 }))
+	];
+
+	if (ToolkitSections->ToolWarningArea)
+	{
+		MainContentVerticalBox->AddSlot()
+		.AutoHeight()
+		.HAlign(HAlign_Fill)
+		.Padding(5)
+		[
+			ToolkitSections->ToolWarningArea->AsShared()
+		];
+	}
+
+	if (ToolkitSections->ToolHeader)
+	{
+		MainContentVerticalBox->AddSlot()
+		.AutoHeight()
+		.HAlign(HAlign_Fill)
+		.Padding(0)
+		[
+			ToolkitSections->ToolHeader->AsShared()
+		];
+	}
+
+	if (ToolkitSections->DetailsView)
+	{
+		MainContentVerticalBox->AddSlot()
+		.HAlign(HAlign_Fill)
+		.FillHeight(1.f)
+		[
+		SNew(SBorder)
+		.BorderImage(&Style.ToolDetailsBackgroundBrush)
+		.Padding(0.f, 2.f, 0.f, 2.f)
+			[
+				ToolkitSections->DetailsView->AsShared()
+			]
+		];
+	}
+	
+	if (ToolkitSections->Footer)
+	{
+		MainContentVerticalBox->AddSlot()
+		.AutoHeight()
+		.HAlign(HAlign_Fill)
+		.VAlign(VAlign_Bottom)
+		.Padding(0)
+		[
+			ToolkitSections->Footer->AsShared()
+		];
+	}
+}
+
+FPCGToolPalette::FPCGToolPalette(TSharedPtr<FUICommandInfo> InLoadToolPaletteAction, const TArray<TSharedPtr<FUICommandInfo>>& InPaletteActions)
+	: FToolPalette(InLoadToolPaletteAction, InPaletteActions)
+{
+}
+
 void FPCGEditorModeToolkit::Init(const TSharedPtr<IToolkitHost>& InitToolkitHost, const TWeakObjectPtr<UEdMode> OwningMode)
 {
 	// @todo_pcg: To be removed with UI overhaul (shared with Modeling Tools) in the future
@@ -104,6 +236,15 @@ void FPCGEditorModeToolkit::OnToolStarted(UInteractiveToolManager* Manager, UInt
 		// Force a refresh of the tool widget with the new tool properties.
 		AcceptToolWidget->Invalidate(EInvalidateWidgetReason::Volatility);
 	}
+
+	if (ToolPresetSection.IsValid() && GetActivePCGToolPalette().IsValid())
+	{
+		ToolPresetSection->HidePresets(false);
+		
+		LastPaletteWithActiveTool = GetActivePCGToolPalette()->LoadToolPaletteAction->GetCommandName();
+	}
+	
+	UpdatePresets();
 }
 
 void FPCGEditorModeToolkit::OnToolEnded(UInteractiveToolManager* Manager, UInteractiveTool* Tool)
@@ -112,10 +253,15 @@ void FPCGEditorModeToolkit::OnToolEnded(UInteractiveToolManager* Manager, UInter
 	{
 		GetToolkitHost()->RemoveViewportOverlayWidget(AcceptToolWidget.ToSharedRef());
 	}
-
+	
+	if (ToolPresetSection.IsValid())
+	{
+		LastPaletteWithActiveTool.Reset();
+		UpdatePresets();
+	}
+	
 	ModeDetailsView->SetObject(nullptr);
-
-	ToolkitBuilder->SetActiveToolDisplayName(FText::GetEmpty());
+	
 	ToolkitBuilder->SetActivePaletteCommandsVisibility(EVisibility::Visible);
 
 	ClearToolNotification();
@@ -228,6 +374,24 @@ FReply FPCGEditorModeToolkit::OnCancelButtonClicked() const
 	return FReply::Handled();
 }
 
+void FPCGEditorModeToolkit::UpdatePresets()
+{
+	if (ToolPresetSection.IsValid())
+	{
+		ToolPresetSection->UpdatePresets();
+	}
+}
+
+TSharedPtr<FPCGToolPalette> FPCGEditorModeToolkit::GetActivePCGToolPalette() const
+{
+	if (ToolkitBuilder->ActivePalette.IsValid())
+	{
+		return StaticCastSharedPtr<FPCGToolPalette>(ToolkitBuilder->ActivePalette);
+	}
+
+	return nullptr;
+}
+
 void FPCGEditorModeToolkit::RegisterToolkitCallbacks()
 {
 	// Subscribe to receive settings change updates.
@@ -237,9 +401,7 @@ void FPCGEditorModeToolkit::RegisterToolkitCallbacks()
 	GetScriptableEditorMode()->GetInteractiveToolsContext(EToolsContextScope::EdMode)->OnToolNotificationMessage.AddSP(this, &FPCGEditorModeToolkit::PostToolNotification);
 	GetScriptableEditorMode()->GetInteractiveToolsContext(EToolsContextScope::EdMode)->OnToolWarningMessage.AddSP(this, &FPCGEditorModeToolkit::PostToolWarning);
 
-	GetToolkitHost()->OnActiveViewportChanged().AddSP(this, &FPCGEditorModeToolkit::OnActiveViewportChanged);
-
-	
+	GetToolkitHost()->OnActiveViewportChanged().AddSP(this, &FPCGEditorModeToolkit::OnActiveViewportChanged);	
 }
 
 void FPCGEditorModeToolkit::UnregisterToolkitCallbacks()
@@ -286,8 +448,6 @@ void FPCGEditorModeToolkit::UpdateActiveTool()
 	
 	CachedActiveToolInfo = CurrentTool->GetToolInfo();
 
-	ToolkitBuilder->SetActiveToolDisplayName(CachedActiveToolInfo.ToolDisplayName);
-
 	if (const UPCGEditorModeSettings* Settings = GetDefault<UPCGEditorModeSettings>())
 	{
 		if (Settings->bHideToolButtonsDuringActiveTool)
@@ -343,26 +503,32 @@ void FPCGEditorModeToolkit::RegisterPalettes()
 	ToolkitBuilderArgs.ToolkitSections = ToolkitSections;
 	ToolkitBuilderArgs.SelectedCategoryTitleVisibility = EVisibility::Collapsed;
 
-	ToolkitBuilder = MakeShared<FToolkitBuilder>(ToolkitBuilderArgs);
-
-	const TSharedPtr<FToolPalette> LineToolPalette = MakeShared<FToolPalette>(PaletteCommands.LoadLineContextPalette, ToolCommands.GetLineContextToolCommands());
-	ToolkitBuilder->AddPalette(std::move(LineToolPalette));
+	ToolkitBuilder = MakeShared<FPCGEditorModeToolkitBuilder>(ToolkitBuilderArgs);
+		
+	{
+		FString ToolIdentifier = UPCGInteractiveToolSettings_Spline::StaticGetToolTag().ToString();
+		const TSharedPtr<FPCGToolPalette> DrawSplineToolSection = MakeShared<FPCGToolPalette>(PaletteCommands.LoadSplineContextPalette, ToolCommands.GetLineContextToolCommands());
+		ToolkitBuilder->AddPalette(std::move(DrawSplineToolSection));
+	}
 
-	const TSharedPtr<FToolPalette> PaintToolPalette = MakeShared<FToolPalette>(PaletteCommands.LoadPaintContextPalette, ToolCommands.GetPaintContextToolCommands());
-	ToolkitBuilder->AddPalette(std::move(PaintToolPalette));
+	{
+		FString ToolIdentifier = UPCGInteractiveToolSettings_PaintTool::StaticGetToolTag().ToString();
+		const TSharedPtr<FPCGToolPalette> PaintToolPalette = MakeShared<FPCGToolPalette>(PaletteCommands.LoadPaintContextPalette, ToolCommands.GetPaintContextToolCommands());
+		ToolkitBuilder->AddPalette(std::move(PaintToolPalette));
+	}
 
-	const TSharedPtr<FToolPalette> VolumeToolPalette = MakeShared<FToolPalette>(PaletteCommands.LoadVolumeContextPalette, ToolCommands.GetVolumeContextToolCommands());
-	ToolkitBuilder->AddPalette(std::move(VolumeToolPalette));
+	{
+		FString ToolIdentifier = UPCGInteractiveToolSettings_Volume::StaticGetToolTag().ToString();
+		const TSharedPtr<FPCGToolPalette> VolumeToolPalette = MakeShared<FPCGToolPalette>(PaletteCommands.LoadVolumeContextPalette, ToolCommands.GetVolumeContextToolCommands());
+		ToolkitBuilder->AddPalette(std::move(VolumeToolPalette));
+	}
 
-	ToolkitBuilder->SetActivePaletteOnLoad(PaletteCommands.LoadLineContextPalette.Get());
+	ToolkitBuilder->SetActivePaletteOnLoad(PaletteCommands.LoadSplineContextPalette.Get());
 
 	ToolkitBuilder->UpdateWidget();
 
 	// If selected palette changes, make sure we are showing the palette command buttons, which may be hidden by active Tool.
-	ActivePaletteChangedHandle = ToolkitBuilder->OnActivePaletteChanged.AddLambda([this]()
-	{
-		ToolkitBuilder->SetActivePaletteCommandsVisibility(EVisibility::Visible);
-	});
+	ActivePaletteChangedHandle = ToolkitBuilder->OnActivePaletteChanged.AddSP(this, &FPCGEditorModeToolkit::OnActivePaletteChanged);
 }
 
 void FPCGEditorModeToolkit::CreateWidgets()
@@ -386,6 +552,9 @@ void FPCGEditorModeToolkit::CreateWidgets()
 		ToolkitSections->ModeWarningArea->AsShared()
 	];
 
+	ToolPresetSection = SNew(SPCGToolPresetSection, *OwningEditorMode.Get());
+	ToolkitSections->ToolPresetArea = ToolPresetSection;
+
 	SAssignNew(AcceptToolWidget, SPCGEdModeAcceptToolOverlay)
 	.ActiveToolDisplayName(this, &FPCGEditorModeToolkit::GetActiveToolDisplayName)
 	.ActiveToolIcon(this, &FPCGEditorModeToolkit::GetActiveToolIcon)
@@ -401,4 +570,18 @@ void FPCGEditorModeToolkit::CreateWidgets()
 	.OnCancelButtonClicked(this, &FPCGEditorModeToolkit::OnCancelButtonClicked);
 }
 
+void FPCGEditorModeToolkit::OnActivePaletteChanged()
+{
+	ToolkitBuilder->SetActivePaletteCommandsVisibility(EVisibility::Visible);
+
+	FName CurrentlyActivePaletteName = GetActivePCGToolPalette()->LoadToolPaletteAction->GetCommandName();
+	
+	// If the palette changes, hide the preset widget.
+	if (ToolPresetSection.IsValid() && LastPaletteWithActiveTool.IsSet())
+	{
+		bool bPaletteChanged = !LastPaletteWithActiveTool.GetValue().IsEqual(CurrentlyActivePaletteName); 
+		ToolPresetSection->HidePresets(bPaletteChanged);
+	}
+}
+
 #undef LOCTEXT_NAMESPACE
diff --git a/Engine/Plugins/PCG/Source/PCGEditor/Private/EditorMode/PCGEdModeToolkit.h b/Engine/Plugins/PCG/Source/PCGEditor/Private/EditorMode/PCGEdModeToolkit.h
--- a/Engine/Plugins/PCG/Source/PCGEditor/Private/EditorMode/PCGEdModeToolkit.h
+++ b/Engine/Plugins/PCG/Source/PCGEditor/Private/EditorMode/PCGEdModeToolkit.h
@@ -2,14 +2,35 @@
 
 #pragma once
 
+#include "PCGGraph.h"
+
 #include "InteractiveTool.h"
 #include "StatusBarSubsystem.h"
 #include "Toolkits/BaseToolkit.h"
 
 #include "Widgets/SPCGEdModeAcceptToolOverlay.h"
+#include "Widgets/Views/STileView.h"
 
+class SPCGToolPresetSection;
 class UPCGInteractiveTool;
 
+/** A customized Toolkit Builder used for the PCG Editor Mode  */
+class FPCGEditorModeToolkitBuilder : public FToolkitBuilder
+{
+public:
+	explicit FPCGEditorModeToolkitBuilder(FToolkitBuilderArgs& Args);
+
+protected:
+	virtual void UpdateContentForCategory(FName InActiveCategoryName, FText InActiveCategoryText) override;
+};
+
+struct FPCGToolPalette : FToolPalette
+{
+	FPCGToolPalette(TSharedPtr<FUICommandInfo> InLoadToolPaletteAction, const TArray<TSharedPtr<FUICommandInfo>>& InPaletteActions);
+
+	/** Future custom properties go here. */
+};
+
 /** The basic toolkit needed to show the docked editor mode panel and set up the tool shelves and details views. */
 class FPCGEditorModeToolkit : public FModeToolkit
 {
@@ -56,6 +77,9 @@ public:
 	FReply OnAcceptButtonClicked() const;
 	FReply OnCancelButtonClicked() const;
 
+	void UpdatePresets();
+
+	TSharedPtr<FPCGToolPalette> GetActivePCGToolPalette() const;
 private:
 	void RegisterToolkitCallbacks();
 	void UnregisterToolkitCallbacks();
@@ -73,13 +97,19 @@ private:
 	/** A utility function to create additional widgets for the Toolkit. */
 	void CreateWidgets();
 
+	void OnActivePaletteChanged();
+
 	/** UI overlay in the viewport to accept or cancel the current tool. */
 	TSharedPtr<SPCGEdModeAcceptToolOverlay> AcceptToolWidget;
 
+	TSharedPtr<SPCGToolPresetSection> ToolPresetSection;
+	
 	/** Currently only used to cancel the current tool. */
 	FDelegateHandle ActivePaletteChangedHandle;
 
 	FStatusBarMessageHandle ActiveToolMessageHandle;
 	FText ActiveWarning{};
 	FInteractiveToolInfo CachedActiveToolInfo;
+
+	TOptional<FName> LastPaletteWithActiveTool;
 };
diff --git a/Engine/Plugins/PCG/Source/PCGEditor/Private/EditorMode/Tools/Helpers/PCGEdModeEditorUtilities.cpp b/Engine/Plugins/PCG/Source/PCGEditor/Private/EditorMode/Tools/Helpers/PCGEdModeEditorUtilities.cpp
--- a/Engine/Plugins/PCG/Source/PCGEditor/Private/EditorMode/Tools/Helpers/PCGEdModeEditorUtilities.cpp
+++ b/Engine/Plugins/PCG/Source/PCGEditor/Private/EditorMode/Tools/Helpers/PCGEdModeEditorUtilities.cpp
@@ -19,76 +19,124 @@ static TAutoConsoleVariable<bool> CVarClearToolAssetCache(
 
 TArray<UE::PCG::EditorMode::Utility::FPCGGraphToolEditorData> UE::PCG::EditorMode::Utility::GetGraphToolsWithToolTag(FName ToolTag)
 {
-	TArray<UE::PCG::EditorMode::Utility::FPCGGraphToolEditorData> Result;
+	TArray<FPCGGraphToolEditorData> Result;
 	
 	TArray<TSharedPtr<FUICommandInfo>> GraphAssetCommands;
 	FAssetRegistryModule& AssetRegistryModule = FModuleManager::Get().LoadModuleChecked<FAssetRegistryModule>("AssetRegistry");
 
 	FARFilter Filter;
-	Filter.ClassPaths.Add(UPCGGraph::StaticClass()->GetClassPathName());
-	Filter.TagsAndValues.Add(GET_MEMBER_NAME_CHECKED(UPCGGraph, ToolData));
+	Filter.ClassPaths.Add(UPCGGraphInterface::StaticClass()->GetClassPathName());
+	Filter.bRecursiveClasses = true;
 
 	TArray<FAssetData> GraphAssets;
 	AssetRegistryModule.Get().GetAssets(Filter, GraphAssets);
 
 	for(auto It = GraphAssets.CreateIterator(); It; ++It)
 	{
-		FPCGGraphToolData ToolData;
-
-		FString ToolDataText;
-		if(It->GetTagValue(GET_MEMBER_NAME_CHECKED(UPCGGraph, ToolData), ToolDataText))
+		if (DoesPCGGraphInterfaceHaveToolTag(ToolTag, *It))
 		{
-			FOutputDeviceNull NullOut;
-			StaticStruct<FPCGGraphToolData>()->ImportText(*ToolDataText, &ToolData, nullptr, PPF_None, &NullOut,
-				FPCGGraphToolData::StaticStruct()->GetName(), true);
-
-			if(ToolData.CompatibleToolTags.Contains(ToolTag))
-			{
-				FPCGGraphToolEditorData GraphToolEditorData;
-				GraphToolEditorData.AssetData = (*It);
-				GraphToolEditorData.GraphToolData = ToolData;
-				Result.Add(GraphToolEditorData);
-			}
+			TOptional<FPCGGraphToolData> GraphToolData = GetGraphToolDataFromPCGGraphInterface(*It);
+			FPCGGraphToolEditorData GraphToolEditorData;
+			GraphToolEditorData.AssetData = (*It);
+			GraphToolEditorData.GraphToolData = GraphToolData.Get(FPCGGraphToolData());
+			Result.Add(GraphToolEditorData);
 		}
 	}
 
 	return Result;
 }
 
-bool UE::PCG::EditorMode::Utility::DoesPCGGraphHaveToolTag(FName ToolTag, const FAssetData& AssetData)
-{
-	if (AssetData.GetClass()->IsChildOf(UPCGGraph::StaticClass()))
+bool UE::PCG::EditorMode::Utility::DoesPCGGraphInterfaceHaveToolTag(FName ToolTag, const FAssetData& AssetData)
+{	
+	if ( ToolTag.IsNone() || !AssetData.IsValid() || !AssetData.GetClass()->IsChildOf(UPCGGraphInterface::StaticClass()))
 	{
-		FPCGGraphToolData ToolData;
-		FString ToolDataText;
-		if (AssetData.GetTagValue(GET_MEMBER_NAME_CHECKED(UPCGGraph, ToolData), ToolDataText))
-		{
-			FOutputDeviceNull NullOut;
-			StaticStruct<FPCGGraphToolData>()->ImportText(*ToolDataText, &ToolData, nullptr, PPF_None, &NullOut,
-				FPCGGraphToolData::StaticStruct()->GetName(), true);
+		return false;
+	}
 
-			if (ToolData.CompatibleToolTags.Contains(ToolTag))
-			{
-				return true;
-			}
-		}
+	TOptional<FPCGGraphToolData> GraphToolData = GetGraphToolDataFromPCGGraphInterface(AssetData);
+
+	if (GraphToolData.IsSet())
+	{
+		return GraphToolData->CompatibleToolTags.Contains(ToolTag);
 	}
 
 	return false;
 }
 
-bool UE::PCG::EditorMode::Utility::DoesPCGGraphInterfaceHaveToolTag(FName ToolTag, const FAssetData& AssetData)
-{	
+bool UE::PCG::EditorMode::Utility::IsPCGGraphInterfaceCompatibleWithActor(AActor* Actor, const FAssetData& AssetData)
+{
+	if (!ensure(Actor))
+	{
+		return false;
+	}
+
 	if (!AssetData.IsValid() || !AssetData.GetClass()->IsChildOf(UPCGGraphInterface::StaticClass()))
 	{
 		return false;
 	}
+
+	TOptional<FPCGGraphToolData> GraphToolData = GetGraphToolDataFromPCGGraphInterface(AssetData);
+
+	if (GraphToolData.IsSet())
+	{
+		return Actor->IsA(GraphToolData->InitialActorClassToSpawn);
+	}
 	
+	return false;
+}
+
+TOptional<FPCGGraphToolData> UE::PCG::EditorMode::Utility::GetGraphToolDataFromPCGGraphInterface(const FAssetData& AssetData)
+{
+	if (!AssetData.IsValid() || !AssetData.GetClass()->IsChildOf(UPCGGraphInterface::StaticClass()))
+	{
+		return {};
+	}
+
+	// Prefer the loaded data, if available
+	if (AssetData.IsAssetLoaded())
+	{
+		if (UPCGGraphInterface* LoadedAsset = Cast<UPCGGraphInterface>(AssetData.GetAsset()))
+		{
+			return LoadedAsset->GetGraphToolData();
+		}
+	}
+
+	// Otherwise, retrieve the tool data from the Asset Registry
 	if (AssetData.GetClass()->IsChildOf(UPCGGraph::StaticClass()))
 	{
-		return DoesPCGGraphHaveToolTag(ToolTag, AssetData);
+		return GetARGraphToolDataFromPCGGraph(AssetData);
 	}
 	else if (AssetData.GetClass()->IsChildOf(UPCGGraphInstance::StaticClass()))
+	{
+		return GetARGraphToolDataFromPCGGraphInstance(AssetData);
+	}
+
+	return {};
+}
+
+TOptional<FPCGGraphToolData> UE::PCG::EditorMode::Utility::GetARGraphToolDataFromPCGGraph(const FAssetData& AssetData)
+{
+	if (AssetData.GetClass()->IsChildOf(UPCGGraph::StaticClass()))
+	{
+		FPCGGraphToolData GraphToolData;
+		
+		FString GraphToolDataText;
+		if (AssetData.GetTagValue(GET_MEMBER_NAME_CHECKED(UPCGGraph, ToolData), GraphToolDataText))
+		{
+			FOutputDeviceNull NullOut;
+			StaticStruct<FPCGGraphToolData>()->ImportText(*GraphToolDataText, &GraphToolData, nullptr, PPF_None, &NullOut,
+				FPCGGraphToolData::StaticStruct()->GetName(), true);
+
+			return GraphToolData;
+		}
+	}
+
+	return {};
+}
+
+TOptional<FPCGGraphToolData> UE::PCG::EditorMode::Utility::GetARGraphToolDataFromPCGGraphInstance(const FAssetData& AssetData)
+{
+	if (AssetData.GetClass()->IsChildOf(UPCGGraphInstance::StaticClass()))
 	{
 		FString PCGGraphText;
 		if (AssetData.GetTagValue("Graph", PCGGraphText))
@@ -102,71 +150,42 @@ bool UE::PCG::EditorMode::Utility::DoesPCGGraphInterfaceHaveToolTag(FName ToolTa
 
 				if (LinkedPCGGraphAssetData.IsValid() && LinkedPCGGraphAssetData.GetClass()->IsChildOf(UPCGGraph::StaticClass()))
 				{
-					return DoesPCGGraphHaveToolTag(ToolTag, LinkedPCGGraphAssetData);
+					TOptional<FPCGGraphToolData> GraphToolData = GetARGraphToolDataFromPCGGraph(LinkedPCGGraphAssetData);
+
+					TOptional<FPCGGraphInstanceToolDataOverrides> GraphInstanceToolDataOverrides = GetARGraphToolDataOverridesFromPCGGraphInstance(AssetData);
+
+					if (GraphToolData.IsSet() && GraphInstanceToolDataOverrides.IsSet())
+					{
+						GraphToolData->ApplyOverrides(GraphInstanceToolDataOverrides.GetValue());
+					}
+					
+					return GraphToolData;
 				}
 			}
 		}
 	}
 
-	return false;
+	return {};
 }
 
-bool UE::PCG::EditorMode::Utility::IsPCGGraphCompatibleWithActor(AActor* Actor, const FAssetData& AssetData)
+TOptional<FPCGGraphInstanceToolDataOverrides> UE::PCG::EditorMode::Utility::GetARGraphToolDataOverridesFromPCGGraphInstance(const FAssetData& AssetData)
 {
-	if (AssetData.GetClass()->IsChildOf(UPCGGraph::StaticClass()))
+	if (AssetData.GetClass()->IsChildOf(UPCGGraphInstance::StaticClass()))
 	{
-		FPCGGraphToolData ToolData;
-		FString ToolDataText;
-		if (AssetData.GetTagValue(GET_MEMBER_NAME_CHECKED(UPCGGraph, ToolData), ToolDataText))
+		FPCGGraphInstanceToolDataOverrides GraphToolDataOverrides;
+		
+		FString GraphToolDataOverridesText;
+		if (AssetData.GetTagValue(GET_MEMBER_NAME_CHECKED(UPCGGraphInstance, ToolDataOverrides), GraphToolDataOverridesText))
 		{
 			FOutputDeviceNull NullOut;
-			StaticStruct<FPCGGraphToolData>()->ImportText(*ToolDataText, &ToolData, nullptr, PPF_None, &NullOut,
-				FPCGGraphToolData::StaticStruct()->GetName(), true);
+			StaticStruct<FPCGGraphInstanceToolDataOverrides>()->ImportText(*GraphToolDataOverridesText, &GraphToolDataOverrides, nullptr, PPF_None, &NullOut,
+				FPCGGraphInstanceToolDataOverrides::StaticStruct()->GetName(), true);
 
-			return Actor->IsA(ToolData.InitialActorClassToSpawn);
+			return GraphToolDataOverrides;
 		}
 	}
 
-	return false;
-}
-
-bool UE::PCG::EditorMode::Utility::IsPCGGraphInterfaceCompatibleWithActor(AActor* Actor, const FAssetData& AssetData)
-{
-	if (!ensure(Actor))
-	{
-		return false;
-	}
-
-	if (!AssetData.IsValid() || !AssetData.GetClass()->IsChildOf(UPCGGraphInterface::StaticClass()))
-	{
-		return false;
-	}
-	
-	if (AssetData.GetClass()->IsChildOf(UPCGGraph::StaticClass()))
-	{
-		return IsPCGGraphCompatibleWithActor(Actor, AssetData);
-	}
-	else if (AssetData.GetClass()->IsChildOf(UPCGGraphInstance::StaticClass()))
-    {
-    	FString PCGGraphText;
-    	if (AssetData.GetTagValue("Graph", PCGGraphText))
-    	{
-    		FSoftObjectPath PCGGraphPath(PCGGraphText);
-
-    		if(PCGGraphPath.IsAsset())
-    		{
-    			FAssetRegistryModule& AssetRegistryModule = FModuleManager::Get().LoadModuleChecked<FAssetRegistryModule>("AssetRegistry");
-    			FAssetData LinkedPCGGraphAssetData = AssetRegistryModule.Get().GetAssetByObjectPath(PCGGraphPath);
-
-    			if (LinkedPCGGraphAssetData.IsValid() && LinkedPCGGraphAssetData.GetClass()->IsChildOf(UPCGGraph::StaticClass()))
-    			{
-    				return IsPCGGraphCompatibleWithActor(Actor, LinkedPCGGraphAssetData);
-    			}
-    		}
-    	}
-    }
-	
-	return false;
+	return {};
 }
 
 TArray<TSharedPtr<FUICommandInfo>> UE::PCG::EditorMode::Utility::GetOrCreateUICommandsWithToolTag(TSharedRef<FBindingContext> BindingContext, FName ToolTag)
diff --git a/Engine/Plugins/PCG/Source/PCGEditor/Private/EditorMode/Tools/Helpers/PCGEdModeEditorUtilities.h b/Engine/Plugins/PCG/Source/PCGEditor/Private/EditorMode/Tools/Helpers/PCGEdModeEditorUtilities.h
--- a/Engine/Plugins/PCG/Source/PCGEditor/Private/EditorMode/Tools/Helpers/PCGEdModeEditorUtilities.h
+++ b/Engine/Plugins/PCG/Source/PCGEditor/Private/EditorMode/Tools/Helpers/PCGEdModeEditorUtilities.h
@@ -20,9 +20,16 @@ namespace UE::PCG::EditorMode::Utility
 	};
 	
 	TArray<FPCGGraphToolEditorData> GetGraphToolsWithToolTag(FName ToolTag);
-	bool DoesPCGGraphHaveToolTag(FName ToolTag, const FAssetData& AssetData);
+	
 	bool DoesPCGGraphInterfaceHaveToolTag(FName ToolTag, const FAssetData& AssetData);
-	bool IsPCGGraphCompatibleWithActor(AActor* Actor, const FAssetData& AssetData);
 	bool IsPCGGraphInterfaceCompatibleWithActor(AActor* Actor, const FAssetData& AssetData);
+	
+	TOptional<FPCGGraphToolData> GetGraphToolDataFromPCGGraphInterface(const FAssetData& AssetData);
+	
+	TOptional<FPCGGraphToolData> GetARGraphToolDataFromPCGGraph(const FAssetData& AssetData);
+	TOptional<FPCGGraphToolData> GetARGraphToolDataFromPCGGraphInstance(const FAssetData& AssetData);
+
+	TOptional<FPCGGraphInstanceToolDataOverrides> GetARGraphToolDataOverridesFromPCGGraphInstance(const FAssetData& AssetData);
+	
 	TArray<TSharedPtr<FUICommandInfo>> GetOrCreateUICommandsWithToolTag(TSharedRef<FBindingContext> BindingContext, FName ToolTag);
 }
diff --git a/Engine/Plugins/PCG/Source/PCGEditor/Private/EditorMode/Tools/PCGInteractiveToolSettings.cpp b/Engine/Plugins/PCG/Source/PCGEditor/Private/EditorMode/Tools/PCGInteractiveToolSettings.cpp
--- a/Engine/Plugins/PCG/Source/PCGEditor/Private/EditorMode/Tools/PCGInteractiveToolSettings.cpp
+++ b/Engine/Plugins/PCG/Source/PCGEditor/Private/EditorMode/Tools/PCGInteractiveToolSettings.cpp
@@ -820,6 +820,11 @@ bool UPCGInteractiveToolSettings::GraphAssetFilter(const FAssetData& AssetData)
 	return !Utility::IsPCGGraphInterfaceCompatibleWithActor(EditActor.Get(), AssetData);
 }
 
+bool UPCGInteractiveToolSettings::UsesExistingActor() const
+{
+	return !GeneratedResources.SpawnedActor.IsValid() && EditActor.IsValid();
+}
+
 UPCGComponent* UPCGInteractiveToolSettings::GetWorkingPCGComponent() const
 {
 	if(WorkingComponentCache.IsValid() == false)
diff --git a/Engine/Plugins/PCG/Source/PCGEditor/Private/EditorMode/Tools/PCGInteractiveToolSettings.h b/Engine/Plugins/PCG/Source/PCGEditor/Private/EditorMode/Tools/PCGInteractiveToolSettings.h
--- a/Engine/Plugins/PCG/Source/PCGEditor/Private/EditorMode/Tools/PCGInteractiveToolSettings.h
+++ b/Engine/Plugins/PCG/Source/PCGEditor/Private/EditorMode/Tools/PCGInteractiveToolSettings.h
@@ -72,6 +72,9 @@ public:
 	 * If your target actor has any special requirements, create an actor for that purpose and return it in your settings. */
 	virtual TSubclassOf<AActor> GetWorkingActorClass() const;
 
+	/** Whether we are operating on a spawned or a previously-existing actor. */
+	bool UsesExistingActor() const;
+	
 	virtual void SetActorClassToSpawn(TSubclassOf<AActor> Class);
 
 	void SetActorLabel(FName InLabel);
diff --git a/Engine/Plugins/PCG/Source/PCGEditor/Private/EditorMode/Widgets/SPCGToolPresetSection.cpp b/Engine/Plugins/PCG/Source/PCGEditor/Private/EditorMode/Widgets/SPCGToolPresetSection.cpp
new file mode 100644
--- /dev/null
+++ b/Engine/Plugins/PCG/Source/PCGEditor/Private/EditorMode/Widgets/SPCGToolPresetSection.cpp
@@ -0,0 +1,215 @@
+﻿// Copyright Epic Games, Inc. All Rights Reserved.
+
+#include "SPCGToolPresetSection.h"
+#include "EditorMode/PCGEdMode.h"
+#include "EditorMode/PCGEdModeToolkit.h"
+#include "EditorMode/Tools/PCGInteractiveToolSettings.h"
+#include "EditorMode/Tools/Helpers/PCGEdModeEditorUtilities.h"
+
+#include "ToolMenus.h"
+#include "Styling/SlateIconFinder.h"
+#include "Framework/MultiBox/MultiBoxBuilder.h"
+#include "Framework/MultiBox/SToolBarButtonBlock.h"
+
+void SPCGToolPresetSection::Construct(const FArguments& InArgs, UEdMode& InOwningEditorMode)
+{
+	OwningEditorMode = &InOwningEditorMode;
+
+	ThisContext.Reset(NewObject<UPCGToolPresetMenuContext>());
+	ThisContext->ThisSection = SharedThis(this).ToWeakPtr();
+	
+	SetVisibility(TAttribute<EVisibility>::CreateSP(this, &SPCGToolPresetSection::OnGetVisibility));
+}
+
+void SPCGToolPresetSection::UpdatePresets()
+{
+	Presets.Reset();
+
+	if(OwningEditorMode.IsValid() == false)
+	{
+		return;
+	}
+
+	UPCGEditorMode* EdMode = CastChecked<UPCGEditorMode>(OwningEditorMode.Get());
+
+	if(EdMode->GetToolkit().IsValid() == false)
+	{
+		return;
+	}
+
+	UPCGInteractiveToolSettings* ToolSettings = EdMode->GetCurrentToolSettings();
+
+	if(ToolSettings == nullptr)
+	{
+		return;
+	}
+	TSharedPtr<FPCGEditorModeToolkit> Toolkit = StaticCastSharedPtr<FPCGEditorModeToolkit>(EdMode->GetToolkit().Pin());
+
+	TSharedPtr<FPCGToolPalette> PCGToolPalette = Toolkit->GetActivePCGToolPalette();
+	if(PCGToolPalette.IsValid() == false)
+	{
+		return;
+	}
+
+	FName ToolTag = ToolSettings->GetToolTag();
+
+	if(ToolTag.IsNone())
+	{
+		return;
+	}
+
+	using namespace UE::PCG::EditorMode::Utility;
+
+	TArray<FPCGGraphToolEditorData> GraphToolEditorData = GetGraphToolsWithToolTag(ToolTag);
+
+	bool bUsesExistingActor = ToolSettings->UsesExistingActor();
+	AActor* EditActor = ToolSettings->GetTypedWorkingActor<AActor>();
+
+	Presets.Reserve(GraphToolEditorData.Num());
+	for (const FPCGGraphToolEditorData& Candidate : GraphToolEditorData)
+	{
+		if (!Candidate.GraphToolData.bIsPreset)
+		{
+			continue;
+		}
+		
+		// If our existing edit actor isn't compatible with the graph candidate, we skip it
+		if (bUsesExistingActor && EditActor)
+		{
+			if (!IsPCGGraphInterfaceCompatibleWithActor(EditActor, Candidate.AssetData))
+			{
+				continue;
+			}
+		}
+		
+		FPresetItemPtr NewPreset = MakeShared<FPresetItem>();
+		NewPreset->Item = TSoftObjectPtr<UPCGGraphInterface>(Candidate.AssetData.ToSoftObjectPath());
+		NewPreset->Name = Candidate.GraphToolData.DisplayName.IsEmpty()
+			? FText::AsCultureInvariant(Candidate.AssetData.AssetName.ToString())
+			: Candidate.GraphToolData.DisplayName;
+		NewPreset->Tooltip = Candidate.GraphToolData.Tooltip;
+		Presets.Add(NewPreset);
+	}
+
+	ChildSlot
+	[
+		MakePresetToolbar().ToSharedRef()
+	];
+}
+
+bool SPCGToolPresetSection::HasPresets() const
+{
+	return Presets.Num() > 0;
+}
+
+TSharedPtr<SWidget> SPCGToolPresetSection::MakePresetToolbar() const
+{
+	TSharedRef<FSlimHorizontalUniformToolBarBuilder> ToolBarBuilder = MakeShared<FSlimHorizontalUniformToolBarBuilder>(nullptr,
+		FMultiBoxCustomization::None, nullptr, false);
+	ToolBarBuilder->SetStyle(&FAppStyle::Get(), "SlimPaletteToolBar");
+
+	static FName MenuName = "PCGToolPresetMenu";
+	if (!UToolMenus::Get()->IsMenuRegistered(MenuName))
+	{
+		FToolMenuOwnerScoped OwnerScoped(UE_MODULE_NAME);
+
+		UToolMenu* ToolMenu = UToolMenus::Get()->RegisterMenu(MenuName, NAME_None);
+		ToolMenu->MenuType = EMultiBoxType::SlimHorizontalUniformToolBar;
+		ToolMenu->StyleName = "SlimPaletteToolBar";
+		ToolMenu->AddDynamicSection(NAME_None, FNewToolMenuDelegate::CreateStatic(&SPCGToolPresetSection::OnGeneratePresetMenu));
+	}
+
+	FToolMenuContext Context;
+	Context.AddObject(ThisContext.Get());
+	
+	return UToolMenus::Get()->GenerateWidget(MenuName, Context);
+}
+
+void SPCGToolPresetSection::OnGeneratePresetMenu(UToolMenu* ToolMenu)
+{
+	const UPCGToolPresetMenuContext* MenuContext = ToolMenu->FindContext<UPCGToolPresetMenuContext>();
+
+	if(!ensure(MenuContext && MenuContext->ThisSection.IsValid()))
+	{
+		return;
+	}
+
+	TSharedPtr<SPCGToolPresetSection> ThisSection = MenuContext->ThisSection.Pin();
+	
+	for(const FPresetItemPtr& PresetItemPtr : ThisSection->GetPresets())
+	{
+		FUIAction Action;
+		Action.CanExecuteAction = FCanExecuteAction::CreateStatic(&SPCGToolPresetSection::CanActivatePreset, PresetItemPtr.ToWeakPtr(), MenuContext);
+		Action.ExecuteAction = FExecuteAction::CreateStatic(&SPCGToolPresetSection::ActivatePreset, PresetItemPtr.ToWeakPtr(), MenuContext);
+		Action.GetActionCheckState = FGetActionCheckState::CreateStatic(&SPCGToolPresetSection::IsPresetActive, PresetItemPtr.ToWeakPtr(), MenuContext);
+
+		TAttribute<FText> TooltipAttribute = TAttribute<FText>::CreateStatic(&SPCGToolPresetSection::GetPresetTooltip, PresetItemPtr.ToWeakPtr(), MenuContext);
+
+		ToolMenu->AddMenuEntry(NAME_None, FToolMenuEntry::InitToolBarButton(FName(PresetItemPtr->Name.ToString()), Action, PresetItemPtr->Name, TooltipAttribute, FSlateIconFinder::FindIconForClass(UPCGGraph::StaticClass()), EUserInterfaceActionType::ToggleButton));
+	}
+}
+
+bool SPCGToolPresetSection::CanActivatePreset(TWeakPtr<FPresetItem> PresetItemWeakPtr, const UPCGToolPresetMenuContext* Context)
+{
+	return PresetItemWeakPtr.IsValid();
+}
+
+void SPCGToolPresetSection::ActivatePreset(TWeakPtr<FPresetItem> PresetItemWeakPtr, const UPCGToolPresetMenuContext* Context)
+{
+	if (PresetItemWeakPtr.IsValid() == false || Context->ThisSection.IsValid() == false)
+	{
+		return;
+	}
+
+	UPCGEditorMode* EdMode = CastChecked<UPCGEditorMode>(Context->ThisSection.Pin()->OwningEditorMode.Get());
+
+	if (UPCGInteractiveToolSettings* ToolSettings = EdMode->GetCurrentToolSettings())
+	{
+		if(UPCGGraphInterface* Graph = PresetItemWeakPtr.Pin()->Item.LoadSynchronous())
+		{
+			ToolSettings->SetToolGraph(Graph);
+		}
+	}
+}
+
+FText SPCGToolPresetSection::GetPresetTooltip(TWeakPtr<FPresetItem> PresetItemWeakPtr, const UPCGToolPresetMenuContext* Context)
+{
+	if (PresetItemWeakPtr.IsValid() == false || Context->ThisSection.IsValid() == false)
+	{
+		return FText::GetEmpty();
+	}
+
+	return PresetItemWeakPtr.Pin()->Tooltip;
+}
+
+ECheckBoxState SPCGToolPresetSection::IsPresetActive(TWeakPtr<FPresetItem> PresetItemWeakPtr, const UPCGToolPresetMenuContext* Context)
+{
+	if (PresetItemWeakPtr.IsValid() == false || Context->ThisSection.IsValid() == false)
+	{
+		return ECheckBoxState::Undetermined;
+	}
+	
+	TSharedPtr<FPresetItem> PresetItem = PresetItemWeakPtr.Pin();
+	
+	UPCGEditorMode* EdMode = CastChecked<UPCGEditorMode>(Context->ThisSection.Pin()->OwningEditorMode.Get());
+	
+	if (UPCGInteractiveToolSettings* ToolSettings = EdMode->GetCurrentToolSettings())
+	{
+		if (UPCGGraphInterface* Graph = PresetItem->Item.LoadSynchronous())
+		{
+			return ToolSettings->GetToolGraph() == Graph ? ECheckBoxState::Checked : ECheckBoxState::Unchecked;
+		}
+	}
+
+	return ECheckBoxState::Unchecked;
+}
+
+EVisibility SPCGToolPresetSection::OnGetVisibility() const
+{
+	if (bHidePresets)
+	{
+		return EVisibility::Collapsed;
+	}
+
+	return HasPresets() ? EVisibility::Visible : EVisibility::Collapsed;
+}
diff --git a/Engine/Plugins/PCG/Source/PCGEditor/Private/EditorMode/Widgets/SPCGToolPresetSection.h b/Engine/Plugins/PCG/Source/PCGEditor/Private/EditorMode/Widgets/SPCGToolPresetSection.h
new file mode 100644
--- /dev/null
+++ b/Engine/Plugins/PCG/Source/PCGEditor/Private/EditorMode/Widgets/SPCGToolPresetSection.h
@@ -0,0 +1,70 @@
+﻿// Copyright Epic Games, Inc. All Rights Reserved.
+
+#pragma once
+
+#include "Widgets/SCompoundWidget.h"
+#include "Styling/SlateTypes.h"
+#include "Templates/SharedPointer.h"
+#include "Templates/ValueOrError.h"
+#include "UObject/StrongObjectPtrTemplates.h"
+#include "Misc/Optional.h"
+#include "SPCGToolPresetSection.generated.h"
+
+class UPCGGraphInterface;
+class UEdMode;
+class UToolMenu;
+
+UCLASS()
+class UPCGToolPresetMenuContext : public UObject
+{
+	GENERATED_BODY()
+
+public:
+	TWeakPtr<class SPCGToolPresetSection> ThisSection;	
+};
+
+class SPCGToolPresetSection : public SCompoundWidget
+{
+public:
+	SLATE_BEGIN_ARGS(SPCGToolPresetSection)
+	{
+	}
+	SLATE_END_ARGS()
+
+	void Construct(const FArguments& InArgs, UEdMode& InOwningEditorMode);
+
+	void UpdatePresets();
+
+	void HidePresets(bool bInHide) { bHidePresets = bInHide; }
+	
+	bool HasPresets() const;
+
+	struct FPresetItem
+	{
+		TSoftObjectPtr<UPCGGraphInterface> Item;
+		FText Name;
+		FText Tooltip;
+	};
+
+	using FPresetItemPtr = TSharedPtr<FPresetItem, ESPMode::ThreadSafe>;
+	
+	const TArray<FPresetItemPtr>& GetPresets() const { return Presets; }
+	
+private:
+	TSharedPtr<SWidget> MakePresetToolbar() const;
+	static void OnGeneratePresetMenu(UToolMenu* ToolMenu);
+
+	static bool CanActivatePreset(TWeakPtr<FPresetItem> PresetItem, const UPCGToolPresetMenuContext* Context);
+	static void ActivatePreset(TWeakPtr<FPresetItem> PresetItem, const UPCGToolPresetMenuContext* Context);
+	static FText GetPresetTooltip(TWeakPtr<FPresetItem> PresetItem, const UPCGToolPresetMenuContext* Context);
+	static ECheckBoxState IsPresetActive(TWeakPtr<FPresetItem> PresetItem, const UPCGToolPresetMenuContext* Context);
+
+	EVisibility OnGetVisibility() const;
+private:
+	TStrongObjectPtr<UPCGToolPresetMenuContext> ThisContext;
+	
+	TWeakObjectPtr<UEdMode> OwningEditorMode;
+	TArray<FPresetItemPtr> Presets;
+	
+	bool bHidePresets = false;
+};
