[PCG] Fix SA warning in PCGCommon

#rnx

diff --git a/Engine/Plugins/PCG/Source/PCG/Private/PCGCommon.cpp b/Engine/Plugins/PCG/Source/PCG/Private/PCGCommon.cpp
--- a/Engine/Plugins/PCG/Source/PCG/Private/PCGCommon.cpp
+++ b/Engine/Plugins/PCG/Source/PCG/Private/PCGCommon.cpp
@@ -248,18 +248,18 @@ uint32 GetTypeHash(const FPCGRuntimeGenerationRadii& InGenerationRadii)
 {
 	uint32 Hash = 0;
 
-	Hash = HashCombine(Hash, InGenerationRadii.GenerationRadius);
-	Hash = HashCombine(Hash, InGenerationRadii.GenerationRadius400);
-	Hash = HashCombine(Hash, InGenerationRadii.GenerationRadius800);
-	Hash = HashCombine(Hash, InGenerationRadii.GenerationRadius1600);
-	Hash = HashCombine(Hash, InGenerationRadii.GenerationRadius3200);
-	Hash = HashCombine(Hash, InGenerationRadii.GenerationRadius6400);
-	Hash = HashCombine(Hash, InGenerationRadii.GenerationRadius12800);
-	Hash = HashCombine(Hash, InGenerationRadii.GenerationRadius25600);
-	Hash = HashCombine(Hash, InGenerationRadii.GenerationRadius51200);
-	Hash = HashCombine(Hash, InGenerationRadii.GenerationRadius102400);
-	Hash = HashCombine(Hash, InGenerationRadii.GenerationRadius204800);
-	Hash = HashCombine(Hash, InGenerationRadii.CleanupRadiusMultiplier);
+	Hash = HashCombine(Hash, GetTypeHash(InGenerationRadii.GenerationRadius));
+	Hash = HashCombine(Hash, GetTypeHash(InGenerationRadii.GenerationRadius400));
+	Hash = HashCombine(Hash, GetTypeHash(InGenerationRadii.GenerationRadius800));
+	Hash = HashCombine(Hash, GetTypeHash(InGenerationRadii.GenerationRadius1600));
+	Hash = HashCombine(Hash, GetTypeHash(InGenerationRadii.GenerationRadius3200));
+	Hash = HashCombine(Hash, GetTypeHash(InGenerationRadii.GenerationRadius6400));
+	Hash = HashCombine(Hash, GetTypeHash(InGenerationRadii.GenerationRadius12800));
+	Hash = HashCombine(Hash, GetTypeHash(InGenerationRadii.GenerationRadius25600));
+	Hash = HashCombine(Hash, GetTypeHash(InGenerationRadii.GenerationRadius51200));
+	Hash = HashCombine(Hash, GetTypeHash(InGenerationRadii.GenerationRadius102400));
+	Hash = HashCombine(Hash, GetTypeHash(InGenerationRadii.GenerationRadius204800));
+	Hash = HashCombine(Hash, GetTypeHash(InGenerationRadii.CleanupRadiusMultiplier));
 
 	return Hash;
 }
