[PCG] Fix GPU static mesh spawner crash trying to use invalid asset as static mesh.

Disables raw buffer optimization due to invalid reuse of string keys, because the spawner will use string keys from the upstream analysis compute graph without remapping to the downstream graph.

#rnx

diff --git a/Engine/Plugins/PCG/Source/PCG/Private/Compute/DataInterfaces/Elements/PCGStaticMeshSpawnerDataInterface.cpp b/Engine/Plugins/PCG/Source/PCG/Private/Compute/DataInterfaces/Elements/PCGStaticMeshSpawnerDataInterface.cpp
--- a/Engine/Plugins/PCG/Source/PCG/Private/Compute/DataInterfaces/Elements/PCGStaticMeshSpawnerDataInterface.cpp
+++ b/Engine/Plugins/PCG/Source/PCG/Private/Compute/DataInterfaces/Elements/PCGStaticMeshSpawnerDataInterface.cpp
@@ -431,6 +431,23 @@ bool UPCGStaticMeshSpawnerDataProvider::PrepareForExecute_GameThread(UPCGDataBin
 			}
 		}
 
+		bool bAnyInvalid = false;
+
+		// User could pass us any soft object path so verify static meshes were loaded.
+		for (const FPCGProceduralISMComponentDescriptor& Descriptor : PrimitiveDescriptors)
+		{
+			if (!Cast<UStaticMesh>(Descriptor.StaticMesh.Get()))
+			{
+				UE_LOG(LogPCG, Error, TEXT("Tried to use asset '%s' as a static mesh."), *Descriptor.StaticMesh.ToString());
+				bAnyInvalid = true;
+			}
+		}
+
+		if (bAnyInvalid)
+		{
+			PrimitiveDescriptors.Reset();
+		}
+
 		bStaticMeshesLoaded = true;
 	}
 
diff --git a/Engine/Plugins/PCG/Source/PCG/Private/Elements/PCGStaticMeshSpawner.cpp b/Engine/Plugins/PCG/Source/PCG/Private/Elements/PCGStaticMeshSpawner.cpp
--- a/Engine/Plugins/PCG/Source/PCG/Private/Elements/PCGStaticMeshSpawner.cpp
+++ b/Engine/Plugins/PCG/Source/PCG/Private/Elements/PCGStaticMeshSpawner.cpp
@@ -35,13 +35,6 @@ static TAutoConsoleVariable<bool> CVarAllowISMReuse(
 	true,
 	TEXT("Controls whether ISMs can be reused and skipped when re-executing"));
 
-static TAutoConsoleVariable<bool> CVarUseRawDataForAnalysis(
-	TEXT("pcg.RuntimeGeneration.ISM.UseRawDataForGPUAnalysis"),
-	true,
-	TEXT("Read back raw data in GPU primitive instancing path rather than constructing an attribute set, to save CPU time. Read-only, must be set from ini."),
-	ECVF_ReadOnly // Structural change, could flush graph compilation cache, but trickier to flush compute graph instance pool in UPCGGraph currently.
-	);
-
 UPCGStaticMeshSpawnerSettings::UPCGStaticMeshSpawnerSettings(const FObjectInitializer &ObjectInitializer)
 {
 	MeshSelectorType = UPCGMeshSelectorWeighted::StaticClass();
@@ -59,7 +52,8 @@ void UPCGStaticMeshSpawnerSettings::CreateKernels(FPCGGPUCompilationContext& InO
 {
 	PCGKernelHelpers::FCreateKernelParams CreateParams(InObjectOuter, this);
 
-	const bool bUseRawBuffer = CVarUseRawDataForAnalysis.GetValueOnAnyThread();
+	// Not supported in 5.7, as the spawner will use string keys from the upstream analysis compute graph without remapping to the downstream graph.
+	const bool bUseRawBuffer = false;
 
 	UPCGStaticMeshSpawnerKernel* SpawnerKernel = PCGKernelHelpers::CreateKernel<UPCGStaticMeshSpawnerKernel>(InOutContext, CreateParams, OutKernels, OutEdges);
 	SpawnerKernel->SetUseRawInstanceCountsBuffer(bUseRawBuffer);
