[PCG] Fix generate landscape textures to deal with subtle details with packing - duplicated rows/columns at common boundaries of landscape components and landscape subsections.

Also divorce the code a bit better from grass maps, the current would not export the height map if no grass maps were present.

#rnx

diff --git a/Engine/Plugins/PCG/Shaders/Private/PCGGrassMapUnpackerCS.usf b/Engine/Plugins/PCG/Shaders/Private/PCGGrassMapUnpackerCS.usf
--- a/Engine/Plugins/PCG/Shaders/Private/PCGGrassMapUnpackerCS.usf
+++ b/Engine/Plugins/PCG/Shaders/Private/PCGGrassMapUnpackerCS.usf
@@ -19,18 +19,27 @@
 Texture2D<float4> InPackedGrassMaps;
 
 // Grass maps output. Each slice is a single grass combined grass map.
-// Dimension is (NumComponentsX * InLandscapeComponentResolution, NumComponentsY * InLandscapeComponentResolution, NumGrassMaps).
+// Each component is an inner grid cell. Res is vert count and includes outer edges, so a row/col is removed from each internal boundary.
 //
-//                                    -------------
-//                                  -------------  |
-//                                 |             | |
-//        Num Comps Y *            |             | |
-//  InLandscapeComponentResolution |             | |
-//                                 |             | |
-//                                 |             |-  NumGrassMaps
-//                                  -------------
-//                                  Num Comps X *
-//                         InLandscapeComponentResolution
+//                         -------------------------------
+//                       -------------------------------  |
+//                      |  Res-1  |  Res-1  |    Res    | |
+//                      |         |         |           | |
+//  (Res-1) * NumTilesY |Res-1    |         |           | |
+//         + 1          |         |         |           | |
+//                      |---------+---------+-----------| |
+//                      |         |         |           | |
+//                      |Res-1    |         |           | |
+//                      |         |         |           | |
+//                      |         |         |           | |
+//                      |---------+---------+-----------| |
+//                      |         |         |           | |
+//                      |         |         |           | |
+//                      |Res      |         |           | |
+//                      |         |         |           | |
+//                      |         |         |           |-  NumGrassMaps slices
+//                       -------------------------------
+//                            (Res-1) * NumTilesX + 1
 //
 RWTexture2DArray<float> OutUnpackedGrassMaps;
 
@@ -38,8 +47,8 @@ RWTexture2DArray<float> OutUnpackedGrassMaps;
 RWTexture2D<float> OutUnpackedHeight;
 
 int4 InLinearTileIndexToComponentIndex[PCG_MAX_NUM_LANDSCAPE_COMPONENTS]; // Only x component used, 16 byte alignment on arrays
-uint InNumTilesX;
-uint InLandscapeComponentResolution;
+uint2 InNumTiles;
+uint InLandscapeComponentResolution; // One texel per vert, so number of quads + 1
 uint InNumGrassMapPasses;
 float3 InLandscapeGridScale; // x == LS Actor DrawScale.X, y == LS Actor DrawScale.y, z == LS Actor DrawScale.z / 128.0f (ZSCALE)
 float InLandscapeLocationZ;
@@ -56,13 +65,21 @@ void PCGGrassMapUnpacker_CS(uint3 DispatchThreadId : SV_DispatchThreadID)
 	}
 
 	// The rendered landscape components in the input are not in any particular order, so lookup which component index for this output tile.
-	const uint2 TileCoords = DispatchThreadId.xy / InLandscapeComponentResolution;
-	const uint TileLinearIndex = TileCoords.y * InNumTilesX + TileCoords.x;
+	// Clamp because we are dividing by resolution-1 and last col/row of texture will overflow.
+	const uint2 TileCoords = min(DispatchThreadId.xy / (InLandscapeComponentResolution - 1u), InNumTiles - 1u);
+
+	const uint TileLinearIndex = TileCoords.y * InNumTiles.x + TileCoords.x;
 	const int ComponentIndex = InLinearTileIndexToComponentIndex[TileLinearIndex][0];
 
 	// Invalid component index means there was no component present for this output tile.
 	if (ComponentIndex == -1)
 	{
+		// First layer of threads also deals with height.
+		if (InOutputHeight == 1 && DispatchThreadId.z == 0)
+		{
+			OutUnpackedHeight[DispatchThreadId.xy] = 0.0f;
+		}
+
 		OutUnpackedGrassMaps[DispatchThreadId] = 0.0f;
 		return;
 	}
@@ -71,7 +88,10 @@ void PCGGrassMapUnpacker_CS(uint3 DispatchThreadId : SV_DispatchThreadID)
 	const uint GrassMapIndex = DispatchThreadId.z;
 	const uint PassIndex = (GrassMapIndex + 2) / 4;
 	const uint ComponentXOffset = (ComponentIndex * InNumGrassMapPasses + PassIndex) * InLandscapeComponentResolution;
-	const uint2 InputCoords = uint2(DispatchThreadId.x % InLandscapeComponentResolution + ComponentXOffset, DispatchThreadId.y % InLandscapeComponentResolution);
+
+	// Coordinates into packed data which are laid out horizontally (see above).
+	const uint2 LocalTileCoords = uint2(DispatchThreadId.x - TileCoords.x * (InLandscapeComponentResolution - 1), DispatchThreadId.y - TileCoords.y * (InLandscapeComponentResolution - 1));
+	const uint2 InputCoords = uint2(LocalTileCoords.x + ComponentXOffset, LocalTileCoords.y);
 
 	// Heights. The first layer threads also deal with height.
 	if (InOutputHeight == 1 && DispatchThreadId.z == 0)
diff --git a/Engine/Plugins/PCG/Source/PCG/Private/Elements/PCGGenerateLandscapeTextures.cpp b/Engine/Plugins/PCG/Source/PCG/Private/Elements/PCGGenerateLandscapeTextures.cpp
--- a/Engine/Plugins/PCG/Source/PCG/Private/Elements/PCGGenerateLandscapeTextures.cpp
+++ b/Engine/Plugins/PCG/Source/PCG/Private/Elements/PCGGenerateLandscapeTextures.cpp
@@ -332,6 +332,7 @@ bool FPCGGenerateLandscapeTexturesElement::ExecuteInternal(FPCGContext* InContex
 						Context->GrassMapBounds = LandscapeComponentBounds;
 
 						Context->LandscapeComponentExtent = LandscapeComponent->Bounds.BoxExtent.X * 2.0;
+						Context->TexelSizeWorld = Context->LandscapeComponentExtent / static_cast<double>(LandscapeProxy->ComponentSizeQuads);
 
 						if (!ensure(FMath::IsNearlyEqual(LandscapeComponent->Bounds.BoxExtent.X, LandscapeComponent->Bounds.BoxExtent.Y)))
 						{
@@ -388,7 +389,7 @@ bool FPCGGenerateLandscapeTexturesElement::ExecuteInternal(FPCGContext* InContex
 			}
 		}
 
-		if (!Context->LandscapeProxy.Get() || Context->LandscapeComponents.IsEmpty() || Context->SelectedGrassTypes.IsEmpty())
+		if (!Context->LandscapeProxy.Get() || Context->LandscapeComponents.IsEmpty() || (Context->SelectedGrassTypes.IsEmpty() && !Settings->bGenerateHeightMap))
 		{
 			return true;
 		}
@@ -411,12 +412,15 @@ bool FPCGGenerateLandscapeTexturesElement::ExecuteInternal(FPCGContext* InContex
 						Context->TexturesToStream.Add(HeightMap);
 					}
 
-					const ERHIFeatureLevel::Type FeatureLevel = World->GetFeatureLevel();
-					for (UTexture2D* WeightmapTexture : LandscapeComponent->GetRenderedWeightmapTexturesForFeatureLevel(FeatureLevel))
+					if (Context->NumGrassTypes > 0)
 					{
-						if (WeightmapTexture)
+						const ERHIFeatureLevel::Type FeatureLevel = World->GetFeatureLevel();
+						for (UTexture2D* WeightmapTexture : LandscapeComponent->GetRenderedWeightmapTexturesForFeatureLevel(FeatureLevel))
 						{
-							Context->TexturesToStream.Add(WeightmapTexture);
+							if (WeightmapTexture)
+							{
+								Context->TexturesToStream.Add(WeightmapTexture);
+							}
 						}
 					}
 				}
@@ -531,18 +535,18 @@ bool FPCGGenerateLandscapeTexturesElement::ExecuteInternal(FPCGContext* InContex
 		Context->LandscapeGrassWeightExporter = new FLandscapeGrassWeightExporter(
 			LandscapeProxy,
 			LandscapeComponents,
-			/*bInNeedsGrassmap=*/true,
+			/*bInNeedsGrassmap=*/Context->NumGrassTypes > 0,
 			Settings->bGenerateHeightMap,
 			/*HeightMips=*/{},
 			/*bExport=*/false,
 			/*bShouldReadback=*/false);
 
-		const FIntPoint& GrassMapTextureTargetSize = Context->LandscapeGrassWeightExporter->GetTargetSize();
+		const FIntPoint& OutputTextureSize = Context->LandscapeGrassWeightExporter->GetTargetSize();
 		const int32 Max2DTextureDimension = GetMax2DTextureDimension();
 
-		if (GrassMapTextureTargetSize.X < 1 || GrassMapTextureTargetSize.X > Max2DTextureDimension || GrassMapTextureTargetSize.Y < 1 || GrassMapTextureTargetSize.Y > Max2DTextureDimension)
+		if (OutputTextureSize.X < 1 || OutputTextureSize.X > Max2DTextureDimension || OutputTextureSize.Y < 1 || OutputTextureSize.Y > Max2DTextureDimension)
 		{
-			PCGLog::LogErrorOnGraph(FText::Format(LOCTEXT("InvalidTextureSize", "Invalid grass map texture size ({0}, {1}), cannot exceed ({2}, {2}). Consider partitioning the component onto a smaller grid size."), GrassMapTextureTargetSize.X, GrassMapTextureTargetSize.Y, Max2DTextureDimension), InContext);
+			PCGLog::LogErrorOnGraph(FText::Format(LOCTEXT("InvalidTextureSize", "Invalid output texture size ({0}, {1}), cannot exceed ({2}, {2}). Consider partitioning the component onto a smaller grid size."), OutputTextureSize.X, OutputTextureSize.Y, Max2DTextureDimension), InContext);
 			return true;
 		}
 
@@ -582,13 +586,14 @@ bool FPCGGenerateLandscapeTexturesElement::ExecuteInternal(FPCGContext* InContex
 			const uint32 NumTilesX = FMath::RoundToInt(GrassMapExtent.X / LandscapeComponentExtent);
 			const uint32 NumTilesY = FMath::RoundToInt(GrassMapExtent.Y / LandscapeComponentExtent);
 
-			// Each corner of a quad in the landscape corresponds to one texel in the grass map.
+			// Each corner of a quad in the landscape corresponds to one texel in the grass map. 
 			const uint32 LandscapeComponentResolution = ComponentSizeQuads + 1;
-			const FIntPoint GrassMapResolution(LandscapeComponentResolution * NumTilesX, LandscapeComponentResolution * NumTilesY);
+			// Internal boundaries between components in the output share a row/column, so these are dropped to get a correct world space output.
+			const FIntPoint OutputResolution((LandscapeComponentResolution - 1u) * NumTilesX + 1u, (LandscapeComponentResolution - 1u) * NumTilesY + 1u);
 			const uint32 NumGrassTypes = Context->NumGrassTypes;
 
 			FRDGTextureDesc GrassHeightMapDesc = FRDGTextureDesc::Create2D(
-				bGenerateHeightMap ? GrassMapResolution : FIntPoint(1, 1), // No default system texture usable for UAV, so always create but with 1x1 resolution
+				bGenerateHeightMap ? OutputResolution : FIntPoint(1, 1), // No default system texture usable for UAV, so always create but with 1x1 resolution
 				PF_R32_FLOAT,
 				FClearValueBinding(),
 				ETextureCreateFlags::ShaderResource | ETextureCreateFlags::UAV);
@@ -598,11 +603,11 @@ bool FPCGGenerateLandscapeTexturesElement::ExecuteInternal(FPCGContext* InContex
 
 			// Output texture is array of textures, one per grass map.
 			FRDGTextureDesc GrassMapDesc = FRDGTextureDesc::Create2DArray(
-				GrassMapResolution,
+				OutputResolution,
 				PF_G8,
 				FClearValueBinding(),
 				ETextureCreateFlags::ShaderResource | ETextureCreateFlags::UAV,
-				NumGrassTypes);
+				FMath::Max(NumGrassTypes, 1u));
 
 			FRDGTextureRef GrassMap = GraphBuilder.CreateTexture(GrassMapDesc, TEXT("PCGLandscapeGrassMapUnpacked"));
 			FRDGTextureUAVRef GrassMapUAV = GraphBuilder.CreateUAV(GrassMap);
@@ -612,11 +617,11 @@ bool FPCGGenerateLandscapeTexturesElement::ExecuteInternal(FPCGContext* InContex
 			Parameters->InPackedGrassMaps = GrassMapTextureSRV;
 			Parameters->OutUnpackedGrassMaps = GrassMapUAV;
 			Parameters->OutUnpackedHeight = GrassHeightMapUAV;
-			Parameters->InNumTilesX = NumTilesX;
+			Parameters->InNumTiles = FUintVector2(NumTilesX, NumTilesY);
 			Parameters->InLandscapeGridScale = static_cast<FVector3f>(LandscapeGridScale);
 			Parameters->InLandscapeLocationZ = static_cast<float>(LandscapeZ);
 			Parameters->InLandscapeComponentResolution = LandscapeComponentResolution;
-			Parameters->InOutputResolution = FUintVector2(GrassMapResolution.X, GrassMapResolution.Y);
+			Parameters->InOutputResolution = FUintVector2(OutputResolution.X, OutputResolution.Y);
 			Parameters->InOutputHeight = bGenerateHeightMap ? 1u : 0u;
 
 			// The first 2 channels are reserved for height data. See illustration of packing in PCGGrassMapUnpackerCS.usf.
@@ -635,8 +640,8 @@ bool FPCGGenerateLandscapeTexturesElement::ExecuteInternal(FPCGContext* InContex
 			}
 
 			TShaderMapRef<FPCGGrassMapUnpackerCS> Shader(GetGlobalShaderMap(GMaxRHIFeatureLevel));
-			const int GroupCountX = FMath::DivideAndRoundUp<int>(GrassMapResolution.X, FPCGGrassMapUnpackerCS::ThreadGroupDim);
-			const int GroupCountY = FMath::DivideAndRoundUp<int>(GrassMapResolution.Y, FPCGGrassMapUnpackerCS::ThreadGroupDim);
+			const int GroupCountX = FMath::DivideAndRoundUp<int>(OutputResolution.X, FPCGGrassMapUnpackerCS::ThreadGroupDim);
+			const int GroupCountY = FMath::DivideAndRoundUp<int>(OutputResolution.Y, FPCGGrassMapUnpackerCS::ThreadGroupDim);
 			const int GroupCountZ = NumGrassTypes;
 
 			FComputeShaderUtils::AddPass(GraphBuilder, RDG_EVENT_NAME("PCGUnpackGrassMap"), ERDGPassFlags::Compute, Shader, Parameters, FIntVector(GroupCountX, GroupCountY, GroupCountZ));
@@ -699,13 +704,22 @@ bool FPCGGenerateLandscapeTexturesElement::ExecuteInternal(FPCGContext* InContex
 		return true;
 	}
 
-	const FTransform OutputTextureTransform(FQuat::Identity, Context->GrassMapBounds.GetCenter(), Context->GrassMapBounds.GetExtent());
+	if (!Context->OutputTextureTransform)
+	{
+		// Dilate the horizontal extents by half a texel in both X and Y. This will place the texel centers on the vert grid.
+		FVector Extent = Context->GrassMapBounds.GetExtent();
+		Extent.X += Context->TexelSizeWorld / 2.0;
+		Extent.Y += Context->TexelSizeWorld / 2.0;
+
+		Context->OutputTextureTransform = FTransform(FQuat::Identity, Context->GrassMapBounds.GetCenter(), Extent);
+	}
+
 	bool bAllTexturesInitialized = true;
 
 	if (Context->HeightHandle)
 	{
 		check(Context->HeightTextureData);
-		bAllTexturesInitialized &= Context->HeightTextureData->Initialize(Context->HeightHandle, /*TextureIndex=*/0, OutputTextureTransform, Settings->bSkipReadbackToCPU);
+		bAllTexturesInitialized &= Context->HeightTextureData->Initialize(Context->HeightHandle, /*TextureIndex=*/0, *Context->OutputTextureTransform, Settings->bSkipReadbackToCPU);
 	}
 
 	if (Context->GrassMapHandle)
@@ -717,7 +731,11 @@ bool FPCGGenerateLandscapeTexturesElement::ExecuteInternal(FPCGContext* InContex
 			check(TextureData);
 
 			// Poll initialize (fine to be called even when initialization was already complete).
-			bAllTexturesInitialized &= TextureData->Initialize(Context->GrassMapHandle, /*TextureIndex=*/Context->SelectedGrassTypes[DataIndex].Get<1>(), OutputTextureTransform, Settings->bSkipReadbackToCPU);
+			bAllTexturesInitialized &= TextureData->Initialize(
+				Context->GrassMapHandle,
+				/*TextureIndex=*/Context->SelectedGrassTypes[DataIndex].Get<1>(),
+				*Context->OutputTextureTransform,
+				Settings->bSkipReadbackToCPU);
 		}
 	}
 
diff --git a/Engine/Plugins/PCG/Source/PCG/Public/Elements/PCGGenerateLandscapeTextures.h b/Engine/Plugins/PCG/Source/PCG/Public/Elements/PCGGenerateLandscapeTextures.h
--- a/Engine/Plugins/PCG/Source/PCG/Public/Elements/PCGGenerateLandscapeTextures.h
+++ b/Engine/Plugins/PCG/Source/PCG/Public/Elements/PCGGenerateLandscapeTextures.h
@@ -120,6 +120,10 @@ public:
 	/** True when grass map generation has been scheduled on the render thread. */
 	bool bGenerationScheduled = false;
 
+	double TexelSizeWorld = 0.0;
+
+	TOptional<FTransform> OutputTextureTransform;
+
 protected:
 	virtual void AddExtraStructReferencedObjects(FReferenceCollector& Collector) override;
 };
diff --git a/Engine/Plugins/PCG/Source/PCGCompute/Internal/PCGGrassMapUnpackerCS.h b/Engine/Plugins/PCG/Source/PCGCompute/Internal/PCGGrassMapUnpackerCS.h
--- a/Engine/Plugins/PCG/Source/PCGCompute/Internal/PCGGrassMapUnpackerCS.h
+++ b/Engine/Plugins/PCG/Source/PCGCompute/Internal/PCGGrassMapUnpackerCS.h
@@ -23,7 +23,7 @@ public:
 		SHADER_PARAMETER_RDG_TEXTURE_UAV(RWTexture2D<float>, OutUnpackedHeight)
 		SHADER_PARAMETER_RDG_TEXTURE_UAV(RWTexture2DArray<float>, OutUnpackedGrassMaps)
 		SHADER_PARAMETER_ARRAY(FIntVector4, InLinearTileIndexToComponentIndex, [MaxNumLandscapeComponents])
-		SHADER_PARAMETER(uint32, InNumTilesX)
+		SHADER_PARAMETER(FUintVector2, InNumTiles)
 		SHADER_PARAMETER(uint32, InLandscapeComponentResolution)
 		SHADER_PARAMETER(uint32, InNumGrassMapPasses)
 		SHADER_PARAMETER(FVector3f, InLandscapeGridScale)
