[PCG] Fixed a crash/infinite loop when selecting actors that aren't valid with the tools. Fixed crash when selecting no actor class. Cleaned up more component pointers so BP actors can be edited while in the tool. Also reviewed pointer type so we don't resurrect components that would have been properly removed. Improved management of actor pivots, especially on newly created actors. Fixed minor UI problem in the attribute llist view.

#rnx

diff --git a/Engine/Plugins/PCG/Source/PCG/Internal/Data/Tool/PCGToolPointData.h b/Engine/Plugins/PCG/Source/PCG/Internal/Data/Tool/PCGToolPointData.h
--- a/Engine/Plugins/PCG/Source/PCG/Internal/Data/Tool/PCGToolPointData.h
+++ b/Engine/Plugins/PCG/Source/PCG/Internal/Data/Tool/PCGToolPointData.h
@@ -46,6 +46,6 @@ protected:
 
 	/** Box component optionally created when the actor doesn't have valid bounds. */
 	UPROPERTY(VisibleAnywhere, Category = "Working Data")
-	TObjectPtr<UBoxComponent> BoxComponent;
+	TSoftObjectPtr<UBoxComponent> BoxComponent;
 #endif // WITH_EDITORONLY_DATA
 };
\ No newline at end of file
diff --git a/Engine/Plugins/PCG/Source/PCG/Internal/Data/Tool/PCGToolSplineData.h b/Engine/Plugins/PCG/Source/PCG/Internal/Data/Tool/PCGToolSplineData.h
--- a/Engine/Plugins/PCG/Source/PCG/Internal/Data/Tool/PCGToolSplineData.h
+++ b/Engine/Plugins/PCG/Source/PCG/Internal/Data/Tool/PCGToolSplineData.h
@@ -44,12 +44,16 @@ private:
 protected:
 	/** Spline this working data interacts with. */
 	UPROPERTY(VisibleAnywhere, Category = "Working Data")
-	TObjectPtr<USplineComponent> SplineComponent;
+	TSoftObjectPtr<USplineComponent> SplineComponent;
 
 #if WITH_EDITORONLY_DATA
 	/** Copy of the spline at the beginning of the interaction, for rollback purposes. */
 	UPROPERTY(Transient)
 	TObjectPtr<USplineComponent> OnToolStartSplineComponent;
+
+	/** Transform of the spline at the beginning of the interaction, for rollback purposes. */
+	UPROPERTY(Transient)
+	FTransform OnToolStartSplineTransform;
 #endif // WITH_EDITORONLY_DATA
 };
 
diff --git a/Engine/Plugins/PCG/Source/PCG/Internal/Data/Tool/PCGToolVolumeData.h b/Engine/Plugins/PCG/Source/PCG/Internal/Data/Tool/PCGToolVolumeData.h
--- a/Engine/Plugins/PCG/Source/PCG/Internal/Data/Tool/PCGToolVolumeData.h
+++ b/Engine/Plugins/PCG/Source/PCG/Internal/Data/Tool/PCGToolVolumeData.h
@@ -21,15 +21,15 @@ struct FPCGInteractiveToolWorkingData_Volume : public FPCGInteractiveToolWorking
 public:
 	/** Volume actor this working data is affecting. */
 	UPROPERTY(VisibleAnywhere, Category = "Working Data")
-	TWeakObjectPtr<AVolume> Volume = nullptr;
+	TObjectPtr<AVolume> Volume = nullptr;
 
 	/** Brush component this working data is affecting. */
 	UPROPERTY(VisibleAnywhere, Category = "Working Data")
-	TWeakObjectPtr<UBrushComponent> BrushComponent = nullptr;
+	TSoftObjectPtr<UBrushComponent> BrushComponent = nullptr;
 
 	/** Box component this working data is affecting. */
 	UPROPERTY(VisibleAnywhere, Category = "Working Data")
-	TWeakObjectPtr<UBoxComponent> BoxComponent = nullptr;
+	TSoftObjectPtr<UBoxComponent> BoxComponent = nullptr;
 
 	//~Begin FPCGInteractiveToolWorkingData interface
 public:
@@ -65,12 +65,16 @@ protected:
 	UPROPERTY(Transient)
 	TObjectPtr<UBrushBuilder> OnToolStartBrushBuilder = nullptr;
 
-	/** Original transform of the volume actor, used for rollback. */
+	/** Original transform of the actor (volume or otherwise), used for rollback. */
 	UPROPERTY(Transient)
-	FTransform OnToolStartVolumeTransform;
+	FTransform OnToolStartActorTransform;
 
 	/** Copy of the original box component on tool initialization, used for rollback purposes. */
 	UPROPERTY(Transient)
 	TObjectPtr<UBoxComponent> OnToolStartBoxComponent = nullptr;
+
+	/** Original transform of the box component, used for rollback. */
+	UPROPERTY(Transient)
+	FTransform OnToolStartBoxTransform;
 #endif
 };
\ No newline at end of file
diff --git a/Engine/Plugins/PCG/Source/PCG/Private/Data/Tool/PCGToolPointData.cpp b/Engine/Plugins/PCG/Source/PCG/Private/Data/Tool/PCGToolPointData.cpp
--- a/Engine/Plugins/PCG/Source/PCG/Private/Data/Tool/PCGToolPointData.cpp
+++ b/Engine/Plugins/PCG/Source/PCG/Private/Data/Tool/PCGToolPointData.cpp
@@ -37,10 +37,18 @@ void FPCGInteractiveToolWorkingData_PointArrayData::InitializeInternal(const FPC
 		BoxComponent->SetBoxExtent(FVector::OneVector);
 		BoxComponent->SetSimulatePhysics(false);
 
-		if (!Context.OwningActor->GetRootComponent())
+		// If there wasn't a root component or it's the default one, we can certainly replace it.
+		if (!Context.OwningActor->GetRootComponent() || Context.OwningActor->GetRootComponent()->GetFName() == USceneComponent::GetDefaultSceneRootVariableName())
 		{
+			USceneComponent* DefaultSceneRootComponent = Context.OwningActor->GetRootComponent();
+
 			BoxComponent->SetMobility(EComponentMobility::Static);
 			Context.OwningActor->SetRootComponent(BoxComponent.Get());
+
+			if (DefaultSceneRootComponent)
+			{
+				DefaultSceneRootComponent->DestroyComponent();
+			}
 		}
 		else
 		{
@@ -51,7 +59,7 @@ void FPCGInteractiveToolWorkingData_PointArrayData::InitializeInternal(const FPC
 		BoxComponent->RegisterComponent();
 		Context.OwningActor->AddInstanceComponent(BoxComponent.Get());
 
-		GeneratedResources.GeneratedComponents.Add(BoxComponent);
+		GeneratedResources.GeneratedComponents.Add(BoxComponent.Get());
 	}
 }
 
@@ -67,15 +75,15 @@ void FPCGInteractiveToolWorkingData_PointArrayData::OnToolStart(const FPCGIntera
 
 void FPCGInteractiveToolWorkingData_PointArrayData::OnToolApply(const FPCGInteractiveToolWorkingDataContext& Context)
 {
-	// Update bounds if needed
-	if (BoxComponent && GeneratedPointData)
+	// Update bounds if needed (e.g. they were originally created from this tool.)
+	if (BoxComponent.IsValid() && GeneratedPointData)
 	{
 		FBox PointBounds = GeneratedPointData->GetBounds();
 		if (PointBounds.IsValid)
 		{
 			BoxComponent->SetRelativeLocation(PointBounds.GetCenter());
 
-			FVector NewExtent = BoxComponent->GetUnscaledBoxExtent().ComponentMax(PointBounds.GetExtent());
+			FVector NewExtent = PointBounds.GetExtent().ComponentMax(FVector::OneVector);
 			BoxComponent->SetBoxExtent(NewExtent);
 		}
 	}
@@ -153,4 +161,4 @@ bool FPCGInteractiveToolWorkingData_PointArrayData::IsValid() const
 UPCGPointArrayData* FPCGInteractiveToolWorkingData_PointArrayData::GetPointArrayData() const
 {
 	return GeneratedPointData;
-}
\ No newline at end of file
+}
diff --git a/Engine/Plugins/PCG/Source/PCG/Private/Data/Tool/PCGToolSplineData.cpp b/Engine/Plugins/PCG/Source/PCG/Private/Data/Tool/PCGToolSplineData.cpp
--- a/Engine/Plugins/PCG/Source/PCG/Private/Data/Tool/PCGToolSplineData.cpp
+++ b/Engine/Plugins/PCG/Source/PCG/Private/Data/Tool/PCGToolSplineData.cpp
@@ -27,23 +27,25 @@ void FPCGInteractiveToolWorkingData_Spline::OnToolStart(const FPCGInteractiveToo
 {
 	Super::OnToolStart(Context);
 
-	if (SplineComponent && !OnToolStartSplineComponent)
+	if (SplineComponent.IsValid() && !OnToolStartSplineComponent)
 	{
 		OnToolStartSplineComponent = Cast<USplineComponent>(StaticDuplicateObject(SplineComponent.Get(), GetTransientPackage()));
+		OnToolStartSplineTransform = SplineComponent->GetComponentTransform();
 	}
 }
 
 void FPCGInteractiveToolWorkingData_Spline::OnToolCancel(const FPCGInteractiveToolWorkingDataContext& Context)
 {
 	// If the Spline Component is valid and isn't a generated component, we revert its state.
-	if (SplineComponent && GeneratedResources.GeneratedComponents.Contains(SplineComponent) == false)
+	if (SplineComponent.IsValid() && GeneratedResources.GeneratedComponents.Contains(SplineComponent.Get()) == false)
 	{
 		if (OnToolStartSplineComponent)
 		{
 			// We restore the spline component to its previous state
-			UEngine::CopyPropertiesForUnrelatedObjects(OnToolStartSplineComponent, SplineComponent);
+			UEngine::CopyPropertiesForUnrelatedObjects(OnToolStartSplineComponent, SplineComponent.Get());
 
 			SplineComponent->UpdateSpline();
+			SplineComponent->SetWorldTransform(OnToolStartSplineTransform);
 
 			FPropertyChangedEvent EmptyPropertyChangedEvent(nullptr);
 			FCoreUObjectDelegates::OnObjectPropertyChanged.Broadcast(SplineComponent.Get(), EmptyPropertyChangedEvent);
@@ -67,12 +69,12 @@ void FPCGInteractiveToolWorkingData_Spline::OnToolShutdown(const FPCGInteractive
 
 void FPCGInteractiveToolWorkingData_Spline::OnResetToolDataRequested(const FPCGInteractiveToolWorkingDataContext& Context)
 {
-	if (SplineComponent)
+	if (SplineComponent.IsValid())
 	{
 		SplineComponent->ClearSplinePoints(/*bUpdateSpline=*/true);
 
 		FPropertyChangedEvent EmptyPropertyChangedEvent(nullptr);
-		FCoreUObjectDelegates::OnObjectPropertyChanged.Broadcast(SplineComponent, EmptyPropertyChangedEvent);
+		FCoreUObjectDelegates::OnObjectPropertyChanged.Broadcast(SplineComponent.Get(), EmptyPropertyChangedEvent);
 	}
 }
 
@@ -155,7 +157,7 @@ void FPCGInteractiveToolWorkingData_Spline::InitializeRuntimeElementData(FPCGCon
 
 	TArray<FPCGTaggedData>& Outputs = Context->OutputData.TaggedData;
 
-	if (SplineComponent)
+	if (SplineComponent.IsValid())
 	{
 		FPCGTaggedData& Output = Outputs.Emplace_GetRef();
 		// We duplicate it so that PCG can take 'ownership'.
@@ -164,10 +166,10 @@ void FPCGInteractiveToolWorkingData_Spline::InitializeRuntimeElementData(FPCGCon
 		// And if we recreate it every time it gets properly managed
 		// If it gets dropped from cache, it's expected to be invalid
 		UPCGSplineData* SplineData = FPCGContext::NewObject_AnyThread<UPCGSplineData>(Context);
-		SplineData->Initialize(SplineComponent);
+		SplineData->Initialize(SplineComponent.Get());
 		Output.Data = SplineData;
 #if WITH_EDITOR
-		FPCGSelectionKey Key = FPCGSelectionKey::CreateFromObjectPtr(SplineComponent);
+		FPCGSelectionKey Key = FPCGSelectionKey::CreateFromPath(SplineComponent.ToSoftObjectPath());
 		FPCGDynamicTrackingHelper::AddSingleDynamicTrackingKey(Context, MoveTemp(Key), false);
 #endif
 	}
@@ -175,7 +177,7 @@ void FPCGInteractiveToolWorkingData_Spline::InitializeRuntimeElementData(FPCGCon
 
 bool FPCGInteractiveToolWorkingData_Spline::IsValid() const
 {
-	return SplineComponent != nullptr;
+	return SplineComponent.IsValid();
 }
 
 #if WITH_EDITOR
diff --git a/Engine/Plugins/PCG/Source/PCG/Private/Data/Tool/PCGToolVolumeData.cpp b/Engine/Plugins/PCG/Source/PCG/Private/Data/Tool/PCGToolVolumeData.cpp
--- a/Engine/Plugins/PCG/Source/PCG/Private/Data/Tool/PCGToolVolumeData.cpp
+++ b/Engine/Plugins/PCG/Source/PCG/Private/Data/Tool/PCGToolVolumeData.cpp
@@ -30,21 +30,21 @@ void FPCGInteractiveToolWorkingData_Volume::InitializeRuntimeElementData(FPCGCon
 
 	TArray<FPCGTaggedData>& Outputs = Context->OutputData.TaggedData;
 
-	if (!Volume.IsValid() && !BoxComponent.IsValid())
+	if (!Volume && !BoxComponent.IsValid())
 	{
 		return;
 	}
 
 	FPCGTaggedData& Output = Outputs.Emplace_GetRef();
 
-	if (Volume.IsValid())
+	if (Volume)
 	{
 		UPCGVolumeData* VolumeData = FPCGContext::NewObject_AnyThread<UPCGVolumeData>(Context);
-		VolumeData->Initialize(Volume.Get());
+		VolumeData->Initialize(Volume);
 		Output.Data = VolumeData;
 
 #if WITH_EDITOR
-		FPCGSelectionKey Key = FPCGSelectionKey::CreateFromObjectPtr(Volume.Get());
+		FPCGSelectionKey Key = FPCGSelectionKey::CreateFromObjectPtr(Volume);
 		FPCGDynamicTrackingHelper::AddSingleDynamicTrackingKey(Context, MoveTemp(Key), false);
 #endif
 	}
@@ -55,7 +55,7 @@ void FPCGInteractiveToolWorkingData_Volume::InitializeRuntimeElementData(FPCGCon
 		Output.Data = BoxData;
 
 #if WITH_EDITOR
-		FPCGSelectionKey Key = FPCGSelectionKey::CreateFromObjectPtr(BoxComponent.Get());
+		FPCGSelectionKey Key = FPCGSelectionKey::CreateFromPath(BoxComponent.ToSoftObjectPath());
 		FPCGDynamicTrackingHelper::AddSingleDynamicTrackingKey(Context, MoveTemp(Key), false);
 #endif
 	}
@@ -93,7 +93,7 @@ void FPCGInteractiveToolWorkingData_Volume::OnToolStart(const FPCGInteractiveToo
 		bUpdateTransform = true;
 	}
 
-	if (Volume.IsValid() && Volume->BrushBuilder && !OnToolStartBrushBuilder)
+	if (Volume && Volume->BrushBuilder && !OnToolStartBrushBuilder)
 	{
 		OnToolStartBrushBuilder = DuplicateObject<UBrushBuilder>(Volume->BrushBuilder, GetTransientPackage());
 		bUpdateTransform = true;
@@ -101,25 +101,26 @@ void FPCGInteractiveToolWorkingData_Volume::OnToolStart(const FPCGInteractiveToo
 
 	if (bUpdateTransform)
 	{
-		if (Volume.IsValid())
+		if (Volume)
 		{
-			OnToolStartVolumeTransform = Volume->GetActorTransform();
+			OnToolStartActorTransform = Volume->GetActorTransform();
 		}
 		else if(BoxComponent.IsValid())
 		{
-			OnToolStartVolumeTransform = BoxComponent->GetComponentTransform();
+			OnToolStartActorTransform = BoxComponent->GetOwner() ? BoxComponent->GetOwner()->GetActorTransform() : FTransform::Identity;
+			OnToolStartBoxTransform = BoxComponent->GetComponentTransform();
 		}
 	}
 }
 
 void FPCGInteractiveToolWorkingData_Volume::OnToolCancel(const FPCGInteractiveToolWorkingDataContext& Context)
 {
-	if (Volume.IsValid() && BrushComponent.IsValid() && !GeneratedResources.GeneratedComponents.Contains(BrushComponent) && OnToolStartBrushComponent)
+	if (Volume && BrushComponent.IsValid() && !GeneratedResources.GeneratedComponents.Contains(BrushComponent.Get()) && OnToolStartBrushComponent)
 	{
 		// Restore the brush component to its previous state
 		UEngine::CopyPropertiesForUnrelatedObjects(OnToolStartBrushComponent, BrushComponent.Get());
 
-		Volume->SetActorTransform(OnToolStartVolumeTransform);
+		Volume->SetActorTransform(OnToolStartActorTransform);
 
 		if (OnToolStartBrushBuilder)
 		{
@@ -130,12 +131,12 @@ void FPCGInteractiveToolWorkingData_Volume::OnToolCancel(const FPCGInteractiveTo
 		FPropertyChangedEvent EmptyPropertyChangedEvent(nullptr);
 		FCoreUObjectDelegates::OnObjectPropertyChanged.Broadcast(BrushComponent.Get(), EmptyPropertyChangedEvent);
 	}
-	else if (BoxComponent.IsValid() && !GeneratedResources.GeneratedComponents.Contains(BoxComponent) && OnToolStartBoxComponent)
+	else if (BoxComponent.IsValid() && !GeneratedResources.GeneratedComponents.Contains(BoxComponent.Get()) && OnToolStartBoxComponent)
 	{
 		// Restore the box component to its previous state
 		UEngine::CopyPropertiesForUnrelatedObjects(OnToolStartBoxComponent, BoxComponent.Get());
 
-		BoxComponent->SetWorldTransform(OnToolStartVolumeTransform);
+		BoxComponent->SetWorldTransform(OnToolStartBoxTransform);
 
 		FPropertyChangedEvent EmptyPropertyChangedEvent(nullptr);
 		FCoreUObjectDelegates::OnObjectPropertyChanged.Broadcast(BoxComponent.Get(), EmptyPropertyChangedEvent);
@@ -164,14 +165,15 @@ void FPCGInteractiveToolWorkingData_Volume::OnToolShutdown(const FPCGInteractive
 		OnToolStartBrushComponent = nullptr;
 	}
 
-	OnToolStartVolumeTransform = FTransform::Identity;
-
 	if (OnToolStartBoxComponent)
 	{
 		OnToolStartBoxComponent->MarkAsGarbage();
 		OnToolStartBoxComponent = nullptr;
 	}
 
+	OnToolStartActorTransform = FTransform::Identity;
+	OnToolStartBoxTransform = FTransform::Identity;
+
 	Super::OnToolShutdown(Context);
 }
 
diff --git a/Engine/Plugins/PCG/Source/PCGEditor/PCGEditor.Build.cs b/Engine/Plugins/PCG/Source/PCGEditor/PCGEditor.Build.cs
--- a/Engine/Plugins/PCG/Source/PCGEditor/PCGEditor.Build.cs
+++ b/Engine/Plugins/PCG/Source/PCGEditor/PCGEditor.Build.cs
@@ -64,6 +64,8 @@ namespace UnrealBuildTool.Rules
 					"SourceControl",
 					"StatusBar",
 					"StructUtilsEditor",
+					"SubobjectEditor",
+					"SubobjectDataInterface",
 					"ToolMenus",
 					"ToolWidgets",
 					"TypedElementFramework",
diff --git a/Engine/Plugins/PCG/Source/PCGEditor/Private/EditorMode/Tools/Line/PCGDrawSplineTool.cpp b/Engine/Plugins/PCG/Source/PCGEditor/Private/EditorMode/Tools/Line/PCGDrawSplineTool.cpp
--- a/Engine/Plugins/PCG/Source/PCGEditor/Private/EditorMode/Tools/Line/PCGDrawSplineTool.cpp
+++ b/Engine/Plugins/PCG/Source/PCGEditor/Private/EditorMode/Tools/Line/PCGDrawSplineTool.cpp
@@ -420,7 +420,7 @@ FInputRayHit UPCGDrawSplineToolBase::IsHitByClick(const FInputDeviceRay& ClickPo
 	return {};
 }
 
-void UPCGDrawSplineToolBase::ClearSplineOnFirstInteraction(UPCGInteractiveToolSettings_SplineBase* ToolSettings)
+void UPCGDrawSplineToolBase::ClearSplineOnFirstInteraction(UPCGInteractiveToolSettings_SplineBase* ToolSettings, const FVector& NewStartingLocation)
 {
 	if(ToolSettings && ToolSettings->DrawMode != PreviousDrawMode)
 	{
@@ -428,6 +428,21 @@ void UPCGDrawSplineToolBase::ClearSplineOnFirstInteraction(UPCGInteractiveToolSe
 		{
 			SplineComponent->bSplineHasBeenEdited = true;
 			SplineComponent->ClearSplinePoints(true);
+
+			// Update the actor pivot if it's a new actor
+			if (ToolSettings->HasSpawnedActor() && SplineComponent->GetOwner() && SplineComponent->GetOwner()->GetRootComponent() != SplineComponent)
+			{
+				SplineComponent->GetOwner()->GetRootComponent()->SetWorldTransform(FTransform(NewStartingLocation));
+
+				// Propagate change so the rest of the engine is aware of the change (esp. rendering/proxies).
+				FPropertyChangedEvent EmptyPropertyChangedEvent(nullptr);
+				FCoreUObjectDelegates::OnObjectPropertyChanged.Broadcast(SplineComponent->GetOwner()->GetRootComponent(), EmptyPropertyChangedEvent);
+			}
+
+			SplineComponent->SetWorldTransform(FTransform(NewStartingLocation));
+
+			FPropertyChangedEvent EmptyPropertyChangedEvent(nullptr);
+			FCoreUObjectDelegates::OnObjectPropertyChanged.Broadcast(SplineComponent, EmptyPropertyChangedEvent);
 		}
 
 		// For "continuous" draw modes like FreeDraw, we reset the spline data so so the user can re-drag as needed.
@@ -445,7 +460,7 @@ void UPCGDrawSplineToolBase::OnClicked(const FInputDeviceRay& ClickPosition)
 	if (TOptional<Scene::FViewHitResult> Result = Scene::ViewportRay(ClickPosition, CachedRayParams))
 	{
 		UPCGInteractiveToolSettings_SplineBase* ToolSettings = GetSplineSettings();
-		ClearSplineOnFirstInteraction(ToolSettings);
+		ClearSplineOnFirstInteraction(ToolSettings, Result->ImpactPosition);
 		
 		USplineComponent* WorkingSpline = ToolSettings->GetWorkingSplineComponent(ToolSettings->DataInstance);
 		switch (ToolSettings->DrawMode)
@@ -506,7 +521,7 @@ void UPCGDrawSplineToolBase::OnClickPress(const FInputDeviceRay& ClickPosition)
 		// @todo_pcg
 		// When the Spline Tool exposes more controls, such as selecting spline points to move them, resetting the spline etc., we will remove this.
 		UPCGInteractiveToolSettings_SplineBase* ToolSettings = GetSplineSettings();
-		ClearSplineOnFirstInteraction(ToolSettings);
+		ClearSplineOnFirstInteraction(ToolSettings, Result->ImpactPosition);
 		
 		AddSplinePoint(Result->ImpactPosition, Result->ImpactNormal);
 		
diff --git a/Engine/Plugins/PCG/Source/PCGEditor/Private/EditorMode/Tools/Line/PCGDrawSplineTool.h b/Engine/Plugins/PCG/Source/PCGEditor/Private/EditorMode/Tools/Line/PCGDrawSplineTool.h
--- a/Engine/Plugins/PCG/Source/PCGEditor/Private/EditorMode/Tools/Line/PCGDrawSplineTool.h
+++ b/Engine/Plugins/PCG/Source/PCGEditor/Private/EditorMode/Tools/Line/PCGDrawSplineTool.h
@@ -133,7 +133,7 @@ private:
 	void AddSplinePoint(const FVector& HitLocation, const FVector& HitNormal) const;
 	FVector GetNewPointNormal(const FVector& HitNormal) const;
 
-	void ClearSplineOnFirstInteraction(UPCGInteractiveToolSettings_SplineBase* InSettings);
+	void ClearSplineOnFirstInteraction(UPCGInteractiveToolSettings_SplineBase* InSettings, const FVector& NewStartingLocation);
 
 	UPROPERTY(Instanced)
 	TObjectPtr<USingleClickOrDragInputBehavior> ClickOrDragBehavior = nullptr;
diff --git a/Engine/Plugins/PCG/Source/PCGEditor/Private/EditorMode/Tools/PCGInteractiveToolSettings.cpp b/Engine/Plugins/PCG/Source/PCGEditor/Private/EditorMode/Tools/PCGInteractiveToolSettings.cpp
--- a/Engine/Plugins/PCG/Source/PCGEditor/Private/EditorMode/Tools/PCGInteractiveToolSettings.cpp
+++ b/Engine/Plugins/PCG/Source/PCGEditor/Private/EditorMode/Tools/PCGInteractiveToolSettings.cpp
@@ -128,15 +128,18 @@ TValueOrError<bool, FText> UPCGInteractiveToolSettings::InitActorAndComponent()
 			// Our spawned actor is also the actor we are editing
 			EditActor = SpawnedActor;
 
-			SpawnedActor->SetActorLabel(NewActorName.ToString(), false);
+			if (SpawnedActor)
+			{
+				SpawnedActor->SetActorLabel(NewActorName.ToString(), false);
 
-			FSelectedObjectsChangeList ChangeList;
-			ChangeList.ModificationType = ESelectedObjectsModificationType::Replace;
-			ChangeList.Actors = { SpawnedActor };
+				FSelectedObjectsChangeList ChangeList;
+				ChangeList.ModificationType = ESelectedObjectsModificationType::Replace;
+				ChangeList.Actors = { SpawnedActor };
 
-			Tool->GetToolManager()->RequestSelectionChange(ChangeList);
+				Tool->GetToolManager()->RequestSelectionChange(ChangeList);
 
-			GeneratedResources.SpawnedActor = SpawnedActor;
+				GeneratedResources.SpawnedActor = SpawnedActor;
+			}
 		}
 	}
 
@@ -309,7 +312,7 @@ void UPCGInteractiveToolSettings::RegisterPropertyWatchers()
 	});
 }
 
-bool UPCGInteractiveToolSettings::InitializeWorkingData(FName DataInstanceIdentifier)
+bool UPCGInteractiveToolSettings::InitializeWorkingData(FName DataInstanceIdentifier, bool bAllowRetryAfterInit)
 {	
 	UPCGComponent* PCGComponent = GetWorkingPCGComponent();
 	const UScriptStruct* Struct = GetWorkingDataType();
@@ -366,7 +369,7 @@ bool UPCGInteractiveToolSettings::InitializeWorkingData(FName DataInstanceIdenti
 			WorkingDataToolStartSet.Add(Context.WorkingDataIdentifier);
 		}
 	}
-	else
+	else if(bAllowRetryAfterInit)
 	{
 		// If the working data is invalid despite being initialized, we delete and request it again.
 		// This can happen when data dependencies have been deleted outside of PCG, such as deleting a spline component on the actor
@@ -376,8 +379,8 @@ bool UPCGInteractiveToolSettings::InitializeWorkingData(FName DataInstanceIdenti
 			PCGComponent->ToolDataContainer.RemoveToolData(Context.WorkingDataIdentifier);
 		}
 
-		// In which case, requesting new working data will recreate a new spline
-		return InitializeWorkingData(DataInstanceIdentifier);
+		// In which case, requesting new working data could recreate new components.
+		return InitializeWorkingData(DataInstanceIdentifier, /*bAllowRetryAfterInit=*/false);
 	}
 
 	return bWorkingDataCreated;
diff --git a/Engine/Plugins/PCG/Source/PCGEditor/Private/EditorMode/Tools/PCGInteractiveToolSettings.h b/Engine/Plugins/PCG/Source/PCGEditor/Private/EditorMode/Tools/PCGInteractiveToolSettings.h
--- a/Engine/Plugins/PCG/Source/PCGEditor/Private/EditorMode/Tools/PCGInteractiveToolSettings.h
+++ b/Engine/Plugins/PCG/Source/PCGEditor/Private/EditorMode/Tools/PCGInteractiveToolSettings.h
@@ -50,9 +50,10 @@ public:
 	
 	/** Initializes or finds the working data associated with the data instance identifier.
 	 * If NAME_None, the ToolTag is used as identifier.
+	 * If bAllowRetryAfterInit is true, if the data was properly initialized, it will retry (once) after. Otherwise, it will not.
 	 * @returns Whether new working data was created and initialized.
 	 * False does not imply working data is invalid; just that it didn't get created. */
-	bool InitializeWorkingData(FName DataInstanceIdentifier = NAME_None);
+	bool InitializeWorkingData(FName DataInstanceIdentifier = NAME_None, bool bAllowRetryAfterInit = true);
 	
 	UPCGGraphInterface* GetToolGraph() const { return ToolGraph; }
 	void SetToolGraph(UPCGGraphInterface* InGraph);
@@ -85,6 +86,14 @@ public:
 	/** Filter details view graph selection. */
 	UFUNCTION()
 	virtual bool GraphAssetFilter(const FAssetData& AssetData) const;
+
+	/** Returns true if the actor being edited has been spawned by this tool. */
+	UFUNCTION()
+	bool HasSpawnedActor() const;
+
+	/** Returns true if the pcg component being edited has been spawned by this tool. */
+	UFUNCTION()
+	bool HasGeneratedPCGComponent() const;
 	
 	UPCGComponent* GetWorkingPCGComponent() const;
 	UPCGGraphInstance* GetGraphInstance() const;
@@ -256,12 +265,6 @@ public:
 	bool bTryToGenerate = false;
 
 private:
-	UFUNCTION()
-	bool HasSpawnedActor() const;
-
-	UFUNCTION()
-	bool HasGeneratedPCGComponent() const;
-
 	UPROPERTY(Transient, meta = (TransientToolProperty))
 	TObjectPtr<UPCGGraphInstance> GraphInstanceCache;
 
diff --git a/Engine/Plugins/PCG/Source/PCGEditor/Private/EditorMode/Tools/Volume/PCGVolumeTool.cpp b/Engine/Plugins/PCG/Source/PCGEditor/Private/EditorMode/Tools/Volume/PCGVolumeTool.cpp
--- a/Engine/Plugins/PCG/Source/PCGEditor/Private/EditorMode/Tools/Volume/PCGVolumeTool.cpp
+++ b/Engine/Plugins/PCG/Source/PCGEditor/Private/EditorMode/Tools/Volume/PCGVolumeTool.cpp
@@ -9,6 +9,7 @@
 
 #include "InteractiveToolManager.h"
 #include "SnappingUtils.h"
+#include "SubobjectDataSubsystem.h"
 #include "ToolTargetManager.h"
 #include "ActorFactories/ActorFactory.h"
 #include "BaseBehaviors/SingleClickOrDragBehavior.h"
@@ -90,10 +91,41 @@ void UPCGInteractiveToolSettings_Volume::SetActorClassToSpawn(TSubclassOf<AActor
 	{
 		if (!Class->IsChildOf(AVolume::StaticClass()))
 		{
-			if (const AActor* ActorCDO = Class->GetDefaultObject<AActor>())
+			bool bActorClassHasBoxComponent = false;
+
+			if (AActor* ActorCDO = Class->GetDefaultObject<AActor>())
 			{
-				bClassIsValid &= ActorCDO->GetComponentByClass<UBoxComponent>() != nullptr;
+				bActorClassHasBoxComponent |= ActorCDO->GetComponentByClass<UBoxComponent>() != nullptr;
+
+				// Check subobjects to cover for BP classes.
+				if (!bActorClassHasBoxComponent && Cast<UBlueprintGeneratedClass>(Class) != nullptr)
+				{
+					TArray<FSubobjectDataHandle> SubobjectHandles;
+					USubobjectDataSubsystem::Get()->GatherSubobjectData(ActorCDO, SubobjectHandles);
+
+					const FSubobjectDataHandle* FoundHandle = SubobjectHandles.FindByPredicate(
+						[](const FSubobjectDataHandle& Handle)
+						{
+							if (FSubobjectData* SubobjectData = Handle.GetData())
+							{
+								if (const UActorComponent* ComponentTemplate = SubobjectData->GetComponentTemplate())
+								{
+									return ComponentTemplate->IsA<UBoxComponent>();
+								}
+							}
+
+							return false;
+						}
+					);
+
+					if (FoundHandle && FoundHandle->IsValid())
+					{
+						bActorClassHasBoxComponent |= (FoundHandle->GetData()->GetObject<UBoxComponent>() != nullptr);
+					}
+				}
 			}
+
+			bClassIsValid &= bActorClassHasBoxComponent;
 		}
 	}
 
@@ -191,16 +223,17 @@ void UPCGVolumeTool::UpdateVolumeFromBox()
 	}
 	else if (UBoxComponent* WorkingBoxComponent = ToolSettings ? ToolSettings->GetWorkingBoxComponent(ToolSettings->DataInstance) : nullptr)
 	{
-		FTransform PreviousBoxTransform = WorkingBoxComponent->GetComponentTransform();
-		WorkingBoxComponent->SetWorldTransform(FTransform(BoxCenter));
+		// If we've spawned this actor, we'll assume we can take over the root transform otherwise, everything will behave weirdly.
+		if (ToolSettings->HasSpawnedActor() && WorkingBoxComponent->GetOwner() && WorkingBoxComponent->GetOwner()->GetRootComponent() != WorkingBoxComponent)
+		{
+			WorkingBoxComponent->GetOwner()->GetRootComponent()->SetWorldTransform(FTransform(BoxCenter));
 
-		const FVector PreviousBoxScale = PreviousBoxTransform.GetScale3D();
-		const FVector InvScaleFactor(
-			(FMath::Abs(PreviousBoxScale.X) > UE_DOUBLE_SMALL_NUMBER ? (1.0 / PreviousBoxScale.X) : 0.0),
-			(FMath::Abs(PreviousBoxScale.Y) > UE_DOUBLE_SMALL_NUMBER ? (1.0 / PreviousBoxScale.Y) : 0.0),
-			(FMath::Abs(PreviousBoxScale.Z) > UE_DOUBLE_SMALL_NUMBER ? (1.0 / PreviousBoxScale.Z) : 0.0));
+			FPropertyChangedEvent EmptyPropertyChangedEvent(nullptr);
+			FCoreUObjectDelegates::OnObjectPropertyChanged.Broadcast(WorkingBoxComponent->GetOwner()->GetRootComponent(), EmptyPropertyChangedEvent);
+		}
 
-		WorkingBoxComponent->SetBoxExtent(0.5 * Extent * InvScaleFactor);
+		WorkingBoxComponent->SetWorldTransform(FTransform(BoxCenter));
+		WorkingBoxComponent->SetBoxExtent(0.5 * Extent);
 
 		FPropertyChangedEvent EmptyPropertyChangedEvent(nullptr);
 		FCoreUObjectDelegates::OnObjectPropertyChanged.Broadcast(WorkingBoxComponent, EmptyPropertyChangedEvent);
@@ -394,4 +427,4 @@ UInteractiveTool* UPCGVolumeToolBuilder::BuildTool(const FToolBuilderState& Scen
 	return NewTool;
 }
 
-#undef LOCTEXT_NAMESPACE
+#undef LOCTEXT_NAMESPACE
\ No newline at end of file
diff --git a/Engine/Plugins/PCG/Source/PCGEditor/Private/Widgets/SPCGEditorGraphAttributeListView.cpp b/Engine/Plugins/PCG/Source/PCGEditor/Private/Widgets/SPCGEditorGraphAttributeListView.cpp
--- a/Engine/Plugins/PCG/Source/PCGEditor/Private/Widgets/SPCGEditorGraphAttributeListView.cpp
+++ b/Engine/Plugins/PCG/Source/PCGEditor/Private/Widgets/SPCGEditorGraphAttributeListView.cpp
@@ -540,6 +540,7 @@ void SPCGEditorGraphAttributeListView::Construct(const FArguments& InArgs, TShar
 				PinComboBox->AsShared()
 			]
 			+SHorizontalBox::Slot()
+			.MinWidth(100.0)
 			.FillWidth(1.0)
 			.Padding(1.0f, 0.0f)
 			[
