[PCG] Fix ClipPaths metadata computation for points

#rnx

diff --git a/Engine/Plugins/PCG/Source/PCG/Private/Elements/Polygon/PCGClipPaths.cpp b/Engine/Plugins/PCG/Source/PCG/Private/Elements/Polygon/PCGClipPaths.cpp
--- a/Engine/Plugins/PCG/Source/PCG/Private/Elements/Polygon/PCGClipPaths.cpp
+++ b/Engine/Plugins/PCG/Source/PCG/Private/Elements/Polygon/PCGClipPaths.cpp
@@ -9,6 +9,8 @@
 #include "Data/PCGPolygon2DData.h"
 #include "Elements/PCGSplitSplines.h"
 #include "Elements/Polygon/PCGPolygon2DUtils.h"
+#include "Helpers/PCGHelpers.h"
+#include "Helpers/PCGPointDataHelpers.h"
 #include "Helpers/PCGPolygon2DProcessingHelpers.h"
 
 #include "Curve/PolygonIntersectionUtils.h"
@@ -102,7 +104,6 @@ bool FPCGClipPathsElement::ExecuteInternal(FPCGContext* Context) const
 		return true;
 	}
 
-	TArray<FVector> Positions;
 	TArray<FVector2D> Positions2D;
 
 	for (const FPCGTaggedData& Input : InputPaths)
@@ -178,21 +179,100 @@ bool FPCGClipPathsElement::ExecuteInternal(FPCGContext* Context) const
 				SplitPointPath->SetNumPoints(ResultPath.Num());
 
 				const bool bHasMetadata = !!(InputPointData->GetAllocatedProperties() & EPCGPointNativeProperties::MetadataEntry);
-				SplitPointPath->AllocateProperties(EPCGPointNativeProperties::Transform | (bHasMetadata ? EPCGPointNativeProperties::MetadataEntry : EPCGPointNativeProperties::None));
+				SplitPointPath->AllocateProperties(EPCGPointNativeProperties::Transform | EPCGPointNativeProperties::Seed | InputPointData->GetAllocatedProperties());
+				SplitPointPath->CopyUnallocatedPropertiesFrom(InputPointData);
 
+				TConstPCGValueRange<FTransform> InputTransformRange = InputPointData->GetConstTransformValueRange();
 				TPCGValueRange<FTransform> PathTransformRange = SplitPointPath->GetTransformValueRange(/*bAllocate=*/false);
-				for(int PathVertexIndex = 0; PathVertexIndex < PathPositions.Num(); ++PathVertexIndex)
-				{
-					PathTransformRange[PathVertexIndex] = FTransform(PathPositions[PathVertexIndex]);
-				}
+				TPCGValueRange<int32> PathSeedRange = SplitPointPath->GetSeedValueRange(/*bAllocate=*/false);
+
+				int32 SegmentStartIndex = 0;
 
-				if (bHasMetadata)
+				for (int PathVertexIndex = 0; PathVertexIndex < PathPositions.Num(); ++PathVertexIndex)
 				{
-					TPCGValueRange<PCGMetadataEntryKey> SplitPointEntryKeysRange = SplitPointPath->GetMetadataEntryValueRange(/*bAllocate=*/false);
-					TArray<PCGMetadataEntryKey> SplitPointEntryKeys = PCGPolygon2DProcessingHelpers::ComputeEntryKeysBySampling(InputPointData, PathPositions, SplitPointPath->MutableMetadata());
-					for (int Index = 0; Index < SplitPointEntryKeys.Num(); ++Index)
+					const FVector PathPosition = PathPositions[PathVertexIndex];
+					const int32 SegmentEndIndex = SegmentStartIndex + 1;
+
+					// By construction, the result path should be in the same order as the input path, so we should in theory never go out of bounds
+					// But if we ever do, just fallback on the current point
+					if (!ensure(InputTransformRange.IsValidIndex(SegmentEndIndex)))
+					{
+						InputPointData->CopyPointsTo(SplitPointPath, SegmentStartIndex, PathVertexIndex, 1);
+						continue;
+					}
+
+					const FVector SegmentStart = InputTransformRange[SegmentStartIndex].GetLocation();
+					const FVector SegmentEnd = InputTransformRange[SegmentEndIndex].GetLocation();
+
+					const FVector StartToEndSegmentVector = SegmentEnd - SegmentStart;
+					const double SegmentSquaredLength = StartToEndSegmentVector.Dot(StartToEndSegmentVector);
+
+					if (FMath::IsNearlyZero(SegmentSquaredLength))
+					{
+						// Segment is of size 0, try again with next segment
+						SegmentStartIndex++;
+						PathVertexIndex--;
+						continue;
+					}
+
+					const FVector StartToPathPositionVector = PathPositions[PathVertexIndex] - SegmentStart;
+					const double Dot = StartToEndSegmentVector.Dot(StartToPathPositionVector);
+					const double Alpha = Dot / SegmentSquaredLength;
+					const FVector Cross = StartToEndSegmentVector.Cross(StartToPathPositionVector);
+
+					if (Alpha < -UE_SMALL_NUMBER || Alpha > (1 + UE_SMALL_NUMBER) || !Cross.IsNearlyZero())
+					{
+						// If we're not within the segment (Alpha not in [0,1]), or not aligned (cross != 0):
+						// Try again with next segment
+						SegmentStartIndex++;
+						PathVertexIndex--;
+						continue;
+					}
+
+					if (FMath::IsNearlyZero(Alpha))
+					{
+						// First point, copy over and continue
+						InputPointData->CopyPointsTo(SplitPointPath, SegmentStartIndex, PathVertexIndex, 1);
+					}
+					else if (FMath::IsNearlyEqual(Alpha, 1.0))
+					{
+						// Second point, copy over, and increase the segment start index
+						InputPointData->CopyPointsTo(SplitPointPath, SegmentEndIndex, PathVertexIndex, 1);
+						SegmentStartIndex++;
+					}
+					else
 					{
-						SplitPointEntryKeysRange[Index] = SplitPointEntryKeys[Index];
+						// In between
+						// We need to do a last verification, that the result direction is the same as the input segment direction, because if the input path
+						// cuts itself, we might have multiple points that satisfy the alignment condition
+						// Also by construction, a point in between vertices is always the first or last point of the result path
+						FVector ResultDir{};
+						if (PathVertexIndex == 0)
+						{
+							ResultDir = PathPositions[PathVertexIndex + 1] - PathPosition;
+						}
+						else
+						{
+							ResultDir = PathPosition - PathPositions[PathVertexIndex - 1];
+						}
+						
+						const FVector ResultCross = StartToEndSegmentVector.Cross(ResultDir);
+						if (!ResultCross.IsNearlyZero())
+						{
+							SegmentStartIndex++;
+							PathVertexIndex--;
+							continue;
+						}
+						
+						// We're in between, do an interpolation
+						TStaticArray<TPair<int32, float>, 2> Coefficients;
+						Coefficients[0] = {SegmentStartIndex, 1.0 - Alpha};
+						Coefficients[1] = {SegmentEndIndex, Alpha};
+						
+						// Do the weighted average for all, but override the final position and the seed.
+						PCGPointDataHelpers::WeightedAverage(InputPointData, Coefficients, SplitPointPath, PathVertexIndex, /*bApplyOnMetadata=*/bHasMetadata);
+						PathTransformRange[PathVertexIndex].SetLocation(PathPosition);
+						PathSeedRange[PathVertexIndex] = PCGHelpers::ComputeSeedFromPosition(PathPosition);
 					}
 				}
 
diff --git a/Engine/Plugins/PCG/Source/PCG/Private/Helpers/PCGPointDataHelpers.cpp b/Engine/Plugins/PCG/Source/PCG/Private/Helpers/PCGPointDataHelpers.cpp
new file mode 100644
--- /dev/null
+++ b/Engine/Plugins/PCG/Source/PCG/Private/Helpers/PCGPointDataHelpers.cpp
@@ -0,0 +1,112 @@
+﻿// Copyright Epic Games, Inc. All Rights Reserved.
+
+#include "Helpers/PCGPointDataHelpers.h"
+
+#include "Data/PCGBasePointData.h"
+
+void PCGPointDataHelpers::WeightedAverage(const UPCGBasePointData* InPointData, TConstArrayView<TPair<int32, float>> Coefficients, UPCGBasePointData* OutPointData, int32 OutPointIndex, bool bApplyOnMetadata)
+{
+	if (!InPointData
+		|| !OutPointData
+		|| Coefficients.IsEmpty()
+		|| !Algo::AllOf(Coefficients, [InPointData](const TPair<int32, float>& It) { return It.Key >= 0 && It.Key < InPointData->GetNumPoints() && It.Value >= 0.f; })
+		|| OutPointIndex < 0
+		|| OutPointIndex >= OutPointData->GetNumPoints())
+	{
+		return;
+	}
+	
+	EPCGPointNativeProperties Properties = InPointData->GetAllocatedProperties();
+	if (Properties == EPCGPointNativeProperties::None)
+	{
+		// No allocated properties, nothing to do
+		return;
+	}
+
+	if (Coefficients.Num() == 1)
+	{
+		InPointData->CopyPointsTo(OutPointData, Coefficients[0].Key, OutPointIndex, /*Count=*/1);
+		return;
+	}
+
+	FConstPCGPointValueRanges InPointValueRanges(InPointData);
+
+	OutPointData->AllocateProperties(Properties);
+
+	FPCGPointValueRanges OutPointValueRanges(OutPointData, /*bAllocate=*/false);
+
+	int32 StrongestIndex = INDEX_NONE;
+	float MaxWeight = -1.0f;
+	float SumWeight = 0.0f;
+	for (const auto [Index, Weight] : Coefficients)
+	{
+		SumWeight += Weight;
+		if (Weight > MaxWeight)
+		{
+			StrongestIndex = Index;
+			MaxWeight = Weight;
+		}
+	}
+	
+	if (FMath::IsNearlyZero(SumWeight) || !ensure(StrongestIndex != INDEX_NONE))
+	{
+		// Copy the first point if all the coefficients are 0
+		InPointData->CopyPointsTo(OutPointData, Coefficients[0].Key, OutPointIndex, /*Count=*/1);
+		return;
+	}
+
+	auto DoAverageIfAllocated = [&Coefficients, Properties, OutPointIndex, SumWeight]<typename T>(TPCGValueRange<T>& ToRange, const TConstPCGValueRange<T>& FromRange, EPCGPointNativeProperties Property)
+	{
+		if (EnumHasAllFlags(Properties, Property))
+		{
+			T& OutValue = ToRange[OutPointIndex];
+			OutValue = PCG::Private::MetadataTraits<T>::ZeroValueForWeightedSum();
+			
+			for (const auto [Index, Weight] : Coefficients) 
+			{
+				OutValue = PCG::Private::MetadataTraits<T>::WeightedSum(OutValue, FromRange[Index], Weight / SumWeight);
+			}
+			
+			if constexpr (PCG::Private::MetadataTraits<T>::InterpolationNeedsNormalization)
+			{
+				static_assert(PCG::Private::MetadataTraits<T>::CanNormalize);
+
+				// We need to normalize the resulting value
+				PCG::Private::MetadataTraits<T>::Normalize(OutValue);
+			}
+		}
+	};
+
+	DoAverageIfAllocated(OutPointValueRanges.TransformRange, InPointValueRanges.TransformRange, EPCGPointNativeProperties::Transform);
+	DoAverageIfAllocated(OutPointValueRanges.DensityRange, InPointValueRanges.DensityRange, EPCGPointNativeProperties::Density);
+	DoAverageIfAllocated(OutPointValueRanges.BoundsMinRange, InPointValueRanges.BoundsMinRange, EPCGPointNativeProperties::BoundsMin);
+	DoAverageIfAllocated(OutPointValueRanges.BoundsMaxRange, InPointValueRanges.BoundsMaxRange, EPCGPointNativeProperties::BoundsMax);
+	DoAverageIfAllocated(OutPointValueRanges.ColorRange, InPointValueRanges.ColorRange, EPCGPointNativeProperties::Color);
+	DoAverageIfAllocated(OutPointValueRanges.SteepnessRange, InPointValueRanges.SteepnessRange, EPCGPointNativeProperties::Steepness);
+
+	// For seed and metadata entry, we can't average, so take the highest coefficient (for metadata entry, only if we are not averaging metadata)
+	if (EnumHasAnyFlags(Properties, EPCGPointNativeProperties::Seed))
+	{
+		OutPointValueRanges.SeedRange[OutPointIndex] = InPointValueRanges.SeedRange[StrongestIndex];
+	}
+
+	if (EnumHasAnyFlags(Properties, EPCGPointNativeProperties::MetadataEntry))
+	{
+		const UPCGMetadata* InMetadata = InPointData->ConstMetadata();
+		UPCGMetadata* OutMetadata = OutPointData->MutableMetadata();
+		
+		if (bApplyOnMetadata && InMetadata && OutMetadata)
+		{
+			OutPointValueRanges.MetadataEntryRange[OutPointIndex] = OutMetadata->AddEntry();
+			TArray<TPair<PCGMetadataEntryKey, float>> AttributeCoefficients;
+			AttributeCoefficients.Reserve(Coefficients.Num());
+			Algo::Transform(Coefficients, AttributeCoefficients, [&Range = InPointValueRanges.MetadataEntryRange](const TPair<int32, float>& It) -> TPair<PCGMetadataEntryKey, float> { return {Range[It.Key], It.Value}; });
+			
+			OutMetadata->ComputeWeightedAttribute(OutPointValueRanges.MetadataEntryRange[OutPointIndex], AttributeCoefficients, InMetadata);
+		}
+		else
+		{
+			OutPointValueRanges.MetadataEntryRange[OutPointIndex] = InPointValueRanges.MetadataEntryRange[StrongestIndex];
+		}
+	}
+}
diff --git a/Engine/Plugins/PCG/Source/PCG/Private/Helpers/PCGPointDataHelpers.h b/Engine/Plugins/PCG/Source/PCG/Private/Helpers/PCGPointDataHelpers.h
new file mode 100644
--- /dev/null
+++ b/Engine/Plugins/PCG/Source/PCG/Private/Helpers/PCGPointDataHelpers.h
@@ -0,0 +1,24 @@
+﻿// Copyright Epic Games, Inc. All Rights Reserved.
+
+#pragma once
+
+#include "PCGCommon.h"
+
+#include "Containers/ContainersFwd.h"
+#include "Templates/Tuple.h"
+
+class UPCGBasePointData;
+
+namespace PCGPointDataHelpers
+{
+	/**
+	 * Do the weighted average of all the input points for allocated properties. Note that unallocated properties are untouched, it is the caller
+	 * responsibility to make sure those are copied from in point data to out point data.
+	 * @param InPointData          Input point data to read from
+	 * @param Coefficients         Array of pairs of input point index and their coefficients
+	 * @param OutPointData         Output point data to write to
+	 * @param OutPointIndex        Output point index to write to
+	 * @param bApplyOnMetadata     If the metadata should also do the average
+	 */
+	void WeightedAverage(const UPCGBasePointData* InPointData, TConstArrayView<TPair<int32, float>> Coefficients, UPCGBasePointData* OutPointData, int32 OutPointIndex, bool bApplyOnMetadata);
+}
