[PCG][BugHawk] Fix Z group count for unpack compute shader.

#rnx

diff --git a/Engine/Plugins/PCG/Source/PCG/Private/Elements/PCGGenerateLandscapeTextures.cpp b/Engine/Plugins/PCG/Source/PCG/Private/Elements/PCGGenerateLandscapeTextures.cpp
--- a/Engine/Plugins/PCG/Source/PCG/Private/Elements/PCGGenerateLandscapeTextures.cpp
+++ b/Engine/Plugins/PCG/Source/PCG/Private/Elements/PCGGenerateLandscapeTextures.cpp
@@ -591,6 +591,7 @@ bool FPCGGenerateLandscapeTexturesElement::ExecuteInternal(FPCGContext* InContex
 			// Internal boundaries between components in the output share a row/column, so these are dropped to get a correct world space output.
 			const FIntPoint OutputResolution((LandscapeComponentResolution - 1u) * NumTilesX + 1u, (LandscapeComponentResolution - 1u) * NumTilesY + 1u);
 			const uint32 NumGrassTypes = Context->NumGrassTypes;
+			const uint32 NumOutputSlices = FMath::Max(NumGrassTypes, 1u);
 
 			FRDGTextureDesc GrassHeightMapDesc = FRDGTextureDesc::Create2D(
 				bGenerateHeightMap ? OutputResolution : FIntPoint(1, 1), // No default system texture usable for UAV, so always create but with 1x1 resolution
@@ -607,7 +608,7 @@ bool FPCGGenerateLandscapeTexturesElement::ExecuteInternal(FPCGContext* InContex
 				PF_G8,
 				FClearValueBinding(),
 				ETextureCreateFlags::ShaderResource | ETextureCreateFlags::UAV,
-				FMath::Max(NumGrassTypes, 1u));
+				NumOutputSlices);
 
 			FRDGTextureRef GrassMap = GraphBuilder.CreateTexture(GrassMapDesc, TEXT("PCGLandscapeGrassMapUnpacked"));
 			FRDGTextureUAVRef GrassMapUAV = GraphBuilder.CreateUAV(GrassMap);
@@ -642,7 +643,7 @@ bool FPCGGenerateLandscapeTexturesElement::ExecuteInternal(FPCGContext* InContex
 			TShaderMapRef<FPCGGrassMapUnpackerCS> Shader(GetGlobalShaderMap(GMaxRHIFeatureLevel));
 			const int GroupCountX = FMath::DivideAndRoundUp<int>(OutputResolution.X, FPCGGrassMapUnpackerCS::ThreadGroupDim);
 			const int GroupCountY = FMath::DivideAndRoundUp<int>(OutputResolution.Y, FPCGGrassMapUnpackerCS::ThreadGroupDim);
-			const int GroupCountZ = NumGrassTypes;
+			const int GroupCountZ = NumOutputSlices;
 
 			FComputeShaderUtils::AddPass(GraphBuilder, RDG_EVENT_NAME("PCGUnpackGrassMap"), ERDGPassFlags::Compute, Shader, Parameters, FIntVector(GroupCountX, GroupCountY, GroupCountZ));
 
