[PCG] Fix monochrome instances when using PerInstanceRandom MG node

Support PerInstanceRandom for GPU instancing

#rnx

diff --git a/Engine/Plugins/Experimental/FastGeoStreaming/Source/FastGeoStreaming/Private/FastGeoProceduralISMComponent.cpp b/Engine/Plugins/Experimental/FastGeoStreaming/Source/FastGeoStreaming/Private/FastGeoProceduralISMComponent.cpp
--- a/Engine/Plugins/Experimental/FastGeoStreaming/Source/FastGeoStreaming/Private/FastGeoProceduralISMComponent.cpp
+++ b/Engine/Plugins/Experimental/FastGeoStreaming/Source/FastGeoStreaming/Private/FastGeoProceduralISMComponent.cpp
@@ -126,6 +126,7 @@ TSharedPtr<FInstanceDataSceneProxy, ESPMode::ThreadSafe>& FFastGeoProceduralISMC
 	View.NumInstancesGPUOnly = NumInstances;
 	View.NumCustomDataFloats = NumCustomDataFloats;
 	View.Flags.bHasPerInstanceCustomData = NumCustomDataFloats > 0;
+	View.Flags.bHasPerInstanceRandom = PrimitiveMaterialDesc.bAnyMaterialHasPerInstanceRandom;
 
 	InstanceSceneDataBuffers.EndWriteAccess(AccessTag);
 	InstanceSceneDataBuffers.ValidateData();
diff --git a/Engine/Plugins/PCG/Shaders/Private/PCGSceneWriter.usf b/Engine/Plugins/PCG/Shaders/Private/PCGSceneWriter.usf
--- a/Engine/Plugins/PCG/Shaders/Private/PCGSceneWriter.usf
+++ b/Engine/Plugins/PCG/Shaders/Private/PCGSceneWriter.usf
@@ -15,6 +15,23 @@ StructuredBuffer<uint> InWriteCounters;
 uint InPrimitiveId;
 uint InCustomDataCount;
 uint InPayloadDataFlags;
+uint InSeed;
+
+// This follows the implementation in PCGHelpers.
+uint ComputeSeed(uint A, uint B)
+{
+	return ((A * 196314165U) + 907633515U) ^ ((B * 73148459U) + 453816763U);
+}
+
+// This follows the implementation in FRandomStream.
+float FRand(inout uint InOutSeed)
+{
+	InOutSeed = (InOutSeed * 196314165U) + 907633515U;
+	
+	const float Result = asfloat(0x3F800000U | (InOutSeed >> 9));
+
+	return Result - 1.0f;
+}
 
 void CullInstance(uint InstanceId)
 {
@@ -51,6 +68,8 @@ void SceneWriter_CS(const uint ThreadId : SV_DispatchThreadID)
 	LocalToWorld[3] = float4(0.0f, 0.0f, 0.0f, 1.0f);
 	LocalToWorld = transpose(LocalToWorld);
 
+	uint Seed = ComputeSeed(InSeed, InstanceId);
+
 	// todo_pcg: We should write local space transforms to the instance data buffer from the SM Spawner and use InitializeInstanceSceneData (non-WS) here instead.
 	InitializeInstanceSceneDataWS(
 		InstanceId,
@@ -58,7 +77,7 @@ void SceneWriter_CS(const uint ThreadId : SV_DispatchThreadID)
 		LocalInstanceId,
  		InPayloadDataFlags,
  		InCustomDataCount,
- 		/*RandomId=*/0.0f,
+ 		FRand(Seed),
 		DFFromTileOffset(LWCPromote(LocalToWorld))
 	);
 	
diff --git a/Engine/Plugins/PCG/Source/PCG/Private/Components/PCGProceduralISMComponent.cpp b/Engine/Plugins/PCG/Source/PCG/Private/Components/PCGProceduralISMComponent.cpp
--- a/Engine/Plugins/PCG/Source/PCG/Private/Components/PCGProceduralISMComponent.cpp
+++ b/Engine/Plugins/PCG/Source/PCG/Private/Components/PCGProceduralISMComponent.cpp
@@ -12,6 +12,7 @@
 #include "InstanceDataSceneProxy.h"
 #include "InstancedStaticMeshSceneProxyDesc.h"
 #include "NaniteSceneProxy.h"
+#include "PrimitiveComponentHelper.h"
 #include "PrimitiveSceneDesc.h"
 #include "PrimitiveSceneProxyDesc.h"
 #include "SceneInterface.h"
@@ -338,6 +339,9 @@ void UPCGProceduralISMComponent::SetMinDrawDistance(int32 InMinDrawDistance)
 
 void UPCGProceduralISMComponent::GetSceneProxyDesc(FInstancedStaticMeshSceneProxyDesc& OutSceneProxyDesc) const
 {
+	const UWorld* World = GetWorld();
+	const bool bHasPerInstanceRandom = ensure(World && World->Scene) && FPrimitiveComponentHelper::GetUsedMaterialPropertyDesc(*this, World->Scene->GetShaderPlatform()).bAnyMaterialHasPerInstanceRandom;
+
 	FInstanceSceneDataBuffers InstanceSceneDataBuffers(/*InbInstanceDataIsGPUOnly=*/true);
 	{
 		FInstanceSceneDataBuffers::FAccessTag AccessTag(PointerHash(this));
@@ -351,6 +355,7 @@ void UPCGProceduralISMComponent::GetSceneProxyDesc(FInstancedStaticMeshSceneProx
 		ProxyData.InstanceLocalBounds[0] = GetStaticMesh()->GetBounds();
 
 		ProxyData.Flags.bHasPerInstanceCustomData = ProxyData.NumCustomDataFloats > 0;
+		ProxyData.Flags.bHasPerInstanceRandom = bHasPerInstanceRandom;
 
 		InstanceSceneDataBuffers.EndWriteAccess(AccessTag);
 		InstanceSceneDataBuffers.ValidateData();
diff --git a/Engine/Plugins/PCG/Source/PCG/Private/Compute/DataInterfaces/PCGInstanceDataInterface.cpp b/Engine/Plugins/PCG/Source/PCG/Private/Compute/DataInterfaces/PCGInstanceDataInterface.cpp
--- a/Engine/Plugins/PCG/Source/PCG/Private/Compute/DataInterfaces/PCGInstanceDataInterface.cpp
+++ b/Engine/Plugins/PCG/Source/PCG/Private/Compute/DataInterfaces/PCGInstanceDataInterface.cpp
@@ -209,11 +209,14 @@ FComputeDataProviderRenderProxy* UPCGInstanceDataProvider::GetRenderProxy()
 {
 	TRACE_CPUPROFILER_EVENT_SCOPE(UPCGInstanceDataProvider::GetRenderProxy);
 
+	const UPCGSettings* Settings = GetProducerKernel() ? GetProducerKernel()->GetSettings() : nullptr;
+
 	FPCGInstanceDataProviderProxy::FParams Params =
 	{
 		.PrimitiveFactory = PrimitiveFactory,
 		.NumInstancesAllPrimitives = NumInstancesAllPrimitives,
 		.NumCustomFloatsPerInstance = NumCustomFloatsPerInstance,
+		.Seed = Settings ? static_cast<uint32>(Settings->GetSeed()) : 42,
 		.DataProvider = MakeWeakObjectPtr(this),
 		.ContextHandle = ContextHandle,
 	};
@@ -236,6 +239,7 @@ FPCGInstanceDataProviderProxy::FPCGInstanceDataProviderProxy(const FParams& InPa
 	: PrimitiveFactory(InParams.PrimitiveFactory)
 	, NumInstancesAllPrimitives(InParams.NumInstancesAllPrimitives)
 	, NumCustomFloatsPerInstance(InParams.NumCustomFloatsPerInstance)
+	, Seed(InParams.Seed)
 	, DataProvider(InParams.DataProvider)
 	, ContextHandle(InParams.ContextHandle)
 {
@@ -344,7 +348,7 @@ void FPCGInstanceDataProviderProxy::PostSubmit(FComputeContext& Context) const
 	const TRefCountPtr<FRDGPooledBuffer> InstanceCustomFloatDataExported = Context.GraphBuilder.ConvertToExternalBuffer(InstanceCustomFloatData);
 	const TRefCountPtr<FRDGPooledBuffer> WriteCountersExported = Context.GraphBuilder.ConvertToExternalBuffer(WriteCounters);
 	
-	FTSTicker::GetCoreTicker().AddTicker(TEXT("ApplyPrimitiveSceneUpdates"), /*InDelay=*/0.0f, [DataProviderWeak=DataProvider, DataProviderGeneration=DataProviderGeneration, InstanceDataExported, InstanceCustomFloatDataExported, WriteCountersExported, ContextHandle=ContextHandle](float)
+	FTSTicker::GetCoreTicker().AddTicker(TEXT("ApplyPrimitiveSceneUpdates"), /*InDelay=*/0.0f, [DataProviderWeak=DataProvider, DataProviderGeneration=DataProviderGeneration, InstanceDataExported, InstanceCustomFloatDataExported, WriteCountersExported, ContextHandle=ContextHandle, Seed=Seed](float)
 	{
 		UPCGInstanceDataProvider* DataProviderInner = DataProviderWeak.Get();
 		if (!DataProviderInner || DataProviderInner->GetGenerationCounter() != DataProviderGeneration)
@@ -416,7 +420,7 @@ void FPCGInstanceDataProviderProxy::PostSubmit(FComputeContext& Context) const
 			FPrimitiveSceneDesc PrimitiveSceneDesc;
 			PrimitiveSceneDesc.SceneProxy = SceneProxy;
 
-			Scene->UpdatePrimitiveInstancesFromCompute(&PrimitiveSceneDesc, FGPUSceneWriteDelegate::CreateLambda([PrimitiveIndex, PrimitiveNumInstances, CumulativeInstanceCount, bTriggerCapture, InstanceDataExported, InstanceCustomFloatDataExported, WriteCountersExported](FRDGBuilder& InGraphBuilder, const FGPUSceneWriteDelegateParams& Params)
+			Scene->UpdatePrimitiveInstancesFromCompute(&PrimitiveSceneDesc, FGPUSceneWriteDelegate::CreateLambda([PrimitiveIndex, PrimitiveNumInstances, CumulativeInstanceCount, Seed, bTriggerCapture, InstanceDataExported, InstanceCustomFloatDataExported, WriteCountersExported](FRDGBuilder& InGraphBuilder, const FGPUSceneWriteDelegateParams& Params)
 			{
 				RDG_EVENT_SCOPE(InGraphBuilder, "SceneComputeUpdateInterface->EnqueueUpdate");
 				check(Params.PersistentPrimitiveId != (uint32)INDEX_NONE);
@@ -443,6 +447,7 @@ void FPCGInstanceDataProviderProxy::PostSubmit(FComputeContext& Context) const
 				PRAGMA_ENABLE_DEPRECATION_WARNINGS
 				Parameters->InCustomDataCount = Params.NumCustomDataFloats;
 				Parameters->InPayloadDataFlags = Params.PackedInstanceSceneDataFlags;
+				Parameters->InSeed = Seed;
 
 #if PCG_INSTANCE_DATA_LOGGING
 				UE_LOG(LogPCG, Log, TEXT("\tScene writer delegate [%d]:\tPrimitive ID %u,\tsource instance offset %u,\tInstanceSceneDataOffset %u, num instances %u"),
diff --git a/Engine/Plugins/PCG/Source/PCG/Private/Compute/DataInterfaces/PCGInstanceDataInterface.h b/Engine/Plugins/PCG/Source/PCG/Private/Compute/DataInterfaces/PCGInstanceDataInterface.h
--- a/Engine/Plugins/PCG/Source/PCG/Private/Compute/DataInterfaces/PCGInstanceDataInterface.h
+++ b/Engine/Plugins/PCG/Source/PCG/Private/Compute/DataInterfaces/PCGInstanceDataInterface.h
@@ -76,6 +76,7 @@ public:
 		TSharedPtr<IPCGRuntimePrimitiveFactory> PrimitiveFactory;
 		uint32 NumInstancesAllPrimitives = 0;
 		uint32 NumCustomFloatsPerInstance = 0;
+		uint32 Seed = 42;
 		TWeakObjectPtr<UPCGInstanceDataProvider> DataProvider;
 		TWeakPtr<FPCGContextHandle> ContextHandle;
 	};
@@ -113,6 +114,7 @@ protected:
 
 	uint32 NumInstancesAllPrimitives = 0;
 	uint32 NumCustomFloatsPerInstance = 0;
+	uint32 Seed = 42;
 
 	TWeakObjectPtr<UPCGInstanceDataProvider> DataProvider;
 	uint64 DataProviderGeneration = 0;
diff --git a/Engine/Plugins/PCG/Source/PCGCompute/Internal/PCGSceneWriterCS.h b/Engine/Plugins/PCG/Source/PCGCompute/Internal/PCGSceneWriterCS.h
--- a/Engine/Plugins/PCG/Source/PCGCompute/Internal/PCGSceneWriterCS.h
+++ b/Engine/Plugins/PCG/Source/PCGCompute/Internal/PCGSceneWriterCS.h
@@ -25,6 +25,7 @@ public:
 		SHADER_PARAMETER(uint32, InPrimitiveId)
 		SHADER_PARAMETER(uint32, InCustomDataCount)
 		SHADER_PARAMETER(uint32, InPayloadDataFlags)
+		SHADER_PARAMETER(uint32, InSeed)
 		SHADER_PARAMETER_STRUCT_INCLUDE(FGPUSceneWriterParameters, GPUSceneWriterParameters)
 	END_SHADER_PARAMETER_STRUCT()
 
