[PCG] Do not add dataless pin input to the input collection

* Also make sure to not add dataless input while compiling, but this won't caught all cases, especially dynamic subgraph calls, hence why we also need it in the graph executor input gatherer.

#rnx

diff --git a/Engine/Plugins/PCG/Source/PCG/Private/Graph/PCGGraphCompiler.cpp b/Engine/Plugins/PCG/Source/PCG/Private/Graph/PCGGraphCompiler.cpp
--- a/Engine/Plugins/PCG/Source/PCG/Private/Graph/PCGGraphCompiler.cpp
+++ b/Engine/Plugins/PCG/Source/PCG/Private/Graph/PCGGraphCompiler.cpp
@@ -241,6 +241,7 @@ TArray<FPCGGraphTask> FPCGGraphCompiler::CompileGraph(UPCGGraph* InGraph, FPCGTa
 				check(InputPin);
 				// Params have to be funneled in the subgraph element (so we can pass down user parameters)
 				const bool bIsOverrideOrUserParamPin = InputPin->Properties.IsOverrideOrUserParamPin();
+				const bool bIsDatalessPin = InputPin->Properties.IsDatalessPin();
 
 				for (const UPCGEdge* InboundEdge : InputPin->Edges)
 				{
@@ -255,7 +256,7 @@ TArray<FPCGGraphTask> FPCGGraphCompiler::CompileGraph(UPCGGraph* InGraph, FPCGTa
 						else
 						{
 							PreTask.Inputs.Emplace(IdMapping[InboundEdge->InputPin->Node], InboundEdge->InputPin->Properties, InboundEdge->OutputPin->Properties, /*bInProvideData=*/false);
-							InputNodeTask->Inputs.Emplace(IdMapping[InboundEdge->InputPin->Node], InboundEdge->InputPin->Properties, InboundEdge->OutputPin->Properties, /*bInProvideData=*/true);
+							InputNodeTask->Inputs.Emplace(IdMapping[InboundEdge->InputPin->Node], InboundEdge->InputPin->Properties, InboundEdge->OutputPin->Properties, /*bInProvideData=*/!bIsDatalessPin);
 						}
 					}
 					else
@@ -330,7 +331,7 @@ TArray<FPCGGraphTask> FPCGGraphCompiler::CompileGraph(UPCGGraph* InGraph, FPCGTa
 			for (const UPCGPin* InputPin : Node->InputPins)
 			{
 				// Cull dependency only data if the pin is marked as such.
-				const bool bIsDependencyOnly = InputPin->Properties.Usage == EPCGPinUsage::DependencyOnly;
+				const bool bIsDatalessPin = InputPin->Properties.IsDatalessPin();
 				for (const UPCGEdge* InboundEdge : InputPin->Edges)
 				{
 					if (!InboundEdge->IsValid())
@@ -342,7 +343,7 @@ TArray<FPCGGraphTask> FPCGGraphCompiler::CompileGraph(UPCGGraph* InGraph, FPCGTa
 					check(InboundEdge->OutputPin == InputPin);
 					if (FPCGTaskId* InboundId = IdMapping.Find(InboundEdge->InputPin->Node))
 					{
-						Task.Inputs.Emplace(*InboundId, InboundEdge->InputPin->Properties, InputPin->Properties, !bIsDependencyOnly);
+						Task.Inputs.Emplace(*InboundId, InboundEdge->InputPin->Properties, InputPin->Properties, !bIsDatalessPin);
 					}
 					else
 					{
diff --git a/Engine/Plugins/PCG/Source/PCG/Private/Graph/PCGGraphExecutor.cpp b/Engine/Plugins/PCG/Source/PCG/Private/Graph/PCGGraphExecutor.cpp
--- a/Engine/Plugins/PCG/Source/PCG/Private/Graph/PCGGraphExecutor.cpp
+++ b/Engine/Plugins/PCG/Source/PCG/Private/Graph/PCGGraphExecutor.cpp
@@ -2233,7 +2233,7 @@ void FPCGGraphExecutor::BuildTaskInput(FPCGGraphTask& Task, TArray<FPCGTaskId>&
 			}
 
 			// If the input does not provide any data, don't add it to the task input.
-			if (!Input.bProvideData)
+			if (!Input.bProvideData || (Input.DownstreamPin.IsSet() && Input.DownstreamPin.GetValue().IsDatalessPin()))
 			{
 				continue;
 			}
