[PCG] Promote OverriddenSettingsPropertyBag to a UProperty for GC visibility

Fixes crash in UPCGSettings::CreateEmptyPropertyBagInstance() when trying to use after free

#rnx

diff --git a/Engine/Plugins/PCG/Source/PCG/Public/PCGSettings.h b/Engine/Plugins/PCG/Source/PCG/Public/PCGSettings.h
--- a/Engine/Plugins/PCG/Source/PCG/Public/PCGSettings.h
+++ b/Engine/Plugins/PCG/Source/PCG/Public/PCGSettings.h
@@ -559,8 +559,9 @@ protected:
 	UPROPERTY(meta = (PCGNoHash))
 	TArray<FPCGSettingsOverridableParam> CachedOverridableParams;
 
-	/** Property bag used to create instances of overridden parameter structs for execution. Lifetime is managed by the reflection system. */
-	const UPropertyBag* OverriddenSettingsPropertyBag = nullptr;
+	/** Property bag used to create instances of overridden parameter structs for execution. */
+	UPROPERTY(Transient)
+	TObjectPtr<const UPropertyBag> OverriddenSettingsPropertyBag = nullptr;
 
 	// We need to make sure that if we have hard references that are overridable, and they are overriden by paths
 	// on objects that are not yet loaded, that we are loading it on the main thread.
