[PCG] Disabled temporal AA & general viewport setup when entering PCG toolmode.

#rnx

diff --git a/Engine/Plugins/PCG/Source/PCGEditor/Private/EditorMode/PCGEdMode.cpp b/Engine/Plugins/PCG/Source/PCGEditor/Private/EditorMode/PCGEdMode.cpp
--- a/Engine/Plugins/PCG/Source/PCGEditor/Private/EditorMode/PCGEdMode.cpp
+++ b/Engine/Plugins/PCG/Source/PCGEditor/Private/EditorMode/PCGEdMode.cpp
@@ -194,12 +194,14 @@ void UPCGEditorMode::Enter()
 	}
 
 	ToolManager->SelectActiveToolType(EToolSide::Left, "DrawSpline");
+
+	ConfigureRealTimeViewportsOverride(/*bEnable=*/true);
 }
 
 void UPCGEditorMode::Exit()
 {
 	// Clear realtime viewport override
-	ConfigureRealTimeViewportsOverride(false);
+	ConfigureRealTimeViewportsOverride(/*bEnable=*/false);
 
 	if (UInteractiveToolManager* ToolManager = GetToolManager())
 	{
@@ -285,6 +287,43 @@ void UPCGEditorMode::OnToolStarted(UInteractiveToolManager* Manager, UInteractiv
 	FSlateThrottleManager::Get().DisableThrottle(true);
 
 	bIsToolActive = true;
+	bEnteredToolSinceLastTick = true;
+}
+
+void UPCGEditorMode::ModeTick(float DeltaTime)
+{
+	Super::ModeTick(DeltaTime);
+
+	if (bEnteredToolSinceLastTick)
+	{
+		bEnteredToolSinceLastTick = false;
+
+		if (!GetDefault<UPCGEditorModeSettings>()->bDisableTemporalAntiAliasingWhenEnteringTool)
+		{
+			// Undo setting the TAA to false as it doesn't look good in the PCG case.
+			// This is an override on top of what's done in UEditorInteractiveToolsContext::SetEditorStateForTool()
+			FLevelEditorModule& LevelEditorModule = FModuleManager::GetModuleChecked<FLevelEditorModule>("LevelEditor");
+			TSharedPtr<ILevelEditor> LevelEditor = LevelEditorModule.GetFirstLevelEditor();
+			if (LevelEditor.IsValid())
+			{
+				TArray<TSharedPtr<SLevelViewport>> Viewports = LevelEditor->GetViewports();
+				for (const TSharedPtr<SLevelViewport>& ViewportWindow : Viewports)
+				{
+					if (ViewportWindow.IsValid())
+					{
+						FEditorViewportClient& Viewport = ViewportWindow->GetAssetViewportClient();
+						Viewport.EnableOverrideEngineShowFlags([](FEngineShowFlags& Flags)
+						{
+							//Flags.SetTemporalAA(false); // <- purposefully not disabled
+							Flags.SetMotionBlur(false);
+							// disable this as depending on fixed exposure settings the entire scene may turn black
+							//Flags.SetEyeAdaptation(false);
+						});
+					}
+				}
+			}
+		}
+	}
 }
 
 void UPCGEditorMode::OnToolEnded(UInteractiveToolManager* Manager, UInteractiveTool* Tool)
diff --git a/Engine/Plugins/PCG/Source/PCGEditor/Private/EditorMode/PCGEdMode.h b/Engine/Plugins/PCG/Source/PCGEditor/Private/EditorMode/PCGEdMode.h
--- a/Engine/Plugins/PCG/Source/PCGEditor/Private/EditorMode/PCGEdMode.h
+++ b/Engine/Plugins/PCG/Source/PCGEditor/Private/EditorMode/PCGEdMode.h
@@ -39,6 +39,7 @@ public:
 	virtual bool HasCustomViewportFocus() const override;
 	virtual FBox ComputeCustomViewportFocus() const override;
 	virtual bool IsSelectionAllowed(AActor* InActor, bool bInSelection) const override;
+	virtual void ModeTick(float DeltaTime) override;
 	//~End UEdMode interface
 
 	UPCGInteractiveToolSettings* GetCurrentToolSettings() const;
@@ -67,4 +68,5 @@ private:
 	// 	const TFunction<bool(UInteractiveToolManager*, EToolSide)>& IsActionChecked) const;
 
 	bool bIsToolActive = false;
+	bool bEnteredToolSinceLastTick = false;
 };
diff --git a/Engine/Plugins/PCG/Source/PCGEditor/Private/EditorMode/PCGEdModeSettings.h b/Engine/Plugins/PCG/Source/PCGEditor/Private/EditorMode/PCGEdModeSettings.h
--- a/Engine/Plugins/PCG/Source/PCGEditor/Private/EditorMode/PCGEdModeSettings.h
+++ b/Engine/Plugins/PCG/Source/PCGEditor/Private/EditorMode/PCGEdModeSettings.h
@@ -60,6 +60,9 @@ public:
 		ConsoleVariable = "pcg.editormode.ShowEditorToastOnToolErrors"))
 	bool bShowEditorToastOnToolErrors = true;
 
+	UPROPERTY(EditAnywhere, Config, Category = "Tools", meta = (EditCondition = "bEnableEditorMode", EditConditionHides))
+	bool bDisableTemporalAntiAliasingWhenEnteringTool = false;
+
 	UPROPERTY(EditAnywhere, Category = "Tools|Output", meta = (EditCondition = "bEnableEditorMode", EditConditionHides))
 	FName DefaultNewActorName = "PCG Tool Actor";
 
