[PCG] Fixed array getting out of scope and corrupting vertex mapping data, leading to bad metadata in offset polygon node.

diff --git a/Engine/Plugins/PCG/Source/PCG/Private/Elements/Polygon/PCGOffsetPolygon2D.cpp b/Engine/Plugins/PCG/Source/PCG/Private/Elements/Polygon/PCGOffsetPolygon2D.cpp
--- a/Engine/Plugins/PCG/Source/PCG/Private/Elements/Polygon/PCGOffsetPolygon2D.cpp
+++ b/Engine/Plugins/PCG/Source/PCG/Private/Elements/Polygon/PCGOffsetPolygon2D.cpp
@@ -128,11 +128,12 @@ bool FPCGOffsetPolygonElement::ExecuteInternal(FPCGContext* Context) const
 		if (bOperationSuccess)
 		{
 			TArray<PCGPolygon2DProcessingHelpers::PolygonVertexMapping> VertexMapping;
+			// Since the PolyToMetadataInfoMap refers to the polygon entry keys, this array must exist in the same (or larger) scope than the info map object.
+			TArray<PCGMetadataEntryKey> PolygonEntryKeys;
 			PCGPolygon2DProcessingHelpers::PolyToMetadataInfoMap PolyToMetadataInfoMap;
 
 			if (Settings->bInheritMetadata)
 			{
-				TArray<PCGMetadataEntryKey> PolygonEntryKeys;
 				PolygonEntryKeys = PolygonData->GetConstVerticesEntryKeys();
 				PCGPolygon2DProcessingHelpers::AddToPolyMetadataInfoMap(PolyToMetadataInfoMap, Polygon, PolygonData->ConstMetadata(), PolygonEntryKeys);
 				VertexMapping = PCGPolygon2DProcessingHelpers::BuildVertexMapping(OutPolys, PolyView, {});
