[PCG] Fix MakeUniqueObjectName using wrong Outer causing object reuse outside of GameThread (name wasn't unique to the actual outer)

#rnx

diff --git a/Engine/Plugins/PCG/Source/PCG/Private/Graph/PCGGraphCompilerGPU.cpp b/Engine/Plugins/PCG/Source/PCG/Private/Graph/PCGGraphCompilerGPU.cpp
--- a/Engine/Plugins/PCG/Source/PCG/Private/Graph/PCGGraphCompilerGPU.cpp
+++ b/Engine/Plugins/PCG/Source/PCG/Private/Graph/PCGGraphCompilerGPU.cpp
@@ -836,7 +836,7 @@ void FPCGGraphCompilerGPU::BuildComputeGraphTask(
 		
 		// Create a new compute graph.
 		UObject* ComputeGraphOuter = InOutContext.GetGraphCompiler().IsCooking() ? InGraph : Cast<UObject>(GetTransientPackage());
-		const FName GraphName = MakeUniqueObjectName(InGraph, UPCGComputeGraph::StaticClass(), InGraph->GetFName());
+		const FName GraphName = MakeUniqueObjectName(ComputeGraphOuter, UPCGComputeGraph::StaticClass(), InGraph->GetFName());
 		ComputeGraph = InOutContext.NewObject_AnyThread<UPCGComputeGraph>(ComputeGraphOuter, GraphName);
 		InOutContext.SetStaticAttributeTable(&ComputeGraph->GetStaticAttributeTable());
 
