[PCG] Temporary fix for thread-safety of Polygon 2D Accessors

* Set in Polygon accessor was not thread safe as it was extracting the vertices to work on, modify the ones we want, and create a new polygon out of it.
* Temporary fix in cpp (since we can't touch headers), to keep a copy of all the touched vertices and apply it on destructor.

#rnx

diff --git a/Engine/Plugins/PCG/Source/PCG/Private/Data/PCGPolygon2DData.cpp b/Engine/Plugins/PCG/Source/PCG/Private/Data/PCGPolygon2DData.cpp
--- a/Engine/Plugins/PCG/Source/PCG/Private/Data/PCGPolygon2DData.cpp
+++ b/Engine/Plugins/PCG/Source/PCG/Private/Data/PCGPolygon2DData.cpp
@@ -16,6 +16,8 @@
 #include "Serialization/ArchiveCrc32.h"
 #include "UObject/Package.h"
 
+#include "Data/PCGPolygon2DData.inl"
+
 #include UE_INLINE_GENERATED_CPP_BY_NAME(PCGPolygon2DData)
 
 #define LOCTEXT_NAMESPACE "PCGPolygon2DData"
@@ -163,9 +165,9 @@ namespace PCGPolygon2DData
 			switch (Property)
 			{
 			case EPCGPolygon2DProperties::Position:
-				return MakeUnique<FPCGPolygon2DVerticesAccessor<FVector, EPCGPolygon2DProperties::Position>>(bConst);
+				return MakeUnique<PCGTemporary::FPCGPolygon2DVerticesAccessor<FVector, EPCGPolygon2DProperties::Position>>(bConst);
 			case EPCGPolygon2DProperties::LocalPosition:
-				return MakeUnique<FPCGPolygon2DVerticesAccessor<FVector2d, EPCGPolygon2DProperties::LocalPosition>>(bConst);
+				return MakeUnique<PCGTemporary::FPCGPolygon2DVerticesAccessor<FVector2d, EPCGPolygon2DProperties::LocalPosition>>(bConst);
 			case EPCGPolygon2DProperties::Rotation: //fall-through
 			case EPCGPolygon2DProperties::SegmentIndex: //fall-through
 			case EPCGPolygon2DProperties::HoleIndex: //fall-through
diff --git a/Engine/Plugins/PCG/Source/PCG/Private/Data/PCGPolygon2DData.inl b/Engine/Plugins/PCG/Source/PCG/Private/Data/PCGPolygon2DData.inl
new file mode 100644
--- /dev/null
+++ b/Engine/Plugins/PCG/Source/PCG/Private/Data/PCGPolygon2DData.inl
@@ -0,0 +1,230 @@
+// Copyright Epic Games, Inc. All Rights Reserved.
+
+#include "Data/PCGPolygon2DData.h"
+#include "Metadata/Accessors/IPCGAttributeAccessorTpl.h"
+#include "UObject/ObjectKey.h"
+
+namespace PCGTemporary
+{
+	/**
+	* Templated accessor for polygon vertices accessor.
+	* NOTE FOR 5.7.1: The original accessor is NOT thread-safe on Set, and is replaced by this one since we can't modify headers in minor versions.
+	* As a temporary fix, we'll make a thread-safe copy of all polygons we try to set on "Set" (which can be very inefficient for enormous polygons)
+	* and write the data on accessor destruction. This is a suboptimal fix, that should cover most cases of how accessors can be used.
+	* Note also that it means the data is not applied until the accessor is destroyed. Meaning that reading polygon data will yield the old values.
+	* And also, writing is thread-safe only if multiple threads are never writing to the same range of indices.
+	*/
+	template<typename T, EPCGPolygon2DProperties Target>
+	class FPCGPolygon2DVerticesAccessor : public IPCGAttributeAccessorT<FPCGPolygon2DVerticesAccessor<T, Target>>
+	{
+	public:
+		using Type = T;
+		using Super = IPCGAttributeAccessorT<FPCGPolygon2DVerticesAccessor<T, Target>>;
+
+		FPCGPolygon2DVerticesAccessor(bool bIsReadOnly)
+			: Super(bIsReadOnly)
+		{
+		}
+		
+		~FPCGPolygon2DVerticesAccessor()
+		{
+			for (auto& [DataKey, Vertices] : ModifiedVertices)
+			{
+				if (UPCGPolygon2DData* Data = DataKey.ResolveObjectPtr())
+				{
+					UE::Geometry::FGeneralPolygon2d NewPolygon;
+					
+					NewPolygon.SetOuter(UE::Geometry::FPolygon2d(MoveTemp(Vertices[0])));
+					for (int i = 1; i < Vertices.Num(); ++i)
+					{
+						NewPolygon.AddHole(UE::Geometry::FPolygon2d(MoveTemp(Vertices[i])));
+					}
+
+					Data->SetPolygon(NewPolygon);
+				}
+			}
+		}
+
+		bool GetRangeImpl(TArrayView<Type> OutValues, int32 Index, const IPCGAttributeAccessorKeys& Keys) const
+		{
+			const void* ContainerKeys = nullptr;
+			TArrayView<const void*> ContainerKeysView(&ContainerKeys, 1);
+			if (!Keys.GetKeys(Index, ContainerKeysView))
+			{
+				return false;
+			}
+
+			// Validation to not access keys that are not the expected type. Done after the GetKeys, as we also want to discard other type of incompatible
+			// keys (like a Default Metadata entry key)
+			if (!ensure(Keys.IsClassSupported(UPCGPolygon2DData::StaticClass())))
+			{
+				return false;
+			}
+
+			const UPCGPolygon2DData* Data = static_cast<const UPCGPolygon2DData*>(ContainerKeys);
+			const FTransform& Transform = Data->GetTransform();
+			const UE::Geometry::FGeneralPolygon2d& Polygon = Data->GetPolygon();
+			const TMap<int, TPair<int, int>>& SegmentIndexToSegmentAndHoleIndices = Data->GetSegmentIndexToSegmentAndHoleIndices();
+
+			const int32 NumPoints = SegmentIndexToSegmentAndHoleIndices.Num();
+			if (NumPoints == 0)
+			{
+				return false;
+			}
+
+			for (int32 i = 0; i < OutValues.Num(); ++i)
+			{
+				const int32 CurrIndex = (Index + i) % NumPoints;
+				auto [SegmentIndex, HoleIndex] = SegmentIndexToSegmentAndHoleIndices[CurrIndex];
+
+				if constexpr (Target == EPCGPolygon2DProperties::Position)
+				{
+					static_assert(std::is_same_v<Type, FVector>);
+					const UE::Geometry::FSegment2d Segment = Polygon.Segment(SegmentIndex, HoleIndex);
+					OutValues[i] = Transform.TransformPosition(FVector(Segment.StartPoint(), 0.0));
+				}
+				else if constexpr (Target == EPCGPolygon2DProperties::Rotation)
+				{
+					static_assert(std::is_same_v<Type, FQuat>);
+					//@todo_pcg: there is a Polygon.GetNormal that we could use, but it blends normals across vertices.
+					const UE::Geometry::FSegment2d Segment = Polygon.Segment(SegmentIndex, HoleIndex);
+					OutValues[i] = Transform.TransformRotation(FRotationMatrix::MakeFromXZ(FVector(Segment.Direction, 0.0), FVector::UpVector).ToQuat());
+				}
+				else if constexpr (Target == EPCGPolygon2DProperties::SegmentIndex)
+				{
+					static_assert(std::is_same_v<Type, int32>);
+					OutValues[i] = SegmentIndex;
+				}
+				else if constexpr (Target == EPCGPolygon2DProperties::HoleIndex)
+				{
+					static_assert(std::is_same_v<Type, int32>);
+					OutValues[i] = HoleIndex;
+				}
+				else if constexpr (Target == EPCGPolygon2DProperties::SegmentLength)
+				{
+					static_assert(std::is_same_v<Type, double>);
+					const UE::Geometry::FSegment2d Segment = Polygon.Segment(SegmentIndex, HoleIndex);
+					OutValues[i] = Segment.Length();
+				}
+				else if constexpr (Target == EPCGPolygon2DProperties::LocalPosition)
+				{
+					static_assert(std::is_same_v<Type, FVector2d>);
+					const UE::Geometry::FSegment2d Segment = Polygon.Segment(SegmentIndex, HoleIndex);
+					OutValues[i] = Segment.StartPoint();
+				}
+				else if constexpr (Target == EPCGPolygon2DProperties::LocalRotation)
+				{
+					static_assert(std::is_same_v<Type, FQuat>);
+					const UE::Geometry::FSegment2d Segment = Polygon.Segment(SegmentIndex, HoleIndex);
+					OutValues[i] = FRotationMatrix::MakeFromXZ(FVector(Segment.Direction, 0.0), FVector::UpVector).ToQuat();
+				}
+				else
+				{
+					// Pitfall static assert
+					static_assert(!std::is_same_v<Type, Type>);
+				}
+			}
+
+			return true;
+		}
+
+		bool SetRangeImpl(TArrayView<const Type> InValues, int32 Index, IPCGAttributeAccessorKeys& Keys, EPCGAttributeAccessorFlags)
+		{
+			void* ContainerKeys = nullptr;
+			TArrayView<void*> ContainerKeysView(&ContainerKeys, 1);
+			if (!Keys.GetKeys(Index, ContainerKeysView))
+			{
+				return false;
+			}
+
+			// Validation to not access keys that are not the expected type. Done after the GetKeys, as we also want to discard other type of incompatible
+			// keys (like a Default Metadata entry key)
+			if (!ensure(Keys.IsClassSupported(UPCGPolygon2DData::StaticClass())))
+			{
+				return false;
+			}
+
+			UPCGPolygon2DData* Data = static_cast<UPCGPolygon2DData*>(ContainerKeys);
+			const TMap<int, TPair<int, int>>& SegmentIndexToSegmentAndHoleIndices = Data->GetSegmentIndexToSegmentAndHoleIndices();
+
+			const int32 NumPoints = SegmentIndexToSegmentAndHoleIndices.Num();
+			if (NumPoints == 0)
+			{
+				return false;
+			}
+
+			if constexpr (Target != EPCGPolygon2DProperties::Position && Target != EPCGPolygon2DProperties::LocalPosition)
+			{
+				return false;
+			}
+			else
+			{
+				const UE::Geometry::FGeneralPolygon2d& Polygon = Data->GetPolygon();
+				TArray<TArray<FVector2d>>* VerticesPtr = nullptr;
+				
+				{
+					UE::TScopeLock LockScope(Lock);
+					if (TArray<TArray<FVector2d>>* It = ModifiedVertices.Find(Data))
+					{
+						VerticesPtr = It;
+					}
+					else
+					{
+						// Need to copy the vertices locally to re-apply them on exit.
+						VerticesPtr = &ModifiedVertices.Emplace(Data);
+						VerticesPtr->Reserve(1 + Polygon.GetHoles().Num());
+						VerticesPtr->Add(Polygon.GetOuter().GetVertices());
+						for (const UE::Geometry::FPolygon2d& Hole : Polygon.GetHoles())
+						{
+							VerticesPtr->Add(Hole.GetVertices());
+						}
+					}
+				}
+
+				if (!ensure(VerticesPtr))
+				{
+					return false;
+				}
+
+				TArray<TArray<FVector2d>>& Vertices = *VerticesPtr;
+				const FTransform& Transform = Data->GetTransform();
+				
+				FTransform InverseTransform;
+
+				if constexpr (Target == EPCGPolygon2DProperties::Position)
+				{
+					InverseTransform = Transform.Inverse();
+				}
+
+				for (int i = 0; i < InValues.Num(); ++i)
+				{
+					const int32 CurrIndex = (Index + i) % NumPoints;
+					auto [SegmentIndex, HoleIndex] = SegmentIndexToSegmentAndHoleIndices[CurrIndex];
+
+					if constexpr (Target == EPCGPolygon2DProperties::Position)
+					{
+						static_assert(std::is_same_v<Type, FVector3d>);
+						FVector LocalPosition = InverseTransform.TransformPosition(InValues[i]);
+						Vertices[HoleIndex + 1][SegmentIndex] = FVector2d(LocalPosition.X, LocalPosition.Y);
+					}
+					else if constexpr (Target == EPCGPolygon2DProperties::LocalPosition)
+					{
+						static_assert(std::is_same_v<Type, FVector2d>);
+						Vertices[HoleIndex + 1][SegmentIndex] = InValues[i];
+					}
+					else
+					{
+						// Pitfall static assert
+						static_assert(!std::is_same_v<Type, Type>);
+					}
+				}
+				
+				return true;
+			}
+		}
+
+	private:
+		mutable TMap<TObjectKey<UPCGPolygon2DData>, TArray<TArray<FVector2d>>> ModifiedVertices;
+		mutable UE::FMutex Lock;
+	};
+}
diff --git a/Engine/Plugins/PCG/Source/PCG/Public/Metadata/Accessors/PCGPolygon2DAccessor.h b/Engine/Plugins/PCG/Source/PCG/Public/Metadata/Accessors/PCGPolygon2DAccessor.h
--- a/Engine/Plugins/PCG/Source/PCG/Public/Metadata/Accessors/PCGPolygon2DAccessor.h
+++ b/Engine/Plugins/PCG/Source/PCG/Public/Metadata/Accessors/PCGPolygon2DAccessor.h
@@ -7,6 +7,9 @@
 
 /**
 * Templated accessor for polygon vertices accessor.
+* NOTE FOR 5.7.1: This accessor is NOT thread-safe on Set, and is replaced by a thread-safe copy in the PCGPolygon2DData.inl
+* since we can't modify this file in a minor version.
+* You must not use this accessor directly.
 */
 template<typename T, EPCGPolygon2DProperties Target>
 class FPCGPolygon2DVerticesAccessor : public IPCGAttributeAccessorT<FPCGPolygon2DVerticesAccessor<T, Target>>
