[PCG] Convert some lazy/paranoid ensures to proper error messages.

#rnx

diff --git a/Engine/Plugins/PCG/Source/PCG/Private/Compute/Elements/PCGComputeGraphElement.cpp b/Engine/Plugins/PCG/Source/PCG/Private/Compute/Elements/PCGComputeGraphElement.cpp
--- a/Engine/Plugins/PCG/Source/PCG/Private/Compute/Elements/PCGComputeGraphElement.cpp
+++ b/Engine/Plugins/PCG/Source/PCG/Private/Compute/Elements/PCGComputeGraphElement.cpp
@@ -159,9 +159,16 @@ bool FPCGComputeGraphElement::ExecuteInternal(FPCGContext* InContext) const
 
 		const IPCGGraphExecutionSource* Source = InContext->ExecutionSource.Get();
 		const UPCGGraph* TopGraph = Source ? Source->GetExecutionState().GetGraph() : nullptr;
+		if (!TopGraph)
+		{
+			UE_LOG(LogPCG, Error, TEXT("Graph was null, compute graph execution aborted."));
+			return true;
+		}
+
 		const FPCGStack* StackPtr = Context->GetStack();
-		if (!ensure(TopGraph) || !ensure(StackPtr))
+		if (!StackPtr)
 		{
+			UE_LOG(LogPCG, Error, TEXT("Stack information missing from context, compute graph execution aborted."));
 			return true;
 		}
 
@@ -176,10 +183,20 @@ bool FPCGComputeGraphElement::ExecuteInternal(FPCGContext* InContext) const
 		if (IPCGBaseSubsystem* Subsystem = Source->GetExecutionState().GetSubsystem())
 		{
 			Context->ComputeGraph = Subsystem->GetComputeGraph(DynamicSubgraph ? DynamicSubgraph : TopGraph, Context->GenerationGridSize, ComputeGraphIndex);
-		}
 
-		if (!ensure(Context->ComputeGraph))
+			if (!Context->ComputeGraph)
+			{
+				UE_LOG(LogPCG, Error, TEXT("Failed to obtain compute graph for %s '%s', grid size %u, compute graph index %d. Compute graph execution aborted."),
+					DynamicSubgraph ? TEXT("dynamic subgraph") : TEXT("top graph"),
+					DynamicSubgraph ? *DynamicSubgraph->GetName() : *TopGraph->GetName(),
+					Context->GenerationGridSize,
+					ComputeGraphIndex);
+				return true;
+			}
+		}
+		else
 		{
+			UE_LOG(LogPCG, Error, TEXT("Failed to obtain compute graph, no subsystem present. Compute graph execution aborted."));
 			return true;
 		}
 
