[PCG] Don't fire world generation process delegates in PIE.

This makes the scene render twice.

[RN] [PCG] Fix bug that forced editor world to render when PCG generates in PIE, producing heavy GPU work.

diff --git a/Engine/Plugins/PCG/Source/PCG/Private/Graph/PCGGraphExecutor.cpp b/Engine/Plugins/PCG/Source/PCG/Private/Graph/PCGGraphExecutor.cpp
--- a/Engine/Plugins/PCG/Source/PCG/Private/Graph/PCGGraphExecutor.cpp
+++ b/Engine/Plugins/PCG/Source/PCG/Private/Graph/PCGGraphExecutor.cpp
@@ -1898,7 +1898,7 @@ void FPCGGraphExecutor::Execute(double& InOutEndTime)
 	bool bFirstLoop = true;
 
 #if WITH_EDITOR
-	bool bOnWorldGenerationProcessBroadcasted = bContinueExecute && World;
+	const bool bOnWorldGenerationProcessBroadcasted = bContinueExecute && World && !World->IsGameWorld();
 	if (bOnWorldGenerationProcessBroadcasted)
 	{
 		FEditorDelegates::OnPreWorldGenerationProcess.Broadcast(World);
