[PCG] Fix crash on engine exit

* Mimic what's done in LevelInstanceEditor

#rnx

diff --git a/Engine/Plugins/PCG/Source/PCGEditor/Private/PCGEditorModule.cpp b/Engine/Plugins/PCG/Source/PCGEditor/Private/PCGEditorModule.cpp
--- a/Engine/Plugins/PCG/Source/PCGEditor/Private/PCGEditorModule.cpp
+++ b/Engine/Plugins/PCG/Source/PCGEditor/Private/PCGEditorModule.cpp
@@ -174,9 +174,13 @@ void FPCGEditorModule::ShutdownModule()
 	{
 		FLevelEditorModule& LevelEditorModule = FModuleManager::GetModuleChecked<FLevelEditorModule>("LevelEditor");
 		LevelEditorModule.OnLevelEditorCreated().RemoveAll(this);
-		if (TSharedPtr<ILevelEditor> FirstLevelEditor = LevelEditorModule.GetFirstLevelEditor())
+
+		if (!IsEngineExitRequested())
 		{
-			FirstLevelEditor->GetEditorModeManager().OnEditorModeIDChanged().RemoveAll(this);
+			if (TSharedPtr<ILevelEditor> FirstLevelEditor = LevelEditorModule.GetFirstLevelEditor())
+			{
+				FirstLevelEditor->GetEditorModeManager().OnEditorModeIDChanged().RemoveAll(this);
+			}
 		}
 	}
 
