[PCG] Cleaned up ensure case when declaration node isn't found.

#rnx

diff --git a/Engine/Plugins/PCG/Source/PCGEditor/Private/Nodes/PCGEditorGraphNodeBase.cpp b/Engine/Plugins/PCG/Source/PCGEditor/Private/Nodes/PCGEditorGraphNodeBase.cpp
--- a/Engine/Plugins/PCG/Source/PCGEditor/Private/Nodes/PCGEditorGraphNodeBase.cpp
+++ b/Engine/Plugins/PCG/Source/PCGEditor/Private/Nodes/PCGEditorGraphNodeBase.cpp
@@ -95,7 +95,7 @@ namespace PCGEditorGraphNodeBase
 				{
 					const UPCGNode* DeclarationPCGNode = RerouteUsageSettings->Declaration ? Cast<UPCGNode>(RerouteUsageSettings->Declaration->GetOuter()) : nullptr;
 					const UPCGEditorGraph* EditorGraph = Cast<UPCGEditorGraph>(InNode->GetOuter());
-					const UPCGEditorGraphNodeBase* DeclarationNode = EditorGraph ? EditorGraph->GetEditorNodeFromPCGNode(DeclarationPCGNode) : nullptr;
+					const UPCGEditorGraphNodeBase* DeclarationNode = EditorGraph && DeclarationPCGNode ? EditorGraph->GetEditorNodeFromPCGNode(DeclarationPCGNode) : nullptr;
 					return !DeclarationNode || ShouldDisplayAsActive(DeclarationNode, InSourceBeingDebugged, InStackBeingInspected);
 				}
 
