[PCG] Fixed crash in spline intersection with no metadata and intersection on preexisting points.

#rnx

diff --git a/Engine/Plugins/PCG/Source/PCG/Private/Elements/PCGSplineIntersection.cpp b/Engine/Plugins/PCG/Source/PCG/Private/Elements/PCGSplineIntersection.cpp
--- a/Engine/Plugins/PCG/Source/PCG/Private/Elements/PCGSplineIntersection.cpp
+++ b/Engine/Plugins/PCG/Source/PCG/Private/Elements/PCGSplineIntersection.cpp
@@ -506,14 +506,22 @@ bool FPCGSplineIntersectionElement::ExecuteInternal(FPCGContext* Context) const
 
 					if ((FMath::Abs(PreviousPoint.InputKey - Key) <= MinKeyDifferenceThreshold && (PreviousPoint.Position - Position).SquaredLength() <= SqrDistanceThreshold))
 					{
-						// Mark previous point as an intersection
-						OptionalWriteIntersectionMetadata(SplinePointsEntryKeys[InsertionIndex - 1], PointToAdd);
+						if (bNeedsMetadata)
+						{
+							// Mark previous point as an intersection
+							OptionalWriteIntersectionMetadata(SplinePointsEntryKeys[InsertionIndex - 1], PointToAdd);
+						}
+
 						continue;
 					}
 					else if ((FMath::Abs(NextPointVirtualKey - Key) <= MinKeyDifferenceThreshold && (NextPoint.Position - Position).SquaredLength() <= SqrDistanceThreshold))
 					{
-						// Mark next point as an intersection
-						OptionalWriteIntersectionMetadata(SplinePointsEntryKeys[bIsOnLastSegment ? 0 : InsertionIndex], PointToAdd);
+						if (bNeedsMetadata)
+						{
+							// Mark next point as an intersection
+							OptionalWriteIntersectionMetadata(SplinePointsEntryKeys[bIsOnLastSegment ? 0 : InsertionIndex], PointToAdd);
+						}
+
 						continue;
 					}
 
