[PCG] Fixed potential crash on cleanup after garbage collection

#rnx

diff --git a/Engine/Plugins/PCG/Source/PCG/Private/Data/PCGDifferenceData.cpp b/Engine/Plugins/PCG/Source/PCG/Private/Data/PCGDifferenceData.cpp
--- a/Engine/Plugins/PCG/Source/PCG/Private/Data/PCGDifferenceData.cpp
+++ b/Engine/Plugins/PCG/Source/PCG/Private/Data/PCGDifferenceData.cpp
@@ -122,8 +122,10 @@ void UPCGDifferenceData::PostLoad()
 
 void UPCGDifferenceData::VisitDataNetwork(TFunctionRef<void(const UPCGData*)> Action) const
 {
-	check(GetSource());
-	GetSource()->VisitDataNetwork(Action);
+	if (GetSource())
+	{
+		GetSource()->VisitDataNetwork(Action);
+	}
 
 	if (GetDifference())
 	{
