[PCG] Fixed assert on null graph related to editor opening - code wasn't needed.

#rnx

diff --git a/Engine/Plugins/PCG/Source/PCGEditor/Private/Widgets/SPCGEditorGraphDebugObjectTree.cpp b/Engine/Plugins/PCG/Source/PCGEditor/Private/Widgets/SPCGEditorGraphDebugObjectTree.cpp
--- a/Engine/Plugins/PCG/Source/PCGEditor/Private/Widgets/SPCGEditorGraphDebugObjectTree.cpp
+++ b/Engine/Plugins/PCG/Source/PCGEditor/Private/Widgets/SPCGEditorGraphDebugObjectTree.cpp
@@ -465,9 +465,6 @@ void SPCGEditorGraphDebugObjectTree::Construct(const FArguments& InArgs, TShared
 	check(InPCGEditor);
 	PCGEditor = InPCGEditor;
 
-	UPCGGraph* PCGGraph = GetPCGGraph();
-	check(PCGGraph);
-
 	USelection::SelectionChangedEvent.AddSP(this, &SPCGEditorGraphDebugObjectTree::OnEditorSelectionChanged);
 
 	const TSharedRef<SScrollBar> VerticalScrollBar = SNew(SScrollBar)
