[PCG] Update on some nodes that should forward their input tags.

#rnx
#Synced-CL 48550537

diff --git a/Engine/Plugins/PCG/Source/PCG/Private/Elements/PCGCollapsePoints.cpp b/Engine/Plugins/PCG/Source/PCG/Private/Elements/PCGCollapsePoints.cpp
--- a/Engine/Plugins/PCG/Source/PCG/Private/Elements/PCGCollapsePoints.cpp
+++ b/Engine/Plugins/PCG/Source/PCG/Private/Elements/PCGCollapsePoints.cpp
@@ -440,7 +440,7 @@ PRAGMA_ENABLE_DEPRECATION_WARNINGS
 		// Rebuild octree
 		PCGCollapsePoints::Algo::RebuildOctree(Context, ExecState, OutState);
 
-		Context->OutputData.TaggedData.Emplace_GetRef().Data = OutPointData;
+		Context->OutputData.TaggedData.Emplace_GetRef(Inputs[IterationIndex]).Data = OutPointData;
 		return EPCGTimeSliceInitResult::Success;
 	});
 
diff --git a/Engine/Plugins/PCG/Source/PCG/Private/Elements/PCGGenerateSeedElement.cpp b/Engine/Plugins/PCG/Source/PCG/Private/Elements/PCGGenerateSeedElement.cpp
--- a/Engine/Plugins/PCG/Source/PCG/Private/Elements/PCGGenerateSeedElement.cpp
+++ b/Engine/Plugins/PCG/Source/PCG/Private/Elements/PCGGenerateSeedElement.cpp
@@ -136,7 +136,7 @@ bool FPCGGenerateSeedElement::ExecuteInternal(FPCGContext* Context) const
 			continue;
 		}
 
-		FPCGTaggedData& Output = Context->OutputData.TaggedData.Emplace_GetRef();
+		FPCGTaggedData& Output = Context->OutputData.TaggedData.Emplace_GetRef(Input);
 		Output.Pin = PCGPinConstants::DefaultOutputLabel;
 		Output.Data = OutputData;
 	}
diff --git a/Engine/Plugins/PCG/Source/PCG/Private/Elements/PCGGetActorDataLayers.cpp b/Engine/Plugins/PCG/Source/PCG/Private/Elements/PCGGetActorDataLayers.cpp
--- a/Engine/Plugins/PCG/Source/PCG/Private/Elements/PCGGetActorDataLayers.cpp
+++ b/Engine/Plugins/PCG/Source/PCG/Private/Elements/PCGGetActorDataLayers.cpp
@@ -61,7 +61,7 @@ bool FPCGGetActorDataLayersElement::ExecuteInternal(FPCGContext* Context) const
 		}
 
 		// Add PCGParam Output with DataLayers
-		FPCGTaggedData& DataLayerSetOutput = Context->OutputData.TaggedData.Emplace_GetRef();
+		FPCGTaggedData& DataLayerSetOutput = Context->OutputData.TaggedData.Emplace_GetRef(Input);
 		UPCGParamData* OutputParamData = FPCGContext::NewObject_AnyThread<UPCGParamData>(Context);
 		DataLayerSetOutput.Pin = PCGPinConstants::DefaultOutputLabel;
 		DataLayerSetOutput.Data = OutputParamData;
