[PCG] Remove over-protective ensure in Get Bounds

* For example, empty point data has invalid bounds, as it has no information.

#rnx
#Synced-CL 48550537

diff --git a/Engine/Plugins/PCG/Source/PCG/Private/Elements/PCGGetBounds.cpp b/Engine/Plugins/PCG/Source/PCG/Private/Elements/PCGGetBounds.cpp
--- a/Engine/Plugins/PCG/Source/PCG/Private/Elements/PCGGetBounds.cpp
+++ b/Engine/Plugins/PCG/Source/PCG/Private/Elements/PCGGetBounds.cpp
@@ -53,8 +53,10 @@ bool FPCGGetBoundsElement::ExecuteInternal(FPCGContext* Context) const
 	const TArray<FPCGTaggedData>& Inputs = Context->InputData.TaggedData;
 	bool bEmittedUnboundedDataWarning = false;
 
-	for (const FPCGTaggedData& Input : Inputs)
+	for (int32 i = 0; i < Inputs.Num(); ++i)
 	{
+		const FPCGTaggedData& Input = Inputs[i];
+		
 		const UPCGSpatialData* InputSpatialData = Cast<UPCGSpatialData>(Input.Data);
 		if (!InputSpatialData)
 		{
@@ -73,9 +75,9 @@ bool FPCGGetBoundsElement::ExecuteInternal(FPCGContext* Context) const
 		}
 
 		const FBox InputBounds = InputSpatialData->GetBounds();
-		if (!ensure(InputBounds.IsValid))
+		if (!InputBounds.IsValid)
 		{
-			PCGE_LOG(Warning, GraphAndLog, LOCTEXT("InvalidBoundsOnInput", "Skipped spatial data that had invalid bounds."));
+			PCGE_LOG(Warning, GraphAndLog, FText::Format(LOCTEXT("InvalidBoundsOnInput", "Skipped data {0}, it has invalid bounds."), i));
 			continue;
 		}
 
