[PCG] Better safeguards to prevent unneeded data cleanup on shutdown

#rnx

diff --git a/Engine/Plugins/PCG/Source/PCGEditor/Private/EditorMode/Tools/PCGInteractiveToolSettings.cpp b/Engine/Plugins/PCG/Source/PCGEditor/Private/EditorMode/Tools/PCGInteractiveToolSettings.cpp
--- a/Engine/Plugins/PCG/Source/PCGEditor/Private/EditorMode/Tools/PCGInteractiveToolSettings.cpp
+++ b/Engine/Plugins/PCG/Source/PCGEditor/Private/EditorMode/Tools/PCGInteractiveToolSettings.cpp
@@ -453,25 +453,30 @@ void UPCGInteractiveToolSettings::DeleteGeneratedResources() const
 	
 	if(GeneratedResources.GeneratedPCGComponent.IsValid())
 	{
-		// We schedule a cleanup and deletion
-		FPCGTaskId CleanupTaskId = GeneratedResources.GeneratedPCGComponent->GetSubsystem()->ScheduleCleanup(GeneratedResources.GeneratedPCGComponent.Get(), true, {});
-		TArray<FPCGTaskId> DestroyTaskDependencies;
-		if (CleanupTaskId != InvalidPCGTaskId)
-		{
-			DestroyTaskDependencies.Add(CleanupTaskId);
-		}
+		UPCGSubsystem* Subsystem = GeneratedResources.GeneratedPCGComponent->GetSubsystem();
 
-		GeneratedResources.GeneratedPCGComponent->GetSubsystem()->ScheduleGeneric([PCGComponent = GeneratedResources.GeneratedPCGComponent]() -> bool
+		if (Subsystem)
 		{
-			if (PCGComponent.IsValid())
+			// We schedule a cleanup and deletion
+			FPCGTaskId CleanupTaskId = GeneratedResources.GeneratedPCGComponent->GetSubsystem()->ScheduleCleanup(GeneratedResources.GeneratedPCGComponent.Get(), true, {});
+			TArray<FPCGTaskId> DestroyTaskDependencies;
+			if (CleanupTaskId != InvalidPCGTaskId)
 			{
-				PCGComponent->DestroyComponent();
+				DestroyTaskDependencies.Add(CleanupTaskId);
 			}
+
+			GeneratedResources.GeneratedPCGComponent->GetSubsystem()->ScheduleGeneric([PCGComponent = GeneratedResources.GeneratedPCGComponent]() -> bool
+			{
+				if (PCGComponent.IsValid())
+				{
+					PCGComponent->DestroyComponent();
+				}
 			
-			return true;
-		},
-		GeneratedResources.GeneratedPCGComponent.Get(),
-		DestroyTaskDependencies);
+				return true;
+			},
+			GeneratedResources.GeneratedPCGComponent.Get(),
+			DestroyTaskDependencies);
+		}
 		
 		GeneratedResources.GeneratedPCGComponent.Reset();
 	}
