[PCG] Revert temporary fix for thread-safety with polygon and replace it with direct access to polygon data

* Initial proposition is fine but unfortunately the accessor can be destroyed way too late, since contextes are released only on the game thread.
* Since the assumption that outer and holes arrays should be fixed while writing with the accessor, we can go with a hacky "const_cast" as a temporary fix

#rnx

diff --git a/Engine/Plugins/PCG/Source/PCG/Private/Data/PCGPolygon2DData.inl b/Engine/Plugins/PCG/Source/PCG/Private/Data/PCGPolygon2DData.inl
--- a/Engine/Plugins/PCG/Source/PCG/Private/Data/PCGPolygon2DData.inl
+++ b/Engine/Plugins/PCG/Source/PCG/Private/Data/PCGPolygon2DData.inl
@@ -8,11 +8,9 @@ namespace PCGTemporary
 {
 	/**
 	* Templated accessor for polygon vertices accessor.
-	* NOTE FOR 5.7.1: The original accessor is NOT thread-safe on Set, and is replaced by this one since we can't modify headers in minor versions.
-	* As a temporary fix, we'll make a thread-safe copy of all polygons we try to set on "Set" (which can be very inefficient for enormous polygons)
-	* and write the data on accessor destruction. This is a suboptimal fix, that should cover most cases of how accessors can be used.
-	* Note also that it means the data is not applied until the accessor is destroyed. Meaning that reading polygon data will yield the old values.
-	* And also, writing is thread-safe only if multiple threads are never writing to the same range of indices.
+	* NOTE FOR 5.7.2: The original accessor is NOT thread-safe on Set, and is replaced by this one since we can't modify headers in minor versions.
+	* As a temporary fix and hack, we'll access directly the values in the polygon vertices. There is a huge assumption that vertices/holes arrays are not modified.
+	* Also, writing is thread-safe only if multiple threads are never writing to the same range of indices.
 	*/
 	template<typename T, EPCGPolygon2DProperties Target>
 	class FPCGPolygon2DVerticesAccessor : public IPCGAttributeAccessorT<FPCGPolygon2DVerticesAccessor<T, Target>>
@@ -25,25 +23,6 @@ namespace PCGTemporary
 			: Super(bIsReadOnly)
 		{
 		}
-		
-		~FPCGPolygon2DVerticesAccessor()
-		{
-			for (auto& [DataKey, Vertices] : ModifiedVertices)
-			{
-				if (UPCGPolygon2DData* Data = DataKey.ResolveObjectPtr())
-				{
-					UE::Geometry::FGeneralPolygon2d NewPolygon;
-					
-					NewPolygon.SetOuter(UE::Geometry::FPolygon2d(MoveTemp(Vertices[0])));
-					for (int i = 1; i < Vertices.Num(); ++i)
-					{
-						NewPolygon.AddHole(UE::Geometry::FPolygon2d(MoveTemp(Vertices[i])));
-					}
-
-					Data->SetPolygon(NewPolygon);
-				}
-			}
-		}
 
 		bool GetRangeImpl(TArrayView<Type> OutValues, int32 Index, const IPCGAttributeAccessorKeys& Keys) const
 		{
@@ -160,33 +139,6 @@ namespace PCGTemporary
 			else
 			{
 				const UE::Geometry::FGeneralPolygon2d& Polygon = Data->GetPolygon();
-				TArray<TArray<FVector2d>>* VerticesPtr = nullptr;
-				
-				{
-					UE::TScopeLock LockScope(Lock);
-					if (TArray<TArray<FVector2d>>* It = ModifiedVertices.Find(Data))
-					{
-						VerticesPtr = It;
-					}
-					else
-					{
-						// Need to copy the vertices locally to re-apply them on exit.
-						VerticesPtr = &ModifiedVertices.Emplace(Data);
-						VerticesPtr->Reserve(1 + Polygon.GetHoles().Num());
-						VerticesPtr->Add(Polygon.GetOuter().GetVertices());
-						for (const UE::Geometry::FPolygon2d& Hole : Polygon.GetHoles())
-						{
-							VerticesPtr->Add(Hole.GetVertices());
-						}
-					}
-				}
-
-				if (!ensure(VerticesPtr))
-				{
-					return false;
-				}
-
-				TArray<TArray<FVector2d>>& Vertices = *VerticesPtr;
 				const FTransform& Transform = Data->GetTransform();
 				
 				FTransform InverseTransform;
@@ -196,21 +148,27 @@ namespace PCGTemporary
 					InverseTransform = Transform.Inverse();
 				}
 
+				// 5.7.2 HACK: We const_cast to modify the values in place.
+				UE::Geometry::TPolygon2<double>& Outer = const_cast<UE::Geometry::TPolygon2<double>&>(Polygon.GetOuter());
+				TArray<UE::Geometry::TPolygon2<double>>& Holes = const_cast<TArray<UE::Geometry::TPolygon2<double>>&>(Polygon.GetHoles());
+
 				for (int i = 0; i < InValues.Num(); ++i)
 				{
 					const int32 CurrIndex = (Index + i) % NumPoints;
 					auto [SegmentIndex, HoleIndex] = SegmentIndexToSegmentAndHoleIndices[CurrIndex];
 
+					UE::Geometry::TPolygon2<double>& Vertices = HoleIndex == -1 ? Outer : Holes[HoleIndex];
+
 					if constexpr (Target == EPCGPolygon2DProperties::Position)
 					{
 						static_assert(std::is_same_v<Type, FVector3d>);
 						FVector LocalPosition = InverseTransform.TransformPosition(InValues[i]);
-						Vertices[HoleIndex + 1][SegmentIndex] = FVector2d(LocalPosition.X, LocalPosition.Y);
+						Vertices[SegmentIndex] = FVector2d(LocalPosition.X, LocalPosition.Y);
 					}
 					else if constexpr (Target == EPCGPolygon2DProperties::LocalPosition)
 					{
 						static_assert(std::is_same_v<Type, FVector2d>);
-						Vertices[HoleIndex + 1][SegmentIndex] = InValues[i];
+						Vertices[SegmentIndex] = InValues[i];
 					}
 					else
 					{
@@ -222,9 +180,5 @@ namespace PCGTemporary
 				return true;
 			}
 		}
-
-	private:
-		mutable TMap<TObjectKey<UPCGPolygon2DData>, TArray<TArray<FVector2d>>> ModifiedVertices;
-		mutable UE::FMutex Lock;
 	};
 }
