Removed ensures and renamed stroke command to reflect its use case in pcg editor mode

#Synced-CL 47950627

diff --git a/Engine/Plugins/PCG/Source/PCGEditor/Private/EditorMode/Tools/Line/PCGDrawSplineTool.cpp b/Engine/Plugins/PCG/Source/PCGEditor/Private/EditorMode/Tools/Line/PCGDrawSplineTool.cpp
--- a/Engine/Plugins/PCG/Source/PCGEditor/Private/EditorMode/Tools/Line/PCGDrawSplineTool.cpp
+++ b/Engine/Plugins/PCG/Source/PCGEditor/Private/EditorMode/Tools/Line/PCGDrawSplineTool.cpp
@@ -80,7 +80,7 @@ namespace PCGDrawSplineTool
 
 			virtual void Revert(USplineComponent& Spline) override
 			{
-				if (ensure(Spline.GetNumberOfSplinePoints() > 0))
+				if (Spline.GetNumberOfSplinePoints() > 0)
 				{
 					Spline.RemoveSplinePoint(Spline.GetNumberOfSplinePoints() - 1, true);
 				}
@@ -137,10 +137,10 @@ namespace PCGDrawSplineTool
 		};
 
 		// Undoes a free draw stroke (multiple points at once)
-		class FStrokeInsertionChange final : public FSplineCommandChange
+		class FFreeDrawStrokeChange final : public FSplineCommandChange
 		{
 		public:
-			FStrokeInsertionChange(USplineComponent* InWorkingSpline,
+			FFreeDrawStrokeChange(USplineComponent* InWorkingSpline,
 								   const TArray<FVector>& HitLocationsIn,
 								   const TArray<FVector>& UpVectorsIn)
 				: FSplineCommandChange(InWorkingSpline)
@@ -170,7 +170,7 @@ namespace PCGDrawSplineTool
 			{
 				for (int32 i = 0; i < HitLocations.Num(); ++i)
 				{
-					if (!ensure(Spline.GetNumberOfSplinePoints() > 0))
+					if (Spline.GetNumberOfSplinePoints() <= 0)
 					{
 						break;
 					}
@@ -181,7 +181,7 @@ namespace PCGDrawSplineTool
 
 			virtual FString ToString() const override
 			{
-				return TEXT("FStrokeInsertionChange");
+				return TEXT("FFreeDrawStrokeChange");
 			}
 
 		protected:
@@ -677,7 +677,7 @@ void UPCGDrawSplineToolBase::OnTerminateDragSequence()
 
 			GetToolManager()->EmitObjectChange(
 				this,
-				MakeUnique<CommandChange::FStrokeInsertionChange>(WorkingSpline, HitLocations, UpVectors),
+				MakeUnique<CommandChange::FFreeDrawStrokeChange>(WorkingSpline, HitLocations, UpVectors),
 				Constants::AddPointTransactionName);
 			break;
 		}
