[PCG] Fix outstanding issues with Blueprint node

* Regression from the data type refactor, there was no more safeguard for to invalid type on DynamicTypeOverride func
* Make sure the title function override does not get called on PostLoad

#rnx
#Synced-CL 48452988

diff --git a/Engine/Plugins/PCG/Source/PCG/Private/Elements/PCGExecuteBlueprint.cpp b/Engine/Plugins/PCG/Source/PCG/Private/Elements/PCGExecuteBlueprint.cpp
--- a/Engine/Plugins/PCG/Source/PCG/Private/Elements/PCGExecuteBlueprint.cpp
+++ b/Engine/Plugins/PCG/Source/PCG/Private/Elements/PCGExecuteBlueprint.cpp
@@ -257,7 +257,8 @@ FPCGDataTypeIdentifier UPCGBlueprintSettings::GetCurrentPinTypesID(const UPCGPin
 	// We can't call a BP function while we are PostLoading, in that case (or if we don't have an instance or it's not dynamic) just return the pin allowed types.
 	if (BlueprintElementInstance && BlueprintElementInstance->bHasDynamicPins && !FUObjectThreadContext::Get().IsRoutingPostLoad)
 	{
-		return FPCGDataTypeIdentifier{ static_cast<EPCGDataType>(BlueprintElementInstance->DynamicPinTypesOverride(this, InPin)) };
+		FPCGDataTypeIdentifier Result{ static_cast<EPCGDataType>(BlueprintElementInstance->DynamicPinTypesOverride(this, InPin)) };
+		return Result.IsValid() ? Result : InPin->Properties.AllowedTypes;
 	}
 	else
 	{
@@ -345,11 +346,16 @@ void UPCGBlueprintSettings::ApplyPreconfiguredSettings(const FPCGPreConfiguredSe
 
 FString UPCGBlueprintSettings::GetAdditionalTitleInformation() const
 {
-	if (BlueprintElementInstance && BlueprintElementInstance->NodeTitleOverride() != NAME_None)
+	// We can't call a BP function while we are PostLoading, in that case (or if we don't have an instance or it's not dynamic) just return the default title.
+	if (BlueprintElementInstance && !FUObjectThreadContext::Get().IsRoutingPostLoad)
 	{
-		return BlueprintElementInstance->NodeTitleOverride().ToString();
+		const FName NodeTitleOverride = BlueprintElementInstance->NodeTitleOverride();
+		if (NodeTitleOverride != NAME_None)
+		{
+			return BlueprintElementInstance->NodeTitleOverride().ToString();
+		}
 	}
-	else
+
 	{
 		FString ElementName;
 
