[PCG] Fix over-protective ensure that prevented to edit multiple values at the same time for Data Types.

* Now will support to set multiple values at the same time
* Will only support composition when all the identifiers that are modified are the same.

#rnx

diff --git a/Engine/Plugins/PCG/Source/PCGEditor/Private/Details/PCGDataTypeIdentifierDetails.cpp b/Engine/Plugins/PCG/Source/PCGEditor/Private/Details/PCGDataTypeIdentifierDetails.cpp
--- a/Engine/Plugins/PCG/Source/PCGEditor/Private/Details/PCGDataTypeIdentifierDetails.cpp
+++ b/Engine/Plugins/PCG/Source/PCGEditor/Private/Details/PCGDataTypeIdentifierDetails.cpp
@@ -42,6 +42,14 @@ namespace PCGDataTypeIdentifierDetails
 		return (Identifier & Other) == Other;
 	}
 
+	bool AreAllSameValues(TSharedPtr<IPropertyHandle> InPropertyHandle)
+	{
+		void* Data = nullptr;
+		FPropertyAccess::Result Result = InPropertyHandle ? InPropertyHandle->GetValueData(Data) : FPropertyAccess::Fail;
+
+		return Result == FPropertyAccess::Success;
+	}
+
 	/**
 	 * Look for a function on the outer object of the property or a flag in the property metadata to check if the widget
 	 * should support composition (have its bitflag-like behavior enabled) or not. 
@@ -54,23 +62,37 @@ namespace PCGDataTypeIdentifierDetails
 			TArray<UObject*> OutObjects;
 			InPropertyHandle->GetOuterObjects(OutObjects);
 
-			if (OutObjects.IsEmpty() || !OutObjects[0])
+			if (OutObjects.IsEmpty())
 			{
 				return false;
 			}
 
-			UObject* Object = OutObjects[0];
+			DECLARE_DELEGATE_RetVal(bool, FSupportComposition);
+
 			// Strip `()` that could be at the end of the function name
 			const FString FunctionName = SupportCompositionFunctionName.Replace(TEXT("()"), TEXT(""));
-			
-			if (!Object->GetClass() || !Object->GetClass()->FindFunctionByName(*FunctionName))
+
+			for (UObject* OuterObject : OutObjects)
 			{
-				return false;
+				if (!OuterObject)
+				{
+					return false;
+				}
+
+				if (!OuterObject->GetClass() || !OuterObject->GetClass()->FindFunctionByName(*FunctionName))
+				{
+					return false;
+				}
+				
+				auto Func = FSupportComposition::CreateUFunction(OuterObject, *FunctionName);
+				const bool bResult = Func.IsBound() && Func.Execute();
+				if (!bResult)
+				{
+					return false;
+				}
 			}
 
-			DECLARE_DELEGATE_RetVal(bool, FSupportComposition);
-			auto Func = FSupportComposition::CreateUFunction(Object, *FunctionName);
-			return Func.IsBound() ? Func.Execute() : false;
+			return true;
 		}
 		else if (InPropertyHandle->HasMetaData(PCGObjectMetadata::DataTypeIdentifierSupportsComposition))
 		{
@@ -101,27 +123,57 @@ namespace PCGDataTypeIdentifierDetails
 			TArray<UObject*> OutObjects;
 			PropertyHandle->GetOuterObjects(OutObjects);
 
-			if (OutObjects.IsEmpty() || !OutObjects[0])
-			{
-				return {};
-			}
+			TArray<FPCGDataTypeIdentifier> Result;
+			bool bHasFilters = false;
+
+			DECLARE_DELEGATE_RetVal(TArray<FPCGDataTypeIdentifier>, FGetFilters);
 
-			UObject* Object = OutObjects[0];
 			// Strip `()` that could be at the end of the function name
 			const FString FunctionName = GetFilterFunctionName.Replace(TEXT("()"), TEXT(""));
 			
-			if (!Object->GetClass() || !Object->GetClass()->FindFunctionByName(*FunctionName))
+			for (UObject* OuterObject : OutObjects)
 			{
-				return {};
-			}
+				if (!OuterObject->GetClass() || !OuterObject->GetClass()->FindFunctionByName(*FunctionName))
+				{
+					continue;
+				}
 
-			DECLARE_DELEGATE_RetVal(TArray<FPCGDataTypeIdentifier>, FSupportComposition);
-			auto Func = FSupportComposition::CreateUFunction(Object, *FunctionName);
+				auto Func = FGetFilters::CreateUFunction(OuterObject, *FunctionName);
 			
-			if (Func.IsBound())
-			{
-				return Func.Execute();
+				if (Func.IsBound())
+				{
+					TArray<FPCGDataTypeIdentifier> FuncResult = Func.Execute();
+					
+					if (!bHasFilters)
+					{
+						bHasFilters = true;
+						Result = MoveTemp(FuncResult);
+					}
+					else
+					{
+						// We can only keep the filters that are in common.
+						TBitArray<> Mask;
+						Mask.Init(false, Result.Num());
+						for (const FPCGDataTypeIdentifier& ID : FuncResult)
+						{
+							if (int32 Index = Result.IndexOfByKey(ID); Index != INDEX_NONE)
+							{
+								Mask[Index] = true;
+							}
+						}
+
+						for (int32 i = Result.Num() - 1; i >= 0; --i)
+						{
+							if (!Mask[i])
+							{
+								Result.RemoveAtSwap(i);
+							}
+						}
+					}
+				}
 			}
+
+			return bHasFilters ? TOptional<TArray<FPCGDataTypeIdentifier>>{MoveTemp(Result)} : TOptional<TArray<FPCGDataTypeIdentifier>>{};
 		}
 
 		return {};
@@ -204,32 +256,31 @@ TSharedRef<SWidget> FPCGDataTypeIdentifierDetails::OnGetMenuContent()
 			(
 				FExecuteAction::CreateLambda([this, ThisDataType = DataType]()
 				{
-					TArray<UObject*> OuterObjects;
-					PropertyHandle->GetOuterObjects(OuterObjects);
-					if (OuterObjects.IsEmpty())
+					TArray<void*> RawData;
+					PropertyHandle->AccessRawData(RawData);
+					if (RawData.IsEmpty())
 					{
 						return;
 					}
 
-					FPCGDataTypeIdentifier* Identifier = GetStruct();
-					if (!ensure(Identifier))
-					{
-						return;
-					}
+					FPCGDataTypeIdentifier* FirstIdentifier = static_cast<FPCGDataTypeIdentifier*>(RawData[0]);
+
+					// We can only support the composition if all the identifiers are the same.
+					bool bCanSupportComposition = bSupportComposition && PCGDataTypeIdentifierDetails::AreAllSameValues(PropertyHandle);
 
 					FPCGDataTypeIdentifier TempIdentifier{};
 
-					if (bSupportComposition && FSlateApplication::Get().GetModifierKeys().IsControlDown())
+					if (bCanSupportComposition && FSlateApplication::Get().GetModifierKeys().IsControlDown())
 					{
-						if (!PCGDataTypeIdentifierDetails::IsSuperSet(*Identifier, ThisDataType))
+						if (!PCGDataTypeIdentifierDetails::IsSuperSet(*FirstIdentifier, ThisDataType))
 						{
-							TempIdentifier = Identifier->Compose(ThisDataType);
+							TempIdentifier = FirstIdentifier->Compose(ThisDataType);
 						}
 						else
 						{
 							const FPCGDataTypeIdentifier DataTypeIdentifier{ ThisDataType };
 							FPCGDataTypeRegistry::FGetIdentifiersDifferenceParams Params = {
-								.SourceIdentifier = Identifier,
+								.SourceIdentifier = FirstIdentifier,
 								.DifferenceIdentifiers = MakeConstArrayView<FPCGDataTypeIdentifier>(&DataTypeIdentifier, 1),
 								.Filter = FPCGDataTypeRegistry::FGetIdentifiersDifferenceParams::ExcludeFilteredTypes,
 								.FilteredTypes = HiddenTypes
@@ -246,21 +297,23 @@ TSharedRef<SWidget> FPCGDataTypeIdentifierDetails::OnGetMenuContent()
 					if (TempIdentifier.IsValid())
 					{
 						// Need to export text and import text to go through the normal flow (like update Blueprint instances)
-						UObject* Outer = OuterObjects[0];
 						FString IdentifierAsText{};
-						FPCGDataTypeIdentifier::StaticStruct()->ExportText(IdentifierAsText, &TempIdentifier, /*Defaults=*/nullptr, Outer, PPF_None, /*ExportRootScope=*/nullptr);
+						FPCGDataTypeIdentifier::StaticStruct()->ExportText(IdentifierAsText, &TempIdentifier, /*Defaults=*/nullptr, /*OwnerObject*/nullptr, PPF_None, /*ExportRootScope=*/nullptr);
 						if (!IdentifierAsText.IsEmpty())
 						{
 							FScopedTransaction Transaction(NSLOCTEXT("PCGDataTypeIdentifierDetails", "ModifyIdentifier", "Modify PCG Data Type Identifier."));
-							PropertyHandle->SetPerObjectValues({MoveTemp(IdentifierAsText)});
+							for (int32 i = 0; i < RawData.Num(); ++i)
+							{
+								PropertyHandle->SetPerObjectValue(i, IdentifierAsText);
+							}
 						}
 					}
 				}),
 				FCanExecuteAction(),
 				FIsActionChecked::CreateLambda([this, ThisDataType = DataType]() -> bool
 				{
-					FPCGDataTypeIdentifier* Identifier = GetStruct();
-					return ensure(Identifier) ? PCGDataTypeIdentifierDetails::IsSuperSet(*Identifier, ThisDataType) : false;
+					const FPCGDataTypeIdentifier* Identifier = GetStruct();
+					return Identifier ? PCGDataTypeIdentifierDetails::IsSuperSet(*Identifier, ThisDataType) : false;
 				})
 			),
 			NAME_None,
@@ -274,8 +327,19 @@ FText FPCGDataTypeIdentifierDetails::GetText() const
 {
 	static const FPCGDataTypeIdentifier DefaultInvalidDataType{};
 
-	const FPCGDataTypeIdentifier* DataType = GetStruct();
-	return DataType ? DataType->ToDisplayText() : DefaultInvalidDataType.ToDisplayText();
+	void* Data = nullptr;
+	FPropertyAccess::Result Result = PropertyHandle->GetValueData(Data);
+
+	switch (Result)
+	{
+	case FPropertyAccess::MultipleValues:
+		return NSLOCTEXT("PCGDataTypeIdentifierDetails", "MultipleValues", "Multiple Values");
+	case FPropertyAccess::Success:
+		return static_cast<FPCGDataTypeIdentifier*>(Data)->ToDisplayText();
+	case FPropertyAccess::Fail: // fall-through
+	default:
+		return DefaultInvalidDataType.ToDisplayText();
+	}
 }
 
 FText FPCGDataTypeIdentifierDetails::GetTooltip() const
