[PCG] Fix ensure when writing large number of instances due to exceeding compute shader group count limits by using wrapped group count.

[RN] [PCG] Fixed ensure and instances dropped when spawning large number of instances via the GPU Static Mesh Spawner.

diff --git a/Engine/Plugins/PCG/Shaders/Private/PCGSceneWriter.usf b/Engine/Plugins/PCG/Shaders/Private/PCGSceneWriter.usf
--- a/Engine/Plugins/PCG/Shaders/Private/PCGSceneWriter.usf
+++ b/Engine/Plugins/PCG/Shaders/Private/PCGSceneWriter.usf
@@ -3,6 +3,7 @@
 // Scene writer - used to inject final instance data into the GPU Scene. Dispatched in the scene update timeline.
 
 #include "/Engine/Private/Common.ush"
+#include "/Engine/Private/ComputeShaderUtils.ush"
 #include "/Engine/Private/GPUScene/GPUSceneWriter.ush"
 #include "/Engine/Public/Platform.ush"
 
@@ -40,9 +41,10 @@ void CullInstance(uint InstanceId)
 }
 
 [numthreads(NUM_THREADS_PER_GROUP, 1, 1)]
-void SceneWriter_CS(const uint ThreadId : SV_DispatchThreadID)
+void SceneWriter_CS(const uint3 GroupId : SV_GroupID, const uint GroupIndex : SV_GroupIndex)
 {
-	const uint LocalInstanceId = ThreadId;
+	const uint LocalInstanceId = GetUnWrappedDispatchThreadId(GroupId, GroupIndex, NUM_THREADS_PER_GROUP);
+
 	if (LocalInstanceId >= InNumInstancesAllocatedInGPUScene)
 	{
 		return;
diff --git a/Engine/Plugins/PCG/Source/PCG/Private/Compute/DataInterfaces/PCGInstanceDataInterface.cpp b/Engine/Plugins/PCG/Source/PCG/Private/Compute/DataInterfaces/PCGInstanceDataInterface.cpp
--- a/Engine/Plugins/PCG/Source/PCG/Private/Compute/DataInterfaces/PCGInstanceDataInterface.cpp
+++ b/Engine/Plugins/PCG/Source/PCG/Private/Compute/DataInterfaces/PCGInstanceDataInterface.cpp
@@ -458,14 +458,11 @@ void FPCGInstanceDataProviderProxy::PostSubmit(FComputeContext& Context) const
 					Parameters->InNumInstancesAllocatedInGPUScene);
 #endif
 
-				// We're using a custom compute shader here rather than using a compute graph kernel, because ultimately the scene update will incorporate the
-				// new instance data. In the future we will not write directly to the scene here.
 				TShaderMapRef<FPCGSceneWriterCS> Shader(GetGlobalShaderMap(GMaxRHIFeatureLevel));
-				int GroupCount = FMath::DivideAndRoundUp<int>(PrimitiveNumInstances, FPCGSceneWriterCS::NumThreadsPerGroup);
-				ensureAlways(GroupCount <= GRHIMaxDispatchThreadGroupsPerDimension.X);
-				GroupCount = FMath::Min(GRHIMaxDispatchThreadGroupsPerDimension.X, GroupCount);
 
-				FComputeShaderUtils::AddPass(InGraphBuilder, RDG_EVENT_NAME("PCGWriteInstanceData"), ERDGPassFlags::Compute | ERDGPassFlags::NeverCull, Shader, Parameters, FIntVector(GroupCount, 1, 1));
+				const FIntVector WrappedGroupCount = FComputeShaderUtils::GetGroupCountWrapped(FMath::DivideAndRoundUp<int>(PrimitiveNumInstances, FPCGSceneWriterCS::NumThreadsPerGroup));
+
+				FComputeShaderUtils::AddPass(InGraphBuilder, RDG_EVENT_NAME("PCGWriteInstanceData"), ERDGPassFlags::Compute | ERDGPassFlags::NeverCull, Shader, Parameters, WrappedGroupCount);
 			}));
 
 			CumulativeInstanceCount += PrimitiveNumInstances;
