[PCG] Soften ensure in PCGDataCollectionPacking::ComputePackedSizeBytes()
#rnx

diff --git a/Engine/Plugins/PCG/Source/PCG/Private/Compute/Packing/PCGDataCollectionPacking.cpp b/Engine/Plugins/PCG/Source/PCG/Private/Compute/Packing/PCGDataCollectionPacking.cpp
--- a/Engine/Plugins/PCG/Source/PCG/Private/Compute/Packing/PCGDataCollectionPacking.cpp
+++ b/Engine/Plugins/PCG/Source/PCG/Private/Compute/Packing/PCGDataCollectionPacking.cpp
@@ -495,7 +495,7 @@ namespace PCGDataCollectionPackingHelpers
 
 		const FPCGDataTypeIdentifier& Type = InDataDesc.GetType();
 
-		if (!ensure(PCGComputeHelpers::IsTypeAllowedInDataCollection(Type)))
+		if (!PCGComputeHelpers::IsTypeAllowedInDataCollection(Type))
 		{
 			return 0u;
 		}
