PCG Editor Mode Spawn Class Selection Hotfix
- Added API for tools to check classes for validity
- Added fallback spawn class API for tools
- Volume Tool spawns APCGVolume as a fallback in case the graph spawn class is incompatible

diff --git a/Engine/Plugins/PCG/Source/PCGEditor/Private/EditorMode/Tools/Line/PCGDrawSplineTool.cpp b/Engine/Plugins/PCG/Source/PCGEditor/Private/EditorMode/Tools/Line/PCGDrawSplineTool.cpp
--- a/Engine/Plugins/PCG/Source/PCGEditor/Private/EditorMode/Tools/Line/PCGDrawSplineTool.cpp
+++ b/Engine/Plugins/PCG/Source/PCGEditor/Private/EditorMode/Tools/Line/PCGDrawSplineTool.cpp
@@ -200,10 +200,10 @@ namespace PCGDrawSplineTool
 	}
 } // namespace PCGDrawSplineTool
 
-bool UPCGInteractiveToolSettings_SplineBase::IsWorkingActorCompatible(const AActor* InActor) const
+bool UPCGInteractiveToolSettings_SplineBase::IsWorkingActorCompatibleInternal(const AActor* InActor) const
 {
 	// Working actor needs to have a root component to attach a spline to it.
-	return Super::IsWorkingActorCompatible(InActor) && InActor && InActor->GetRootComponent();
+	return InActor && InActor->GetRootComponent();
 }
 
 UPCGInteractiveToolSettings_SplineBase::UPCGInteractiveToolSettings_SplineBase()
diff --git a/Engine/Plugins/PCG/Source/PCGEditor/Private/EditorMode/Tools/Line/PCGDrawSplineTool.h b/Engine/Plugins/PCG/Source/PCGEditor/Private/EditorMode/Tools/Line/PCGDrawSplineTool.h
--- a/Engine/Plugins/PCG/Source/PCGEditor/Private/EditorMode/Tools/Line/PCGDrawSplineTool.h
+++ b/Engine/Plugins/PCG/Source/PCGEditor/Private/EditorMode/Tools/Line/PCGDrawSplineTool.h
@@ -30,14 +30,6 @@ class UPCGInteractiveToolSettings_SplineBase : public UPCGInteractiveToolSetting
 	GENERATED_BODY()
 
 public:
-	// ~Begin UPCGInteractiveToolSettings Interface
-
-	/** This function is called by the builder using UPCGInteractiveToolSettings_SplineBase default class which prohibits to have override
-	 * implementation of this function by child classes. Mark it final for that reason, to be revisited if we change the tool classes hierarchy. */
-	virtual bool IsWorkingActorCompatible(const AActor* InActor) const override final;
-
-	// ~End UPCGInteractiveToolSettings Interface
-
 	UPCGInteractiveToolSettings_SplineBase();
 	
 	USplineComponent* GetWorkingSplineComponent(FName DataInstanceIdentifier) const;
@@ -65,6 +57,11 @@ public:
 	/** The settings to manipulate raycast results, such as normals or offsets of a spline point. */
 	UPROPERTY(EditAnywhere, Category = "Raycast", meta = (ShowOnlyInnerProperties))
 	FPCGToolRaycastSettings Settings{};
+
+protected:
+	/** This function is called indirectly by the builder using UPCGInteractiveToolSettings_SplineBase default class which prohibits to have override
+	 * implementation of this function by child classes. Mark it final for that reason, to be revisited if we change the tool classes hierarchy. */
+	virtual bool IsWorkingActorCompatibleInternal(const AActor* InActor) const override final;
 };
 
 UCLASS()
diff --git a/Engine/Plugins/PCG/Source/PCGEditor/Private/EditorMode/Tools/PCGInteractiveToolSettings.cpp b/Engine/Plugins/PCG/Source/PCGEditor/Private/EditorMode/Tools/PCGInteractiveToolSettings.cpp
--- a/Engine/Plugins/PCG/Source/PCGEditor/Private/EditorMode/Tools/PCGInteractiveToolSettings.cpp
+++ b/Engine/Plugins/PCG/Source/PCGEditor/Private/EditorMode/Tools/PCGInteractiveToolSettings.cpp
@@ -406,6 +406,7 @@ void UPCGInteractiveToolSettings::SetDataInstance(FName DataInstanceIdentifier)
 
 void UPCGInteractiveToolSettings::SetActorClassToSpawn(TSubclassOf<AActor> Class)
 {
+	// If this is called, we assume the class is compatible already
 	ActorClassToSpawn = Class;
 
 	// This will call OnPropertyModified, forwarding the change
@@ -502,10 +503,25 @@ TValueOrError<bool, FText> UPCGInteractiveToolSettings::TryLoadDefaults()
 		{
 			return MakeError(FText::FormatOrdered(LOCTEXT("DefaultGraph_IncompatibleToolTag", "Graph {0} does not have compatible tool tag {1}"), FText::FromString(GetPathNameSafe(DefaultGraphInterface)), FText::FromName(GetToolTag())));
 		}
+
+		TSubclassOf<AActor> DefaultActorClassToSpawn = GraphToolData.GetValue().InitialActorClassToSpawn;
+
+		if (!IsActorClassCompatible(DefaultActorClassToSpawn))
+		{
+			if (TOptional<TSubclassOf<AActor>> FallbackClass = GetFallbackActorClassToSpawn())
+			{
+				DefaultActorClassToSpawn = FallbackClass.GetValue();
+			}
+		}
+
+		if (!IsActorClassCompatible(DefaultActorClassToSpawn))
+		{
+			return MakeError(LOCTEXT("DefaultActorClassToSpawn_Invalid", "No valid spawn class found."));
+		}
 		
 		// At this point, the delegates (OnPropertyModified) aren't setup, so we need to set things up manually.
 		SetToolGraph(DefaultGraphInterface);
-		SetActorClassToSpawn(GraphToolData.GetValue().InitialActorClassToSpawn);
+		SetActorClassToSpawn(DefaultActorClassToSpawn);
 		SetDataInstance(GetDefaultDataInstanceName());
 		SetActorLabel(GraphToolData.GetValue().NewActorLabel);
 
@@ -612,6 +628,39 @@ void UPCGInteractiveToolSettings::RestoreGraphInstanceInitialState(UPCGGraphInst
 	}
 }
 
+TSubclassOf<AActor> UPCGInteractiveToolSettings::GetCompatibleActorClassFromGraphOrFallback() const
+{
+	TSubclassOf<AActor> Result = nullptr;
+	
+	if (ToolGraph)
+	{
+		TOptional<FPCGGraphToolData> GraphToolData = ToolGraph->GetGraphToolData();
+
+		if (GraphToolData.IsSet())
+		{
+			TSubclassOf<AActor> GraphToolDataActorClass = GraphToolData->InitialActorClassToSpawn;
+
+			// The graph might have specified an incompatible class too, so we make sure to check
+			if (GraphToolDataActorClass && IsActorClassCompatible(GraphToolDataActorClass))
+			{
+				Result = GraphToolDataActorClass;
+			}
+		}
+	}
+
+	// If the graph didn't have a valid class, attempt to use a tool-specified fallback
+	if (Result == nullptr)
+	{
+		TOptional<TSubclassOf<AActor>> FallbackClass = GetFallbackActorClassToSpawn();
+		if (FallbackClass.IsSet() && FallbackClass.GetValue() != nullptr && IsActorClassCompatible(FallbackClass.GetValue()))
+		{
+			Result = FallbackClass.GetValue();
+		}
+	}
+
+	return Result;
+}
+
 void UPCGInteractiveToolSettings::CommitPropertyChangedEvents()
 {
 	for(auto& Pair : PropertyChangedEventQueue)
@@ -707,8 +756,8 @@ void UPCGInteractiveToolSettings::OnPropertyModified(UObject* Object, FProperty*
 
 			if(GraphToolData.IsSet())
 			{
-				// If the graph changes, we also set the actor class to spawn
-				SetActorClassToSpawn(GraphToolData.GetValue().InitialActorClassToSpawn);
+				// If the graph changes, we also update the actor class to spawn
+				SetActorClassToSpawn(GetCompatibleActorClassFromGraphOrFallback());
 				SetActorLabel(GraphToolData.GetValue().NewActorLabel);
 			}
 		}
@@ -743,6 +792,25 @@ void UPCGInteractiveToolSettings::OnPropertyModified(UObject* Object, FProperty*
 		// Delete previously spawned actor and set up the new one. Ideally recycling the existing PCG Component.
 		if (HasSpawnedActor() && EditActor->GetClass() != ActorClassToSpawn.Get())
 		{
+			// The newly selected actor class might not be compatible, so we give it a chance to fix itself up
+			if (!IsActorClassCompatible(ActorClassToSpawn))
+			{
+				FNotificationInfo Info(LOCTEXT("InvalidActorClassSelected_Notification", "Class is invalid for tool. Resetting to base class."));
+				FSlateNotificationManager::Get().AddNotification(Info)->SetCompletionState(SNotificationItem::CS_Fail);
+
+				ActorClassToSpawn = GetCompatibleActorClassFromGraphOrFallback();
+			}
+			
+			// If the class is still incompatible after attempted fixup, we cancel the tool
+			if (!IsActorClassCompatible(ActorClassToSpawn))
+			{
+				FNotificationInfo Info(LOCTEXT("InvalidActorClassSelected_ToolShutdown_Notification", "No base class found. Shutting down tool."));
+				Info.ExpireDuration = 3.f;
+				FSlateNotificationManager::Get().AddNotification(Info)->SetCompletionState(SNotificationItem::CS_Fail);
+				Tool->GetToolManager()->DeactivateTool(EToolSide::Left, EToolShutdownType::Cancel);
+				return;
+			}
+			
 			AActor* PreviouslySpawnedActor = EditActor.Get();
 			UPCGComponent* PreviousComponent = GetWorkingPCGComponent();
 			EditActor = nullptr;
diff --git a/Engine/Plugins/PCG/Source/PCGEditor/Private/EditorMode/Tools/PCGInteractiveToolSettings.h b/Engine/Plugins/PCG/Source/PCGEditor/Private/EditorMode/Tools/PCGInteractiveToolSettings.h
--- a/Engine/Plugins/PCG/Source/PCGEditor/Private/EditorMode/Tools/PCGInteractiveToolSettings.h
+++ b/Engine/Plugins/PCG/Source/PCGEditor/Private/EditorMode/Tools/PCGInteractiveToolSettings.h
@@ -76,13 +76,13 @@ public:
 	/** Whether we are operating on a spawned or a previously-existing actor. */
 	bool UsesExistingActor() const;
 	
-	virtual void SetActorClassToSpawn(TSubclassOf<AActor> Class);
+	void SetActorClassToSpawn(TSubclassOf<AActor> Class);
 
 	void SetActorLabel(FName InLabel);
 
-	/** Tools can have a requirement on the working actor */
-	virtual bool IsWorkingActorCompatible(const AActor* InActor) const { return InActor != nullptr; }
-
+	bool IsWorkingActorCompatible(const AActor* InActor) const { return InActor != nullptr && IsWorkingActorCompatibleInternal(InActor); }
+	bool IsActorClassCompatible(TSubclassOf<AActor> InClass) const { return InClass != nullptr && !InClass->HasAnyClassFlags(CLASS_Abstract) && IsActorClassCompatibleInternal(InClass); }
+	
 	/** Filter details view graph selection. */
 	UFUNCTION()
 	virtual bool GraphAssetFilter(const FAssetData& AssetData) const;
@@ -191,6 +191,12 @@ public:
 		return Result;
 	}
 protected:
+	
+	/** Tools can have a requirement on the working actor */
+	virtual bool IsWorkingActorCompatibleInternal(const AActor* InActor) const { return true; }
+	virtual bool IsActorClassCompatibleInternal(TSubclassOf<AActor> InClass) const { return true; }
+	virtual TOptional<TSubclassOf<AActor>> GetFallbackActorClassToSpawn() const { return {}; }
+	
 	/** Try to load the defaults (graph, actor class, ...) into the tool settings (ToolGraph, etc.). */
 	virtual TValueOrError<bool, FText> TryLoadDefaults();
 	virtual TValueOrError<bool, FText> InitActorAndComponent();
@@ -209,7 +215,7 @@ protected:
 	virtual void PostWorkingDataInitialized(FPCGInteractiveToolWorkingData* WorkingData) const;
 	virtual void OnPropertyModified(UObject* Object, FProperty* Property);
 	virtual void OnRefreshInternal();
-	
+
 private:
 	UPCGComponent* FindOrGeneratePCGComponent(AActor& OwningActor) const;
 	
@@ -219,6 +225,7 @@ private:
 
 	void CacheGraphInstance(UPCGGraphInstance* InGraphInstance);
 	void RestoreGraphInstanceInitialState(UPCGGraphInstance* OutGraphInstance);
+	TSubclassOf<AActor> GetCompatibleActorClassFromGraphOrFallback() const;
 
 public:
 	UPROPERTY(EditAnywhere, BlueprintReadOnly, Category = "PCG", meta = (TransientToolProperty, NoResetToDefault, GetAssetFilter = "GraphAssetFilter"))
@@ -240,7 +247,8 @@ public:
 	UPROPERTY(EditAnywhere, DisplayName="Component Name", Category = "PCG", meta = (TransientToolProperty, EditCondition="HasGeneratedPCGComponent()", EditConditionHides))
 	FName NewPCGComponentName;
 
-	UPROPERTY(EditAnywhere, Category = "PCG", meta=(TransientToolProperty, EditCondition="HasSpawnedActor()", EditConditionHides))
+	/** The actor class to use when the tool works on a spawned actor opposed to a pre-existing one. */
+	UPROPERTY(EditAnywhere, Category = "PCG", meta=(TransientToolProperty, EditCondition="HasSpawnedActor()", EditConditionHides, NoResetToDefault))
 	TSubclassOf<AActor> ActorClassToSpawn;
 
 	/** The current time left before a graph refresh (in seconds). */
diff --git a/Engine/Plugins/PCG/Source/PCGEditor/Private/EditorMode/Tools/Paint/PCGPaintTool.cpp b/Engine/Plugins/PCG/Source/PCGEditor/Private/EditorMode/Tools/Paint/PCGPaintTool.cpp
--- a/Engine/Plugins/PCG/Source/PCGEditor/Private/EditorMode/Tools/Paint/PCGPaintTool.cpp
+++ b/Engine/Plugins/PCG/Source/PCGEditor/Private/EditorMode/Tools/Paint/PCGPaintTool.cpp
@@ -51,10 +51,10 @@ FName UPCGInteractiveToolSettings_PaintTool::GetToolTag() const
 	return StaticGetToolTag();
 }
 
-bool UPCGInteractiveToolSettings_PaintTool::IsWorkingActorCompatible(const AActor* InActor) const
+bool UPCGInteractiveToolSettings_PaintTool::IsWorkingActorCompatibleInternal(const AActor* InActor) const
 {
 	// Working actor should have a root component, we'll create one only for new actors.
-	return Super::IsWorkingActorCompatible(InActor) && InActor && InActor->GetRootComponent();
+	return InActor && InActor->GetRootComponent();
 }
 
 const UScriptStruct* UPCGInteractiveToolSettings_PaintTool::GetWorkingDataType() const
diff --git a/Engine/Plugins/PCG/Source/PCGEditor/Private/EditorMode/Tools/Paint/PCGPaintTool.h b/Engine/Plugins/PCG/Source/PCGEditor/Private/EditorMode/Tools/Paint/PCGPaintTool.h
--- a/Engine/Plugins/PCG/Source/PCGEditor/Private/EditorMode/Tools/Paint/PCGPaintTool.h
+++ b/Engine/Plugins/PCG/Source/PCGEditor/Private/EditorMode/Tools/Paint/PCGPaintTool.h
@@ -26,7 +26,6 @@ public:
 	virtual FName GetToolTag() const override;
 	virtual const UScriptStruct* GetWorkingDataType() const override;
 	virtual void RegisterPropertyWatchers() override;
-	virtual bool IsWorkingActorCompatible(const AActor* InActor) const override final;
 	// ~End UPCGInteractiveToolSettings interface
 
 	/** The brush shape to use. */
@@ -53,6 +52,9 @@ public:
 
 	/** Tracks the last point that was added. Used for min-spacing during drag. */
 	TOptional<FTransform> LastPointTransform;
+
+protected:
+	virtual bool IsWorkingActorCompatibleInternal(const AActor* InActor) const override final;
 };
 
 UCLASS(Transient)
diff --git a/Engine/Plugins/PCG/Source/PCGEditor/Private/EditorMode/Tools/Volume/PCGVolumeTool.cpp b/Engine/Plugins/PCG/Source/PCGEditor/Private/EditorMode/Tools/Volume/PCGVolumeTool.cpp
--- a/Engine/Plugins/PCG/Source/PCGEditor/Private/EditorMode/Tools/Volume/PCGVolumeTool.cpp
+++ b/Engine/Plugins/PCG/Source/PCGEditor/Private/EditorMode/Tools/Volume/PCGVolumeTool.cpp
@@ -54,88 +54,86 @@ const UScriptStruct* UPCGInteractiveToolSettings_Volume::GetWorkingDataType() co
 	return FPCGInteractiveToolWorkingData_Volume::StaticStruct();
 }
 
-bool UPCGInteractiveToolSettings_Volume::IsWorkingActorCompatible(const AActor* InActor) const
+bool UPCGInteractiveToolSettings_Volume::IsWorkingActorCompatibleInternal(const AActor* InActor) const
 {
 	// Working actor needs to be a volume or have a box component
-	return Super::IsWorkingActorCompatible(InActor) && InActor && (InActor->IsA<AVolume>() || InActor->GetComponentByClass<UBoxComponent>() != nullptr);
+	return InActor && (InActor->IsA<AVolume>() || InActor->GetComponentByClass<UBoxComponent>() != nullptr);
 }
 
-AVolume* UPCGInteractiveToolSettings_Volume::GetWorkingVolume(FName DataInstanceIdentifier) const
-{
-	if (const FPCGInteractiveToolWorkingData_Volume* ToolData = GetTypedWorkingData<FPCGInteractiveToolWorkingData_Volume>(DataInstanceIdentifier))
-	{
-		return ToolData->Volume.Get();
-	}
-
-	return nullptr;
-}
-
-UBoxComponent* UPCGInteractiveToolSettings_Volume::GetWorkingBoxComponent(FName DataInstanceIdentifier) const
-{
-	if (const FPCGInteractiveToolWorkingData_Volume* ToolData = GetTypedWorkingData<FPCGInteractiveToolWorkingData_Volume>(DataInstanceIdentifier))
-	{
-		return ToolData->BoxComponent.Get();
-	}
-
-	return nullptr;
-}
-
-void UPCGInteractiveToolSettings_Volume::SetActorClassToSpawn(TSubclassOf<AActor> InClass)
-{
-	TSubclassOf<AActor> Class = InClass;
-
-	bool bClassIsValid = (Class != nullptr);
-
+bool UPCGInteractiveToolSettings_Volume::IsActorClassCompatibleInternal(TSubclassOf<AActor> InClass) const
+{	
 	// Accept volumes or actors that have a box component
-	if (bClassIsValid)
+	if (InClass != nullptr)
 	{
-		if (!Class->IsChildOf(AVolume::StaticClass()))
+		if (InClass->IsChildOf<AVolume>())
 		{
-			bool bActorClassHasBoxComponent = false;
+			return true;
+		}		
+		
+		bool bActorClassHasBoxComponent = false;
 
-			if (AActor* ActorCDO = Class->GetDefaultObject<AActor>())
-			{
-				bActorClassHasBoxComponent |= ActorCDO->GetComponentByClass<UBoxComponent>() != nullptr;
+		if (AActor* ActorCDO = InClass->GetDefaultObject<AActor>())
+		{
+			bActorClassHasBoxComponent |= ActorCDO->GetComponentByClass<UBoxComponent>() != nullptr;
 
-				// Check subobjects to cover for BP classes.
-				if (!bActorClassHasBoxComponent && Cast<UBlueprintGeneratedClass>(Class) != nullptr)
-				{
-					TArray<FSubobjectDataHandle> SubobjectHandles;
-					USubobjectDataSubsystem::Get()->GatherSubobjectData(ActorCDO, SubobjectHandles);
+			// Check subobjects to cover for BP classes.
+			if (!bActorClassHasBoxComponent && Cast<UBlueprintGeneratedClass>(InClass) != nullptr)
+			{
+				TArray<FSubobjectDataHandle> SubobjectHandles;
+				USubobjectDataSubsystem::Get()->GatherSubobjectData(ActorCDO, SubobjectHandles);
 
-					const FSubobjectDataHandle* FoundHandle = SubobjectHandles.FindByPredicate(
-						[](const FSubobjectDataHandle& Handle)
+				const FSubobjectDataHandle* FoundHandle = SubobjectHandles.FindByPredicate(
+					[](const FSubobjectDataHandle& Handle)
+					{
+						if (FSubobjectData* SubobjectData = Handle.GetData())
 						{
-							if (FSubobjectData* SubobjectData = Handle.GetData())
+							if (const UActorComponent* ComponentTemplate = SubobjectData->GetComponentTemplate())
 							{
-								if (const UActorComponent* ComponentTemplate = SubobjectData->GetComponentTemplate())
-								{
-									return ComponentTemplate->IsA<UBoxComponent>();
-								}
+								return ComponentTemplate->IsA<UBoxComponent>();
 							}
-
-							return false;
 						}
-					);
 
-					if (FoundHandle && FoundHandle->IsValid())
-					{
-						bActorClassHasBoxComponent |= (FoundHandle->GetData()->GetObject<UBoxComponent>() != nullptr);
+						return false;
 					}
+				);
+
+				if (FoundHandle && FoundHandle->IsValid())
+				{
+					bActorClassHasBoxComponent |= (FoundHandle->GetData()->GetObject<UBoxComponent>() != nullptr);
 				}
 			}
-
-			bClassIsValid &= bActorClassHasBoxComponent;
 		}
+
+		return bActorClassHasBoxComponent;
 	}
+	
+	return false;
+}
 
-	if(!bClassIsValid)
+TOptional<TSubclassOf<AActor>> UPCGInteractiveToolSettings_Volume::GetFallbackActorClassToSpawn() const
+{
+	TSubclassOf<AActor> Result = APCGVolume::StaticClass();
+	return Result;
+}
+
+AVolume* UPCGInteractiveToolSettings_Volume::GetWorkingVolume(FName DataInstanceIdentifier) const
+{
+	if (const FPCGInteractiveToolWorkingData_Volume* ToolData = GetTypedWorkingData<FPCGInteractiveToolWorkingData_Volume>(DataInstanceIdentifier))
 	{
-		UE_LOG(LogPCGEditor, Error, TEXT("Invalid class used in volume tool. Will use PCGVolume instead."));
-		Class = APCGVolume::StaticClass();
+		return ToolData->Volume.Get();
 	}
 
-	Super::SetActorClassToSpawn(Class);
+	return nullptr;
+}
+
+UBoxComponent* UPCGInteractiveToolSettings_Volume::GetWorkingBoxComponent(FName DataInstanceIdentifier) const
+{
+	if (const FPCGInteractiveToolWorkingData_Volume* ToolData = GetTypedWorkingData<FPCGInteractiveToolWorkingData_Volume>(DataInstanceIdentifier))
+	{
+		return ToolData->BoxComponent.Get();
+	}
+
+	return nullptr;
 }
 
 void UPCGVolumeTool::Setup()
diff --git a/Engine/Plugins/PCG/Source/PCGEditor/Private/EditorMode/Tools/Volume/PCGVolumeTool.h b/Engine/Plugins/PCG/Source/PCGEditor/Private/EditorMode/Tools/Volume/PCGVolumeTool.h
--- a/Engine/Plugins/PCG/Source/PCGEditor/Private/EditorMode/Tools/Volume/PCGVolumeTool.h
+++ b/Engine/Plugins/PCG/Source/PCGEditor/Private/EditorMode/Tools/Volume/PCGVolumeTool.h
@@ -27,13 +27,7 @@ public:
 	// ~Begin UPCGInteractiveToolSettings interface
 	virtual FName GetToolTag() const override;
 	virtual const UScriptStruct* GetWorkingDataType() const override;
-	virtual void SetActorClassToSpawn(TSubclassOf<AActor> Class) override;
 	virtual bool CanResetToolData(FName DataInstanceIdentifier) override { return false; }
-
-	/** This function is called by the builder using UPCGInteractiveToolSettings_Volume default class which prohibits to have override
-	 * implementation of this function by child classes. Mark it final for that reason, to be revisited if we change the tool classes hierarchy. */
-	virtual bool IsWorkingActorCompatible(const AActor* InActor) const override final;
-
 	// ~End UPCGInteractiveToolSettings interface
 
 	AVolume* GetWorkingVolume(FName DataInstanceIdentifier) const;
@@ -43,6 +37,13 @@ public:
 
 	UPROPERTY(EditAnywhere, Category = Raycast)
 	FPCGRaycastFilterRuleCollection RaycastRuleCollection;
+
+protected:
+	/** This function is called by the builder using UPCGInteractiveToolSettings_Volume default class which prohibits to have override
+	 * implementation of this function by child classes. Mark it final for that reason, to be revisited if we change the tool classes hierarchy. */
+	virtual bool IsWorkingActorCompatibleInternal(const AActor* InActor) const override final;
+	virtual bool IsActorClassCompatibleInternal(TSubclassOf<AActor> InClass) const override final;
+	virtual TOptional<TSubclassOf<AActor>> GetFallbackActorClassToSpawn() const override final;
 };
 
 /** 
