[PCG] Fix dirty generated partition actors on load and components tags not inherited on local component

#rnx

diff --git a/Engine/Plugins/PCG/Source/PCG/Private/PCGComponent.cpp b/Engine/Plugins/PCG/Source/PCG/Private/PCGComponent.cpp
--- a/Engine/Plugins/PCG/Source/PCG/Private/PCGComponent.cpp
+++ b/Engine/Plugins/PCG/Source/PCG/Private/PCGComponent.cpp
@@ -517,6 +517,9 @@ void UPCGComponent::SetPropertiesFromOriginal(const UPCGComponent* Original)
 		return;
 	}
 
+	// Make sure the tags match
+	ComponentTags = Original->ComponentTags;
+
 	const bool bGraphInstanceIsDifferent = !GraphInstance->IsEquivalent(Original->GraphInstance);
 
 #if WITH_EDITOR
@@ -530,8 +533,10 @@ void UPCGComponent::SetPropertiesFromOriginal(const UPCGComponent* Original)
 	bOverrideGenerationRadii = Original->bOverrideGenerationRadii;
 	GenerationRadii = Original->GenerationRadii;
 
-	UPCGGraph* OriginalGraph = Original->GraphInstance ? Original->GraphInstance->GetGraph() : nullptr;
-	if (OriginalGraph != GraphInstance->GetGraph())
+	// Make sure the Graph instance on the local component is the same than the graph instance on the original component
+	// (Same graph and same parameter overrides)
+	UPCGGraphInterface* OriginalGraph = Original->GraphInstance ? Original->GraphInstance->Graph : nullptr;
+	if (OriginalGraph != GraphInstance->Graph)
 	{
 		GraphInstance->SetGraph(OriginalGraph);
 	}
diff --git a/Engine/Plugins/PCG/Source/PCG/Private/PCGGraph.cpp b/Engine/Plugins/PCG/Source/PCG/Private/PCGGraph.cpp
--- a/Engine/Plugins/PCG/Source/PCG/Private/PCGGraph.cpp
+++ b/Engine/Plugins/PCG/Source/PCG/Private/PCGGraph.cpp
@@ -2687,20 +2687,23 @@ void UPCGGraphInstance::UpdatePropertyOverride(const FProperty* InProperty, bool
 
 void UPCGGraphInstance::CopyParameterOverrides(UPCGGraphInterface* InGraph)
 {
-	if (!InGraph)
+	// Can't copy if they do not have the same base graph
+	if (!InGraph || GetGraph() != InGraph->GetGraph())
 	{
 		return;
 	}
 
-	const UPCGGraph* ThisGraph = GetGraph();
-	const UPCGGraph* OtherGraph = InGraph->GetGraph();
-
-	// Can't copy if they have not the same base graph
-	if (ThisGraph != OtherGraph)
+	// Inherit from the instance overrides or reset them.
+	if (const UPCGGraphInstance* OtherGraphInstance = Cast<UPCGGraphInstance>(InGraph))
 	{
-		return;
+		ParametersOverrides.PropertiesIDsOverridden = OtherGraphInstance->ParametersOverrides.PropertiesIDsOverridden;
+	}
+	else
+	{
+		ParametersOverrides.PropertiesIDsOverridden.Reset();
 	}
 
+	// And copy everything, to be sure we are up to date.
 	ParametersOverrides.Parameters.CopyMatchingValuesByID(*InGraph->GetUserParametersStruct());
 }
 
