[PCG] Avoid renaming PCG Partition Actors as it often leads to UB

#rnx

diff --git a/Engine/Plugins/PCG/Source/PCG/Private/RuntimeGen/PCGRuntimeGenScheduler.cpp b/Engine/Plugins/PCG/Source/PCG/Private/RuntimeGen/PCGRuntimeGenScheduler.cpp
--- a/Engine/Plugins/PCG/Source/PCG/Private/RuntimeGen/PCGRuntimeGenScheduler.cpp
+++ b/Engine/Plugins/PCG/Source/PCG/Private/RuntimeGen/PCGRuntimeGenScheduler.cpp
@@ -1650,7 +1650,6 @@ void FPCGRuntimeGenScheduler::CleanupLocalComponent(APCGPartitionActor* Partitio
 			}
 
 #if WITH_EDITOR
-			PartitionActor->Rename(nullptr, PartitionActor->GetOuter(), REN_NonTransactional | REN_DoNotDirty);
 			PartitionActor->SetActorLabel(*PCGRuntimeGenSchedulerConstants::PooledPartitionActorName);
 #endif
 			PartitionActorPool.Push(PartitionActor);
@@ -1662,9 +1661,6 @@ void FPCGRuntimeGenScheduler::CleanupLocalComponent(APCGPartitionActor* Partitio
 				UE_LOG(LogPCG, Warning, TEXT("[RUNTIMEGEN] DESTROY PARTITION ACTOR: '%s'"), *PartitionActor->GetActorNameOrLabel());
 			}
 
-#if WITH_EDITOR
-			PartitionActor->Rename(nullptr, PartitionActor->GetOuter(), REN_NonTransactional | REN_DoNotDirty);
-#endif
 			World->DestroyActor(PartitionActor);
 		}
 	}
@@ -1981,8 +1977,6 @@ APCGPartitionActor* FPCGRuntimeGenScheduler::GetPartitionActorFromPool(const FPC
 
 #if WITH_EDITOR
 	const FName ActorName = *APCGPartitionActor::GetPCGPartitionActorName(GridDescriptor, GridCoords);
-
-	PartitionActor->Rename(*ActorName.ToString(), PartitionActor->GetOuter(), REN_NonTransactional | REN_DoNotDirty);
 	PartitionActor->SetActorLabel(ActorName.ToString());
 #endif
 
@@ -2010,9 +2004,6 @@ void FPCGRuntimeGenScheduler::AddPartitionActorPoolCount(int32 Count)
 
 	FActorSpawnParameters SpawnParams;
 #if WITH_EDITOR
-	SpawnParams.NameMode = FActorSpawnParameters::ESpawnActorNameMode::Requested;
-	SpawnParams.Name = *PCGRuntimeGenSchedulerConstants::PooledPartitionActorName;
-
 	// Always hide pooled actors from outliner. Note that outliner tree view updates can incur significant costs in Slate code.
 	SpawnParams.bHideFromSceneOutliner = PCGRuntimeGenSchedulerHelpers::CVarHideActorsFromOutliner.GetValueOnAnyThread();
 #endif
@@ -2041,9 +2032,6 @@ void FPCGRuntimeGenScheduler::ResetPartitionActorPoolToSize(uint32 NewPoolSize)
 {
 	for (APCGPartitionActor* PartitionActor : PartitionActorPool)
 	{
-#if WITH_EDITOR
-		PartitionActor->Rename(nullptr, PartitionActor->GetOuter(), REN_NonTransactional | REN_DoNotDirty);
-#endif
 		World->DestroyActor(PartitionActor);
 	}
 
