[PCG] Fixed potential crash in inspection data manager

#rnx

#Synced-CL 48405364

diff --git a/Engine/Plugins/PCG/Source/PCGEditor/Private/Managers/PCGEditorInspectionDataManager.cpp b/Engine/Plugins/PCG/Source/PCGEditor/Private/Managers/PCGEditorInspectionDataManager.cpp
--- a/Engine/Plugins/PCG/Source/PCGEditor/Private/Managers/PCGEditorInspectionDataManager.cpp
+++ b/Engine/Plugins/PCG/Source/PCGEditor/Private/Managers/PCGEditorInspectionDataManager.cpp
@@ -142,7 +142,7 @@ FPCGEditorInspectionDataEntry FPCGEditorInspectionDataManager::ConstructEntry(co
 
 bool FPCGEditorInspectionDataManager::GetInspectionData_Internal(const FPCGEditorInspectionDataEntrySetupParams& InParams, FPCGDataCollection const*& OutData, FPCGStack& OutStack, UPCGPin const*& OutPin) const
 {
-	const UPCGNode* PCGNode = InParams.Node->GetPCGNode();
+	const UPCGNode* PCGNode = InParams.Node ? InParams.Node->GetPCGNode() : nullptr;
 	if (!PCGNode)
 	{
 		return false;
@@ -176,7 +176,10 @@ bool FPCGEditorInspectionDataManager::GetInspectionData_Internal(const FPCGEdito
 	StackFrames.Emplace(PCGNode);
 	StackFrames.Emplace(OutPin);
 
-	OutData = PCGSourceBeingInspected.Get()->GetExecutionState().GetInspection().GetInspectionData(OutStack);
-
+	if (IPCGGraphExecutionSource* Source = PCGSourceBeingInspected.Get())
+	{
+		OutData = Source->GetExecutionState().GetInspection().GetInspectionData(OutStack);
+	}
+	
 	return OutData != nullptr;
 }
