[PCG] PCG Partition actors are not generating on BeginPlay in PIE

#rnx

diff --git a/Engine/Plugins/PCG/Source/PCG/Private/PCGActorAndComponentMapping.cpp b/Engine/Plugins/PCG/Source/PCG/Private/PCGActorAndComponentMapping.cpp
--- a/Engine/Plugins/PCG/Source/PCG/Private/PCGActorAndComponentMapping.cpp
+++ b/Engine/Plugins/PCG/Source/PCG/Private/PCGActorAndComponentMapping.cpp
@@ -706,7 +706,9 @@ PRAGMA_ENABLE_DEPRECATION_WARNINGS
 						{
 							// In editor we might load/create partition actors while the component is registering. Because of that,
 							// the mapping might not already exists, even if the component is marked generated.
-							if (!UE::GetIsEditorLoadingPackage())
+							// Make sure that this condition is only checked when we are in an editor world.
+							const UWorld* World = GetWorld();
+							if ((World && World->IsGameWorld()) || !UE::GetIsEditorLoadingPackage())
 							{
 								InActor->AddGraphInstance(Component);
 							}
