[PCG] Adjusted pcg editor tool physics rays to be a bit more flexible to different collision profiles.

diff --git a/Engine/Plugins/PCG/Source/PCGEditor/Private/EditorMode/Tools/Helpers/PCGEdModeSceneQueryHelpers.cpp b/Engine/Plugins/PCG/Source/PCGEditor/Private/EditorMode/Tools/Helpers/PCGEdModeSceneQueryHelpers.cpp
--- a/Engine/Plugins/PCG/Source/PCGEditor/Private/EditorMode/Tools/Helpers/PCGEdModeSceneQueryHelpers.cpp
+++ b/Engine/Plugins/PCG/Source/PCGEditor/Private/EditorMode/Tools/Helpers/PCGEdModeSceneQueryHelpers.cpp
@@ -57,8 +57,15 @@ namespace UE::PCG::EditorMode::Scene
 
 		for (const FHitResult& Hit : OutHits)
 		{
-			// Skip hits on non-physical components, esp. brush components, like PCG volumes.
-			if (Hit.GetComponent() && !Hit.GetComponent()->IsPhysicsCollisionEnabled())
+			UPrimitiveComponent* HitComponent = Hit.GetComponent();
+
+			// Skip "No collision" type actors (non-physical components, esp. brush components, like PCG volumes).
+			if (!HitComponent || !HitComponent->IsQueryCollisionEnabled())
+			{
+				continue;
+			}
+
+			if (HitComponent->GetCollisionResponseToChannel(ECC_WorldStatic) != ECR_Block && HitComponent->GetCollisionResponseToChannel(ECC_WorldDynamic) != ECR_Block)
 			{
 				continue;
 			}
@@ -75,9 +82,9 @@ namespace UE::PCG::EditorMode::Scene
 			// Check for a visible result
 			bool bHitIsValid = IsHitTargetVisible(Hit);
 			// If invisible, check if it's explicitly allowed.
-			bHitIsValid |= Params.AllowedInvisibleComponents && Params.AllowedInvisibleComponents->Contains(Hit.GetComponent());
+			bHitIsValid |= Params.AllowedInvisibleComponents && Params.AllowedInvisibleComponents->Contains(HitComponent);
 			// Check that it's not explicitly ignored.
-			bHitIsValid &= !Params.ComponentsToIgnore || !Params.ComponentsToIgnore->Contains(Hit.GetComponent());
+			bHitIsValid &= !Params.ComponentsToIgnore || !Params.ComponentsToIgnore->Contains(HitComponent);
 
 			if (bHitIsValid)
 			{
