[PCG] Fixed AddReferenceObjects not locking the inspection cache that could lead to an ensure or worse.

#rnx
#Synced-CL 48810645

diff --git a/Engine/Plugins/PCG/Source/PCG/Private/PCGGraphExecutionInspection.cpp b/Engine/Plugins/PCG/Source/PCG/Private/PCGGraphExecutionInspection.cpp
--- a/Engine/Plugins/PCG/Source/PCG/Private/PCGGraphExecutionInspection.cpp
+++ b/Engine/Plugins/PCG/Source/PCG/Private/PCGGraphExecutionInspection.cpp
@@ -284,6 +284,7 @@ bool FPCGGraphExecutionInspection::DidNodeTriggerCPUToGPUUpload(const UPCGNode*
 
 void FPCGGraphExecutionInspection::AddReferencedObjects(FReferenceCollector& Collector)
 {
+	PCG::TSharedScopeLock ReadLock(InspectionCacheLock);
 	for (auto& Pair : InspectionCache)
 	{
 		Pair.Value.AddReferences(Collector);
