[PCG] Fix SetPointOnRange / GetPointFromRange not using the RangeStartIndex properly (and validate)

diff --git a/Engine/Plugins/PCG/Source/PCG/Private/Data/PCGBasePointData.cpp b/Engine/Plugins/PCG/Source/PCG/Private/Data/PCGBasePointData.cpp
--- a/Engine/Plugins/PCG/Source/PCG/Private/Data/PCGBasePointData.cpp
+++ b/Engine/Plugins/PCG/Source/PCG/Private/Data/PCGBasePointData.cpp
@@ -1564,38 +1564,66 @@ void UPCGBasePointData::SetMetadataEntryValuesOnRange(UPARAM(ref) FPCGPointOutpu
 
 FPCGPoint UPCGBasePointData::GetPointFromRange(const FPCGPointInputRange& InputRange, int32 Index)
 {
+	if (InputRange.RangeStartIndex < 0)
+	{
+		UE_LOG(LogPCG, Error, TEXT("GetPointFromRange called with invalid InputRange RangeStartIndex '%d'."), InputRange.RangeStartIndex);
+		return FPCGPoint();
+	}
+
+	if (Index < 0 || Index >= InputRange.RangeSize)
+	{
+		UE_LOG(LogPCG, Error, TEXT("GetPointFromRange called with invalid Index '%d' when valid InputRange is [0, %d]."), Index, InputRange.RangeSize - 1);
+		return FPCGPoint();
+	}
+
 	if (!InputRange.PointData)
 	{
-		UE_LOG(LogPCG, Error, TEXT("GetPointFromRange called with invalid point data."));
+		UE_LOG(LogPCG, Error, TEXT("GetPointFromRange called with invalid InputRange PointData."));
 		return FPCGPoint();
 	}
 
-	if (Index < 0 || Index >= InputRange.PointData->GetNumPoints())
+	const int32 FinalIndex = InputRange.RangeStartIndex + Index;
+
+	if (FinalIndex < 0 || FinalIndex >= InputRange.PointData->GetNumPoints())
 	{
-		UE_LOG(LogPCG, Error, TEXT("GetPointFromRange called with invalid index '%d' when valid range is [0, %d]."), Index, InputRange.PointData->GetNumPoints()-1);
+		UE_LOG(LogPCG, Error, TEXT("GetPointFromRange called with invalid index '%d' when valid range is [0, %d]."), FinalIndex, InputRange.PointData->GetNumPoints()-1);
 		return FPCGPoint();
 	}
 
 	const FConstPCGPointValueRanges PointRanges(InputRange.PointData);
-	return PointRanges.GetPoint(Index);
+	return PointRanges.GetPoint(FinalIndex);
 }
 
 void UPCGBasePointData::SetPointOnRange(UPARAM(ref) FPCGPointOutputRange& OutputRange, int32 Index, const FPCGPoint& Point)
 {
+	if (OutputRange.RangeStartIndex < 0)
+	{
+		UE_LOG(LogPCG, Error, TEXT("SetPointOnRange called with invalid OutputRange RangeStartIndex '%d'."), OutputRange.RangeStartIndex);
+		return;
+	}
+
+	if (Index < 0 || Index >= OutputRange.RangeSize)
+	{
+		UE_LOG(LogPCG, Error, TEXT("SetPointOnRange called with invalid Index '%d' when valid OutputRange is [0, %d]."), Index, OutputRange.RangeSize-1);
+		return;
+	}
+
 	if (!OutputRange.PointData)
 	{
-		UE_LOG(LogPCG, Error, TEXT("SetPointOnRange called with invalid point data."));
+		UE_LOG(LogPCG, Error, TEXT("SetPointOnRange called with invalid OutputRange PointData."));
 		return;
 	}
 
-	if (Index < 0 || Index >= OutputRange.PointData->GetNumPoints())
+	const int32 FinalIndex = OutputRange.RangeStartIndex + Index;
+
+	if (FinalIndex < 0 || FinalIndex >= OutputRange.PointData->GetNumPoints())
 	{
-		UE_LOG(LogPCG, Error, TEXT("SetPointOnRange called with invalid index '%d' when valid range is [0, %d]."), Index, OutputRange.PointData->GetNumPoints()-1);
+		UE_LOG(LogPCG, Error, TEXT("SetPointOnRange called with invalid index '%d' when valid range is [0, %d]."), FinalIndex, OutputRange.PointData->GetNumPoints()-1);
 		return;
 	}
 
 	FPCGPointValueRanges PointRanges(OutputRange.PointData, /*bAllocate=*/false);
-	PointRanges.SetFromPoint(Index, Point);
+	PointRanges.SetFromPoint(FinalIndex, Point);
 }
 
 #undef LOCTEXT_NAMESPACE
