[PCG] Fixed build following integration.

#rnx

diff --git a/Engine/Plugins/PCG/Source/PCG/Private/PCGGraphExecutionInspection.cpp b/Engine/Plugins/PCG/Source/PCG/Private/PCGGraphExecutionInspection.cpp
--- a/Engine/Plugins/PCG/Source/PCG/Private/PCGGraphExecutionInspection.cpp
+++ b/Engine/Plugins/PCG/Source/PCG/Private/PCGGraphExecutionInspection.cpp
@@ -284,7 +284,7 @@ bool FPCGGraphExecutionInspection::DidNodeTriggerCPUToGPUUpload(const UPCGNode*
 
 void FPCGGraphExecutionInspection::AddReferencedObjects(FReferenceCollector& Collector)
 {
-	PCG::TSharedScopeLock ReadLock(InspectionCacheLock);
+	FReadScopeLock ReadLock(InspectionCacheLock);
 	for (auto& Pair : InspectionCache)
 	{
 		Pair.Value.AddReferences(Collector);
