[PCG] Add missing FGCScopeGuard in code paths running outside of game thread

diff --git a/Engine/Plugins/PCG/Source/PCG/Private/Data/PCGLandscapeData.cpp b/Engine/Plugins/PCG/Source/PCG/Private/Data/PCGLandscapeData.cpp
--- a/Engine/Plugins/PCG/Source/PCG/Private/Data/PCGLandscapeData.cpp
+++ b/Engine/Plugins/PCG/Source/PCG/Private/Data/PCGLandscapeData.cpp
@@ -723,6 +723,7 @@ void UPCGLandscapeData::InitializeTargetMetadata(const FPCGInitializeFromDataPar
 	{
 		if (DataProps.bGetLayerWeights)
 		{
+			FGCScopeGuard Guard;
 			for (TSoftObjectPtr<ALandscapeProxy> Landscape : Landscapes)
 			{
 				const TArray<FName> Layers = LandscapeCache->GetLayerNames(Landscape.Get());
diff --git a/Engine/Plugins/PCG/Source/PCG/Private/Graph/PCGGraphTask.cpp b/Engine/Plugins/PCG/Source/PCG/Private/Graph/PCGGraphTask.cpp
--- a/Engine/Plugins/PCG/Source/PCG/Private/Graph/PCGGraphTask.cpp
+++ b/Engine/Plugins/PCG/Source/PCG/Private/Graph/PCGGraphTask.cpp
@@ -105,6 +105,7 @@ void FPCGGraphTask::LoadCookedData()
 {
 	if (ElementSource == EPCGElementSource::FromNode)
 	{
+		FGCScopeGuard Guard;
 		check(NodePtr.Get());
 		Node = NodePtr.Get();
 	}
