[PCG] Fix ensure when undoing spline drawing

#rnx

diff --git a/Engine/Plugins/PCG/Source/PCGEditor/Private/EditorMode/Tools/Line/PCGDrawSplineTool.cpp b/Engine/Plugins/PCG/Source/PCGEditor/Private/EditorMode/Tools/Line/PCGDrawSplineTool.cpp
--- a/Engine/Plugins/PCG/Source/PCGEditor/Private/EditorMode/Tools/Line/PCGDrawSplineTool.cpp
+++ b/Engine/Plugins/PCG/Source/PCGEditor/Private/EditorMode/Tools/Line/PCGDrawSplineTool.cpp
@@ -168,15 +168,21 @@ namespace PCGDrawSplineTool
 
 			virtual void Revert(USplineComponent& Spline) override
 			{
-				for (int32 i = 0; i < HitLocations.Num(); ++i)
+				// A spline is invalid if it has just a single point. If after removing all the hit locations we have less than or exactly 1 point, just clear the spline.
+				if (HitLocations.Num() - Spline.GetNumberOfSplinePoints() <= 1)
+				{
+					// Since Reset is not exposed, will set an empty array of points
+					Spline.SetSplinePoints({}, ESplineCoordinateSpace::World, /*bUpdateSpline=*/true);
+				}
+				else
 				{
-					if (Spline.GetNumberOfSplinePoints() <= 0)
+					for (int32 i = 0; i < HitLocations.Num(); ++i)
 					{
-						break;
+						Spline.RemoveSplinePoint(Spline.GetNumberOfSplinePoints() - 1, /*bUpdateSpline=*/false);
 					}
-					Spline.RemoveSplinePoint(Spline.GetNumberOfSplinePoints() - 1, false);
+
+					Spline.UpdateSpline();
 				}
-				Spline.UpdateSpline();
 			}
 
 			virtual FString ToString() const override
