[PCG] Fixed ensure firing in PCGData.h

* Also added some robustness in tagged data and crc parity in PCG Data
* Make sure the CRCs are forwarded and computed correctly in the CombineParams function

#rnx

diff --git a/Engine/Plugins/PCG/Source/PCG/Private/Graph/PCGGraphExecutor.cpp b/Engine/Plugins/PCG/Source/PCG/Private/Graph/PCGGraphExecutor.cpp
--- a/Engine/Plugins/PCG/Source/PCG/Private/Graph/PCGGraphExecutor.cpp
+++ b/Engine/Plugins/PCG/Source/PCG/Private/Graph/PCGGraphExecutor.cpp
@@ -2359,27 +2359,53 @@ void FPCGGraphExecutor::CombineParams(FPCGGraphTask& Task)
 
 		const int32 NewNumberOfInputs = Task.TaskInput.TaggedData.Num() - AllParamsData.Num() + (CombinedParamData ? 1 : 0);
 		check(NewNumberOfInputs >= 0);
-
+		
 		TArray<FPCGTaggedData> TempTaggedData{};
+		TArray<FPCGCrc> TempDataCrcs{};
+
 		TempTaggedData.Reserve(NewNumberOfInputs);
-		for (FPCGTaggedData& TaggedData : Task.TaskInput.TaggedData)
+
+		bool bForwardCrcs = true;
+
+		if (!ensure(Task.TaskInput.TaggedData.Num() == Task.TaskInput.DataCrcs.Num()))
 		{
+			bForwardCrcs = false;
+			TempDataCrcs.SetNumZeroed(NewNumberOfInputs);
+		}
+		else
+		{
+			TempDataCrcs.Reserve(NewNumberOfInputs);
+		}
+
+		for (int32 i = 0; i < Task.TaskInput.TaggedData.Num(); ++i)
+		{
+			FPCGTaggedData& TaggedData = Task.TaskInput.TaggedData[i];
+
 			if (TaggedData.Pin != PCGPinConstants::DefaultParamsLabel)
 			{
-				TempTaggedData.Add(std::move(TaggedData));
+				TempTaggedData.Add(MoveTemp(TaggedData));
+
+				if (bForwardCrcs)
+				{
+					TempDataCrcs.Add(MoveTemp(Task.TaskInput.DataCrcs[i]));
+				}
 			}
 		}
 
-		// Add to the root set since we created a new object, that needs to be kept alive for the duration of the task.
 		if (CombinedParamData)
 		{
 			FPCGTaggedData CombineParams{};
 			CombineParams.Data = CombinedParamData;
 			CombineParams.Pin = PCGPinConstants::DefaultParamsLabel;
 			TempTaggedData.Add(CombineParams);
+			if (bForwardCrcs)
+			{
+				TempDataCrcs.Add(CombineParams.ComputeCrc(/*bFullDataCrc=*/true));
+			}
 		}
 
-		Task.TaskInput.TaggedData = std::move(TempTaggedData);
+		Task.TaskInput.TaggedData = MoveTemp(TempTaggedData);
+		Task.TaskInput.DataCrcs = MoveTemp(TempDataCrcs);
 	}
 }
 
diff --git a/Engine/Plugins/PCG/Source/PCG/Private/Graph/PCGGraphExecutor.h b/Engine/Plugins/PCG/Source/PCG/Private/Graph/PCGGraphExecutor.h
--- a/Engine/Plugins/PCG/Source/PCG/Private/Graph/PCGGraphExecutor.h
+++ b/Engine/Plugins/PCG/Source/PCG/Private/Graph/PCGGraphExecutor.h
@@ -247,7 +247,7 @@ private:
 	void ProcessCachedResults(TArray<FCachedResult*> CachedResults);
 	TArray<FPCGTaskId> ProcessCachedResultsInternal(TArray<FCachedResult*> CachedResults);
 
-	/** Combine all param data into one on the Params pin, if any.*/
+	/** Combine all param data on the Overrides pin into a single output, if needed.*/
 	void CombineParams(FPCGGraphTask& Task);
 	void StoreResults(FPCGTaskId InTaskId, const FPCGDataCollection& InTaskOutput, bool bNeedsManualClear);
 	void ClearResults();
diff --git a/Engine/Plugins/PCG/Source/PCG/Private/PCGData.cpp b/Engine/Plugins/PCG/Source/PCG/Private/PCGData.cpp
--- a/Engine/Plugins/PCG/Source/PCG/Private/PCGData.cpp
+++ b/Engine/Plugins/PCG/Source/PCG/Private/PCGData.cpp
@@ -613,6 +613,7 @@ void FPCGDataCollection::Reset()
 {
 	// Implementation note: We are assuming that there is no need to remove the data from the rootset here.
 	TaggedData.Reset();
+	DataCrcs.Reset();
 	bCancelExecution = false;
 }
 
diff --git a/Engine/Plugins/PCG/Source/PCG/Private/PCGGraphExecutionInspection.cpp b/Engine/Plugins/PCG/Source/PCG/Private/PCGGraphExecutionInspection.cpp
--- a/Engine/Plugins/PCG/Source/PCG/Private/PCGGraphExecutionInspection.cpp
+++ b/Engine/Plugins/PCG/Source/PCG/Private/PCGGraphExecutionInspection.cpp
@@ -147,7 +147,7 @@ void FPCGGraphExecutionInspection::StoreInspectionData(const FPCGStack* InStack,
 				{
 					if (FPCGDataCollection* CollectionInCache = InOutInspectionCache.Find(Stack))
 					{
-						CollectionInCache->TaggedData.Append(PinDataCollection.TaggedData);
+						(*CollectionInCache) += PinDataCollection;
 					}
 					else
 					{
