[PCG] Fix support for non-power-of-2 grid sizes in runtime gen scheduler

Previously would hit a series of ensures and select an invalid radius for generation and cleanup. Now it should round the grid to the nearest valid power of 2 entry when querying for a grid/radius.

#rnx

diff --git a/Engine/Plugins/PCG/Source/PCG/Private/PCGCommon.cpp b/Engine/Plugins/PCG/Source/PCG/Private/PCGCommon.cpp
--- a/Engine/Plugins/PCG/Source/PCG/Private/PCGCommon.cpp
+++ b/Engine/Plugins/PCG/Source/PCG/Private/PCGCommon.cpp
@@ -176,8 +176,18 @@ namespace PCGHiGenGrid
 		{
 			return EPCGHiGenGrid::Unbounded;
 		}
-		// TODO: support other units
-		return ensure(IsValidGridSize(InGridSize)) ? static_cast<EPCGHiGenGrid>(InGridSize / 100) : EPCGHiGenGrid::Uninitialized;
+		else if (InGridSize <= GridToGridSize(EPCGHiGenGrid::GridMin))
+		{
+			return EPCGHiGenGrid::GridMin;
+		}
+		else if (InGridSize >= GridToGridSize(EPCGHiGenGrid::GridMax))
+		{
+			return EPCGHiGenGrid::GridMax;
+		}
+
+		// Grid size is in cm, but EPCGHiGenGrid is in meters so we convert the units before rounding.
+		const uint32 RoundedGridSizeMeters = FMath::RoundUpToPowerOfTwo(InGridSize / 100);
+		return static_cast<EPCGHiGenGrid>(RoundedGridSizeMeters);
 	}
 
 	uint32 UninitializedGridSize()
diff --git a/Engine/Plugins/PCG/Source/PCG/Private/RuntimeGen/PCGRuntimeGenScheduler.cpp b/Engine/Plugins/PCG/Source/PCG/Private/RuntimeGen/PCGRuntimeGenScheduler.cpp
--- a/Engine/Plugins/PCG/Source/PCG/Private/RuntimeGen/PCGRuntimeGenScheduler.cpp
+++ b/Engine/Plugins/PCG/Source/PCG/Private/RuntimeGen/PCGRuntimeGenScheduler.cpp
@@ -510,7 +510,7 @@ void FPCGRuntimeGenScheduler::TickQueueComponentsForGeneration(
 		{
 			bool bHasUnbounded = false;
 			PCGHiGenGrid::FSizeArray GridSizes;
-			ensure(PCGHelpers::GetGenerationGridSizes(OriginalComponent->GetGraph(), Inputs.PCGWorldActor, GridSizes, bHasUnbounded));
+			PCGHelpers::GetGenerationGridSizes(OriginalComponent->GetGraph(), Inputs.PCGWorldActor, GridSizes, bHasUnbounded);
 
 			if (GridSizes.IsEmpty() && !bHasUnbounded)
 			{
@@ -658,8 +658,6 @@ void FPCGRuntimeGenScheduler::TickQueueComponentsForGeneration(
 				{
 					const uint32 GridSize = GridSizes[GridIndex];
 
-					ensure(PCGHiGenGrid::IsValidGridSize(GridSize));
-
 					const FIntVector GenSourceGridPosition = UPCGActorHelpers::GetCellCoord(GenSourcePosition, GridSize, bIs2DGrid);
 					const double GenerationRadius = OriginalComponent->GetGenerationRadiusFromGrid(PCGHiGenGrid::GridSizeToGrid(GridSize));
 					const int32 GridRadius = FMath::CeilToInt32(GenerationRadius / GridSize); // Radius discretized to # of grid cells.
diff --git a/Engine/Plugins/PCG/Source/PCG/Public/PCGCommon.h b/Engine/Plugins/PCG/Source/PCG/Public/PCGCommon.h
--- a/Engine/Plugins/PCG/Source/PCG/Public/PCGCommon.h
+++ b/Engine/Plugins/PCG/Source/PCG/Public/PCGCommon.h
@@ -410,6 +410,7 @@ namespace PCGHiGenGrid
 	PCG_API bool IsValidGrid(EPCGHiGenGrid InGrid);
 	PCG_API bool IsValidGridOrUninitialized(EPCGHiGenGrid InGrid);
 	PCG_API uint32 GridToGridSize(EPCGHiGenGrid InGrid);
+	/** Rounds up to the nearest power of 2 grid, clamping between GridMin and GridMax (excepting the unbounded grid, which may exceed GridMax). */
 	PCG_API EPCGHiGenGrid GridSizeToGrid(uint32 InGridSize);
 
 	PCG_API uint32 UninitializedGridSize();
