[PCG] When running commandlet, make sure GIsRunningUnattendedScript is set to true around FEditorFileUtils::PromptForCheckoutAndSave calls

diff --git a/Engine/Plugins/PCG/Source/PCG/Private/PCGAssetExporterUtils.cpp b/Engine/Plugins/PCG/Source/PCG/Private/PCGAssetExporterUtils.cpp
--- a/Engine/Plugins/PCG/Source/PCG/Private/PCGAssetExporterUtils.cpp
+++ b/Engine/Plugins/PCG/Source/PCG/Private/PCGAssetExporterUtils.cpp
@@ -6,6 +6,7 @@
 #include "PCGModule.h"
 
 #if WITH_EDITOR
+#include "CoreGlobals.h"
 #include "ContentBrowserModule.h"
 #include "FileHelpers.h"
 #include "IContentBrowserSingleton.h"
@@ -172,6 +173,8 @@ UPackage* UPCGAssetExporterUtils::CreateAsset(const UClass* AssetClass, const FP
 	// Save the file
 	if (Package && Parameters.bSaveOnExportEnded)
 	{
+		// Set GIsRunningUnattendedScript to 'true' if it isn't already set while running commandlets
+		TGuardValue<bool> UnattendedScriptGuard(GIsRunningUnattendedScript, GIsRunningUnattendedScript || IsRunningCommandlet());
 		FEditorFileUtils::PromptForCheckoutAndSave({ Package }, /*bCheckDirty=*/false, /*bPromptToSave=*/false);
 	}
 
@@ -234,6 +237,8 @@ void UPCGAssetExporterUtils::UpdateAssets(const TArray<FAssetData>& PCGAssets, c
 	// Save the file(s)
 	if (!PackagesToSave.IsEmpty() && Parameters.bSaveOnExportEnded)
 	{
+		// Set GIsRunningUnattendedScript to 'true' if it isn't already set while running commandlets
+		TGuardValue<bool> UnattendedScriptGuard(GIsRunningUnattendedScript, GIsRunningUnattendedScript || IsRunningCommandlet());
 		FEditorFileUtils::PromptForCheckoutAndSave(PackagesToSave, /*bCheckDirty=*/false, /*bPromptToSave=*/false);
 	}
 #else
diff --git a/Engine/Plugins/PCG/Source/PCGEditor/Private/PCGLevelToAsset.cpp b/Engine/Plugins/PCG/Source/PCGEditor/Private/PCGLevelToAsset.cpp
--- a/Engine/Plugins/PCG/Source/PCGEditor/Private/PCGLevelToAsset.cpp
+++ b/Engine/Plugins/PCG/Source/PCGEditor/Private/PCGLevelToAsset.cpp
@@ -14,6 +14,7 @@
 #include "Helpers/PCGTagHelpers.h"
 #include "Metadata/PCGMetadata.h"
 
+#include "CoreGlobals.h"
 #include "ContentBrowserModule.h"
 #include "FileHelpers.h"
 #include "IContentBrowserSingleton.h"
@@ -54,6 +55,8 @@ void UPCGLevelToAsset::CreateOrUpdatePCGAssets(const TArray<FAssetData>& WorldAs
 	// Save the file(s)
 	if (!PackagesToSave.IsEmpty() && Parameters.bSaveOnExportEnded)
 	{
+		// Set GIsRunningUnattendedScript to 'true' if it isn't already set while running commandlets
+		TGuardValue<bool> UnattendedScriptGuard(GIsRunningUnattendedScript, GIsRunningUnattendedScript || IsRunningCommandlet());
 		FEditorFileUtils::PromptForCheckoutAndSave(PackagesToSave, false, false);
 	}
 }
