[PCG] Fix LS DI normal computation failing on edges of generation volume

#rnx

diff --git a/Engine/Plugins/PCG/Source/PCG/Private/Compute/DataInterfaces/PCGLandscapeDataInterface.cpp b/Engine/Plugins/PCG/Source/PCG/Private/Compute/DataInterfaces/PCGLandscapeDataInterface.cpp
--- a/Engine/Plugins/PCG/Source/PCG/Private/Compute/DataInterfaces/PCGLandscapeDataInterface.cpp
+++ b/Engine/Plugins/PCG/Source/PCG/Private/Compute/DataInterfaces/PCGLandscapeDataInterface.cpp
@@ -356,6 +356,8 @@ void UPCGLandscapeDataProvider::InitFromLandscape(UPCGDataBinding* InBinding, AL
 		return;
 	}
 
+	LandscapeGridSize = FVector2D(InLandscape->GetTransform().GetScale3D());
+
 	if (bAllowSampleVirtualTexture)
 	{
 		InitializeRuntimeVirtualTextures(InLandscape, bAllowSampleVirtualTextureNormal);
@@ -366,8 +368,6 @@ void UPCGLandscapeDataProvider::InitFromLandscape(UPCGDataBinding* InBinding, AL
 	{
 		InitializeFromLandscapeCollision(InBinding, InLandscape, Bounds);
 	}
-
-	LandscapeGridSize = FVector2D(InLandscape->GetTransform().GetScale3D());
 }
 
 void UPCGLandscapeDataProvider::InitializeRuntimeVirtualTextures(ALandscape* InLandscape, bool bAllowSampleVirtualTextureNormal)
@@ -467,7 +467,8 @@ void UPCGLandscapeDataProvider::InitializeFromLandscapeCollision(UPCGDataBinding
 	const FVector LWCTileOffset = FVector(LWCTile) * FLargeWorldRenderScalar::GetTileSize();
 	FBox SystemWorldBounds = Bounds;
 	SystemWorldBounds.Min += LWCTileOffset;
-	SystemWorldBounds.Max += LWCTileOffset;
+	// Expand max bounds by the LS grid size so that we don't go off the height texture when computing normals for points at the edges of the bounds.
+	SystemWorldBounds.Max += LWCTileOffset + FVector(LandscapeGridSize.X, LandscapeGridSize.Y, 0.0);
 
 	const FTransform& LandscapeActorToWorld = InLandscape->LandscapeActorToWorld();
 	const FVector SystemMinInLandscape = LandscapeActorToWorld.InverseTransformPosition(SystemWorldBounds.Min);
