[PCG] Fix for copy attribute from single value to metadata domains not supporting multiple entries.

diff --git a/Engine/Plugins/PCG/Source/PCG/Private/Helpers/PCGMetadataHelpers.cpp b/Engine/Plugins/PCG/Source/PCG/Private/Helpers/PCGMetadataHelpers.cpp
--- a/Engine/Plugins/PCG/Source/PCG/Private/Helpers/PCGMetadataHelpers.cpp
+++ b/Engine/Plugins/PCG/Source/PCG/Private/Helpers/PCGMetadataHelpers.cpp
@@ -208,6 +208,8 @@ namespace PCGMetadataHelpers
 				// Copy entries only if they come from the same data and they are on the same domain.
 				const bool bCopyEntries = bSameOrigin && SourceAttribute->GetMetadataDomain()->GetDomainID() == TargetMetadataDomain->GetDomainID();
 				const bool bSourceSupportsMultiEntries = SourceMetadata->MetadataDomainSupportsMultiEntries(LocalSourceAttribute.MetadataDomain);
+				const bool bTargetSupportsMultiEntries = TargetMetadata->MetadataDomainSupportsMultiEntries(LocalDestinationAttribute.MetadataDomain);
+
 				if (FPCGMetadataAttributeBase* NewAttr = TargetMetadataDomain->CopyAttribute(SourceAttribute, LocalDestinationAttribute.Name, /*bKeepParent=*/false, /*bCopyEntries=*/bCopyEntries, /*bCopyValues=*/true))
 				{
 					// And finally, we create our source and target keys, to get all the metadata entry keys, to remap if we are not the same origin and that we have multiple entries.
@@ -261,7 +263,7 @@ namespace PCGMetadataHelpers
 								for (int EntryIndex = 0; EntryIndex < AllTargetEntryKeysPtr.Num(); ++EntryIndex)
 								{
 									PCGMetadataEntryKey& EntryKey = *AllTargetEntryKeysPtr[EntryIndex];
-									if (EntryKey == PCGInvalidEntryKey || (bSourceSupportsMultiEntries && EntryKey < TargetMetadataDomain->GetItemKeyCountForParent()))
+									if (EntryKey == PCGInvalidEntryKey || (bTargetSupportsMultiEntries && EntryKey < TargetMetadataDomain->GetItemKeyCountForParent()))
 									{
 										AllTargetEntryKeysPtrTemp.Add(&EntryKey);
 									}
