[PCG] Prevented crash when splitting an invalid spline.

#rnx

diff --git a/Engine/Plugins/PCG/Source/PCG/Private/Elements/PCGSplitSplines.cpp b/Engine/Plugins/PCG/Source/PCG/Private/Elements/PCGSplitSplines.cpp
--- a/Engine/Plugins/PCG/Source/PCG/Private/Elements/PCGSplitSplines.cpp
+++ b/Engine/Plugins/PCG/Source/PCG/Private/Elements/PCGSplitSplines.cpp
@@ -265,6 +265,13 @@ bool FPCGSplitSplineElement::ExecuteInternal(FPCGContext* Context) const
 		}
 
 		TArray<FSplinePoint> SplinePoints = SplineData->GetSplinePoints();
+
+		if (SplinePoints.Num() < 2)
+		{
+			Outputs.Add(Input);
+			continue;
+		}
+
 		TConstArrayView<PCGMetadataEntryKey> SplinePointsEntryKeys = SplineData->GetConstVerticesEntryKeys();
 		check(SplinePoints.Num() == SplinePointsEntryKeys.Num() || SplinePointsEntryKeys.IsEmpty());
 
