[PCG] Fix crash on BeginDestroy in UPCGGraph

* Was crashing when the GC is run when the engine is exiting.

#rnx

diff --git a/Engine/Plugins/PCG/Source/PCG/Private/PCGGraph.cpp b/Engine/Plugins/PCG/Source/PCG/Private/PCGGraph.cpp
--- a/Engine/Plugins/PCG/Source/PCG/Private/PCGGraph.cpp
+++ b/Engine/Plugins/PCG/Source/PCG/Private/PCGGraph.cpp
@@ -845,7 +845,12 @@ void UPCGGraph::BeginDestroy()
 		InputNode->OnNodeChangedDelegate.RemoveAll(this);
 	}
 
-	FPCGModule::GetPCGModuleChecked().NotifyGraphChanged(this, EPCGChangeType::Structural | EPCGChangeType::GenerationGrid);
+	// BeginDestroy can happen when the module is already trying to unload, so make sure we only notify if the module is actually loaded.
+	if (FPCGModule::IsPCGModuleLoaded())
+	{
+		FPCGModule::GetPCGModuleChecked().NotifyGraphChanged(this, EPCGChangeType::Structural | EPCGChangeType::GenerationGrid);
+	}
+
 #endif
 
 	Super::BeginDestroy();
