[PCG] Fix crash when using virtual texture data interface in debug builds

#rnx

diff --git a/Engine/Plugins/PCG/Source/PCG/Private/Compute/DataInterfaces/PCGVirtualTextureDataInterface.cpp b/Engine/Plugins/PCG/Source/PCG/Private/Compute/DataInterfaces/PCGVirtualTextureDataInterface.cpp
--- a/Engine/Plugins/PCG/Source/PCG/Private/Compute/DataInterfaces/PCGVirtualTextureDataInterface.cpp
+++ b/Engine/Plugins/PCG/Source/PCG/Private/Compute/DataInterfaces/PCGVirtualTextureDataInterface.cpp
@@ -184,12 +184,19 @@ void UPCGVirtualTextureDataProvider::Reset()
 
 FPCGVirtualTextureDataProviderProxy::FPCGVirtualTextureDataProviderProxy(const TArray<TObjectPtr<const URuntimeVirtualTexture>>& InVirtualTextures)
 {
-	NumVirtualTextures = InVirtualTextures.Num();
-	for (int32 VirtualTextureIndex = 0; VirtualTextureIndex < NumVirtualTextures; ++VirtualTextureIndex)
+	for (const TObjectPtr<const URuntimeVirtualTexture>& VirtualTexture : InVirtualTextures)
+	{
+		VirtualTextures.Add(VirtualTexture);
+	}
+}
+
+void FPCGVirtualTextureDataProviderProxy::AllocateResources(FRDGBuilder& GraphBuilder, FAllocationData const& InAllocationData)
+{
+	for (int32 VirtualTextureIndex = 0; VirtualTextureIndex < VirtualTextures.Num(); ++VirtualTextureIndex)
 	{
 		ERuntimeVirtualTextureMaterialType MaterialType = ERuntimeVirtualTextureMaterialType::Count;
 
-		if (const URuntimeVirtualTexture* VirtualTexture = InVirtualTextures[VirtualTextureIndex])
+		if (TStrongObjectPtr<const URuntimeVirtualTexture> VirtualTexture = VirtualTextures[VirtualTextureIndex].Pin())
 		{
 			MaterialType = VirtualTexture->GetMaterialType();
 			const int32 NumVirtualTextureLayers = FMath::Min(URuntimeVirtualTexture::GetLayerCount(MaterialType), PCGVirtualTextureDIConstants::MAX_NUM_LAYERS);
@@ -202,8 +209,8 @@ FPCGVirtualTextureDataProviderProxy::FPCGVirtualTextureDataProviderProxy(const T
 				// Note, we always use the first layer to get the world to uv transform because every material type uses at least the first layer.
 				const bool bIncludeWorldToUVTransform = LayerIndex == 0;
 
-				Layers[PackedLayerIndex].Initialize(VirtualTexture, LayerIndex, VirtualTexture->IsLayerSRGB(LayerIndex));
-				PageTables[PackedLayerIndex].Initialize(VirtualTexture, LayerIndex, bIncludeWorldToUVTransform, bIncludeWorldToUVTransform);
+				Layers[PackedLayerIndex].Initialize(VirtualTexture.Get(), LayerIndex, VirtualTexture->IsLayerSRGB(LayerIndex));
+				PageTables[PackedLayerIndex].Initialize(VirtualTexture.Get(), LayerIndex, bIncludeWorldToUVTransform, bIncludeWorldToUVTransform);
 			}
 		}
 
@@ -233,7 +240,7 @@ void FPCGVirtualTextureDataProviderProxy::GatherPermutations(FPermutationData& I
 	FPCGVirtualTextureDataInterfacePermutationIds PermutationIds(InOutPermutationData.PermutationVector);
 	for (int32 InvocationIndex = 0; InvocationIndex < InOutPermutationData.NumInvocations; ++InvocationIndex)
 	{
-		InOutPermutationData.PermutationIds[InvocationIndex] |= (NumVirtualTextures > 1 ? PermutationIds.EnableMultipleVirtualTextures : 0);
+		InOutPermutationData.PermutationIds[InvocationIndex] |= (VirtualTextures.Num() > 1 ? PermutationIds.EnableMultipleVirtualTextures : 0);
 	}
 }
 
@@ -243,7 +250,7 @@ void FPCGVirtualTextureDataProviderProxy::GatherDispatchData(FDispatchData const
 	for (int32 InvocationIndex = 0; InvocationIndex < ParameterArray.Num(); ++InvocationIndex)
 	{
 		FParameters& Parameters = ParameterArray[InvocationIndex];
-		Parameters.NumVirtualTextures = NumVirtualTextures;
+		Parameters.NumVirtualTextures = VirtualTextures.Num();
 		Parameters.Sampler = TStaticSamplerState<SF_Bilinear, AM_Clamp, AM_Clamp, AM_Clamp>::GetRHI();
 
 		for (int32 BindingIndex = 0; BindingIndex < PCGVirtualTextureDIConstants::MAX_NUM_BINDINGS; ++BindingIndex)
diff --git a/Engine/Plugins/PCG/Source/PCG/Private/Compute/DataInterfaces/PCGVirtualTextureDataInterface.h b/Engine/Plugins/PCG/Source/PCG/Private/Compute/DataInterfaces/PCGVirtualTextureDataInterface.h
--- a/Engine/Plugins/PCG/Source/PCG/Private/Compute/DataInterfaces/PCGVirtualTextureDataInterface.h
+++ b/Engine/Plugins/PCG/Source/PCG/Private/Compute/DataInterfaces/PCGVirtualTextureDataInterface.h
@@ -76,13 +76,15 @@ public:
 	bool IsValid(FValidationData const& InValidationData) const override;
 	void GatherPermutations(FPermutationData& InOutPermutationData) const override;
 	void GatherDispatchData(FDispatchData const& InDispatchData) override;
+	void AllocateResources(FRDGBuilder& GraphBuilder, FAllocationData const& InAllocationData) override;
 	//~ End FComputeDataProviderRenderProxy Interface
 
 protected:
 	using FParameters = FPCGVirtualTextureDataInterfaceParameters;
 
+	TArray<TWeakObjectPtr<const URuntimeVirtualTexture>> VirtualTextures;
+
 	ERuntimeVirtualTextureMaterialType MaterialTypes[PCGVirtualTextureDIConstants::MAX_NUM_BINDINGS];
 	PCGVirtualTextureCommon::FVirtualTextureLayer Layers[PCGVirtualTextureDIConstants::MAX_NUM_LAYERS * PCGVirtualTextureDIConstants::MAX_NUM_BINDINGS];
 	PCGVirtualTextureCommon::FVirtualTexturePageTable PageTables[PCGVirtualTextureDIConstants::MAX_NUM_LAYERS * PCGVirtualTextureDIConstants::MAX_NUM_BINDINGS];
-	int32 NumVirtualTextures = 0;
 };
