[PCG] Fix crash in UPCGLandscapeSplineData::GetTransformAtDistance()

#rnx

diff --git a/Engine/Plugins/PCG/Source/PCG/Private/Data/PCGLandscapeSplineData.cpp b/Engine/Plugins/PCG/Source/PCG/Private/Data/PCGLandscapeSplineData.cpp
--- a/Engine/Plugins/PCG/Source/PCG/Private/Data/PCGLandscapeSplineData.cpp
+++ b/Engine/Plugins/PCG/Source/PCG/Private/Data/PCGLandscapeSplineData.cpp
@@ -110,6 +110,13 @@ FTransform UPCGLandscapeSplineData::GetTransformAtDistance(int SegmentIndex, FVe
 	check(Segment);
 
 	const TArray<FLandscapeSplineInterpPoint>& InterpPoints = Segment->GetPoints();
+	const int32 SecondToLastPointIndex = InterpPoints.Num() - 2;
+
+	if (SecondToLastPointIndex < 0)
+	{
+		// Not enough points in the segment to interpolate a transform.
+		return GetTransform();
+	}
 
 	int32 PointIndex = 0;
 	FVector::FReal Alpha = 0.0f;
@@ -117,16 +124,18 @@ FTransform UPCGLandscapeSplineData::GetTransformAtDistance(int SegmentIndex, FVe
 	if (ClampedSegmentIndex == SegmentIndex)
 	{
 		GetInterpPointAtDistance(ClampedSegmentIndex, Distance, PointIndex, /*bComputeAlpha=*/true, Alpha);
+		// If our RightPoint ends up being on the next segment, clamp it back to the current segment.
+		PointIndex = FMath::Clamp(PointIndex, 0, SecondToLastPointIndex);
 	}
 	else if(ClampedSegmentIndex < SegmentIndex)
 	{
+		// If we've landed beyond the last segment, just clamp to the end of the last segment.
 		Alpha = 1.0f;
-		PointIndex = FMath::Max(0, InterpPoints.Num() - 2);
+		PointIndex = SecondToLastPointIndex;
 	}
-
-	check(InterpPoints.IsValidIndex(PointIndex));
+ 
 	const FLandscapeSplineInterpPoint& PreviousPoint = InterpPoints[PointIndex];
-	const FLandscapeSplineInterpPoint& CurrentPoint = InterpPoints[FMath::Min(PointIndex + 1, InterpPoints.Num() - 1)]; // If our RightPoint ends up being on the next segment, clamp it back to the current segment.
+	const FLandscapeSplineInterpPoint& CurrentPoint = InterpPoints[PointIndex + 1];
 
 	const FVector XAxis = CurrentPoint.Center - PreviousPoint.Center;
 	const FVector PreviousYAxis = PreviousPoint.Right - PreviousPoint.Center;
diff --git a/Engine/Plugins/PCG/Source/PCG/Private/Elements/PCGSplineSampler.cpp b/Engine/Plugins/PCG/Source/PCG/Private/Elements/PCGSplineSampler.cpp
--- a/Engine/Plugins/PCG/Source/PCG/Private/Elements/PCGSplineSampler.cpp
+++ b/Engine/Plugins/PCG/Source/PCG/Private/Elements/PCGSplineSampler.cpp
@@ -483,6 +483,8 @@ namespace PCGSplineSamplerHelpers
 				OutResult.Distance = DistanceToCurrentSegment + DistanceAlongSegment;
 			}
 
+			const double ScaleFactor = 0.5 * SegmentStep / (FMath::IsNearlyZero(OutTransform.GetScale3D().X) ? UE_DOUBLE_SMALL_NUMBER : OutTransform.GetScale3D().X);
+
 			if (SubpointIndex == 0)
 			{
 				const FVector::FReal PreviousSegmentLength = LineData->GetSegmentLength(PreviousSegmentIndex);
@@ -490,19 +492,19 @@ namespace PCGSplineSamplerHelpers
 
 				if ((PreviousSegmentEndTransform.GetLocation() - OutTransform.GetLocation()).Length() <= KINDA_SMALL_NUMBER)
 				{
-					OutBox.Min.X *= 0.5 * PreviousSegmentLength / (PreviousSegmentEndTransform.GetScale3D().X * (SubdivisionsPerSegment + 1));
+					OutBox.Min.X *= 0.5 * PreviousSegmentLength / ((FMath::IsNearlyZero(PreviousSegmentEndTransform.GetScale3D().X) ? UE_DOUBLE_SMALL_NUMBER : PreviousSegmentEndTransform.GetScale3D().X) * (SubdivisionsPerSegment + 1));
 				}
 				else
 				{
-					OutBox.Min.X *= 0.5 * SegmentStep / OutTransform.GetScale3D().X;
+					OutBox.Min.X *= 0.5 * SegmentStep / ScaleFactor;
 				}
 			}
 			else
 			{
-				OutBox.Min.X *= 0.5 * SegmentStep / OutTransform.GetScale3D().X;
+				OutBox.Min.X *= 0.5 * SegmentStep / ScaleFactor;
 			}
 
-			OutBox.Max.X *= 0.5 * SegmentStep / OutTransform.GetScale3D().X;
+			OutBox.Max.X *= 0.5 * SegmentStep / ScaleFactor;
 
 			++SubpointIndex;
 			if (SubpointIndex > SubdivisionsPerSegment || bLastKeyPoint)
@@ -601,8 +603,9 @@ namespace PCGSplineSamplerHelpers
 			OutResult.InputKey = LineData->GetInputKeyAtDistance(CurrentSegmentIndex, OffsetDistance);
 
 			// Set min/max to half of extent
-			OutBox.Min.X *= 0.5 * DistanceIncrement / OutTransform.GetScale3D().X;
-			OutBox.Max.X *= 0.5 * DistanceIncrement / OutTransform.GetScale3D().X;
+			const double ScaleFactor = 0.5 * DistanceIncrement / (FMath::IsNearlyZero(OutTransform.GetScale3D().X) ? UE_DOUBLE_SMALL_NUMBER : OutTransform.GetScale3D().X);
+			OutBox.Min.X *= ScaleFactor;
+			OutBox.Max.X *= ScaleFactor;
 
 			if (bComputeCurvature)
 			{
