[PCG] Builder: Fix IterativeCell + LoadAsPrevious resource leakage
- AddGraphInstance can update the PartitionActor local component EditingModes without calling ChangeTransientState

diff --git a/Engine/Plugins/PCG/Source/PCGEditor/Private/WorldPartitionBuilder/PCGWorldPartitionBuilder.cpp b/Engine/Plugins/PCG/Source/PCGEditor/Private/WorldPartitionBuilder/PCGWorldPartitionBuilder.cpp
--- a/Engine/Plugins/PCG/Source/PCGEditor/Private/WorldPartitionBuilder/PCGWorldPartitionBuilder.cpp
+++ b/Engine/Plugins/PCG/Source/PCGEditor/Private/WorldPartitionBuilder/PCGWorldPartitionBuilder.cpp
@@ -111,7 +111,7 @@ struct FPCGWorldPartitionBuilder
 
 	static void SetChainedDispatchToLocalComponents(UWorld* InWorld, bool bInChainedDispatch);
 
-	static void UpdateEditingMode(UPCGComponent* InComponent);
+	static void UpdateEditingMode(UPCGComponent* InComponent, bool bForce = false);
 private:
 	static void CollectSubGridComponentsToGenerateInternal(UPCGSubsystem* InPCGSubsystem, int32 InGridSizeIndex, const PCGHiGenGrid::FSizeArray& InGridSizes, UPCGComponent* InOriginalComponent, APCGPartitionActor* InPartitionActor, TArray<UPCGComponent*>& OutAdditionalComponentsToGenerate, TMap<UPCGComponent*, UPCGComponent*>& OutDependencies);
 
@@ -1285,7 +1285,9 @@ FPCGTaskId FPCGWorldPartitionBuilder::GenerateComponent(UPCGComponent* InCompone
 	ensure(!InComponent->IsManagedByRuntimeGenSystem());
 
 	// Non-partitioned components get their editing mode updated here to make sure generation is persisted
-	FPCGWorldPartitionBuilder::UpdateEditingMode(InComponent);
+	// @todo_pcg: revisit for 5.8:
+	// Force update here to make sure ChangeTransientState gets called even if the editing mode might have been updated by previous AddGraphInstance calls.
+	FPCGWorldPartitionBuilder::UpdateEditingMode(InComponent, /*bForce=*/true);
 
 	// Force generate as components that are already generated may decline the request.
 	FPCGTaskId GenerateTask = InvalidPCGTaskId;
@@ -1642,9 +1644,9 @@ void FPCGWorldPartitionBuilder::SetChainedDispatchToLocalComponents(UWorld* InWo
 	PCGSubsystem->SetChainedDispatchToLocalComponents(bInChainedDispatch);
 }
 
-void FPCGWorldPartitionBuilder::UpdateEditingMode(UPCGComponent* InComponent)
+void FPCGWorldPartitionBuilder::UpdateEditingMode(UPCGComponent* InComponent, bool bForce)
 {
-	if (InComponent->GetSerializedEditingMode() == EPCGEditorDirtyMode::LoadAsPreview && InComponent->GetEditingMode() != EPCGEditorDirtyMode::LoadAsPreview)
+	if (InComponent->GetSerializedEditingMode() == EPCGEditorDirtyMode::LoadAsPreview && (bForce || InComponent->GetEditingMode() != EPCGEditorDirtyMode::LoadAsPreview))
 	{
 		UE_LOG(LogPCGWorldPartitionBuilder, Display, TEXT("Setting PCG editing mode to Load As Preview on actor '%s' label '%s' graph '%s'."),
 			*InComponent->GetOwner()->GetName(),
