[PCG] Clean spline was not triggering a graph update when parameters changed.

diff --git a/Engine/Plugins/PCG/Source/PCG/Private/Elements/PCGCleanSpline.cpp b/Engine/Plugins/PCG/Source/PCG/Private/Elements/PCGCleanSpline.cpp
--- a/Engine/Plugins/PCG/Source/PCG/Private/Elements/PCGCleanSpline.cpp
+++ b/Engine/Plugins/PCG/Source/PCG/Private/Elements/PCGCleanSpline.cpp
@@ -62,6 +62,8 @@ void UPCGCleanSplineSettings::PostEditChangeProperty(FPropertyChangedEvent& Prop
 			CollinearAngleThreshold = FMath::Clamp(CollinearAngleThreshold, 0, bUseRadians ? FMath::DegreesToRadians(MaxCollinearAngleToleranceDegrees) : MaxCollinearAngleToleranceDegrees);
 		}
 	}
+
+	Super::PostEditChangeProperty(PropertyChangedEvent);
 }
 #endif // WITH_EDITOR
 
