[PCG] Fix MinPointSpacing in Spline tool mode

* Value was cached and never updated.

diff --git a/Engine/Plugins/PCG/Source/PCGEditor/Private/EditorMode/Tools/Line/PCGDrawSplineTool.cpp b/Engine/Plugins/PCG/Source/PCGEditor/Private/EditorMode/Tools/Line/PCGDrawSplineTool.cpp
--- a/Engine/Plugins/PCG/Source/PCGEditor/Private/EditorMode/Tools/Line/PCGDrawSplineTool.cpp
+++ b/Engine/Plugins/PCG/Source/PCGEditor/Private/EditorMode/Tools/Line/PCGDrawSplineTool.cpp
@@ -310,7 +310,6 @@ void UPCGDrawSplineToolBase::Setup()
 	Super::Setup();
 
 	UPCGInteractiveToolSettings_SplineBase* ToolSettings = GetSplineSettings();
-	CachedMinPointSpacingSquared = ToolSettings->MinPointSpacing * ToolSettings->MinPointSpacing;
 	
 	ClickOrDragBehavior = NewObject<USingleClickOrDragInputBehavior>();
 	ClickOrDragBehavior->Initialize(this, this);
@@ -609,7 +608,7 @@ void UPCGDrawSplineToolBase::OnClickDrag(const FInputDeviceRay& DragPosition)
 					WorkingSpline->SetLocationAtSplinePoint(FinalSplinePoint, FinalLocation, ESplineCoordinateSpace::World, /*bUpdateSpline=*/false);
 					WorkingSpline->SetUpVectorAtSplinePoint(FinalSplinePoint, Normal, ESplineCoordinateSpace::World, /*bUpdateSpline=*/true);
 					const FVector PreviousPoint = WorkingSpline->GetLocationAtSplinePoint(FinalSplinePoint - 1, ESplineCoordinateSpace::World);
-					if (FVector::DistSquared(FinalLocation, PreviousPoint) >= CachedMinPointSpacingSquared)
+					if (FVector::DistSquared(FinalLocation, PreviousPoint) >= (ToolSettings->MinPointSpacing * ToolSettings->MinPointSpacing))
 					{
 						AddSplinePoint(Result->ImpactPosition, Result->ImpactNormal);
 					}
diff --git a/Engine/Plugins/PCG/Source/PCGEditor/Private/EditorMode/Tools/Line/PCGDrawSplineTool.h b/Engine/Plugins/PCG/Source/PCGEditor/Private/EditorMode/Tools/Line/PCGDrawSplineTool.h
--- a/Engine/Plugins/PCG/Source/PCGEditor/Private/EditorMode/Tools/Line/PCGDrawSplineTool.h
+++ b/Engine/Plugins/PCG/Source/PCGEditor/Private/EditorMode/Tools/Line/PCGDrawSplineTool.h
@@ -141,9 +141,6 @@ private:
 	int32 FreeDrawStrokeStartIndex = 0;
 
 	EPCGDrawSplineDrawMode PreviousDrawMode = EPCGDrawSplineDrawMode::None;
-
-protected:
-	double CachedMinPointSpacingSquared = 0;
 };
 
 UCLASS(Experimental)
