[PCG] Made the get/set soft object path/soft class path resilient to both types of storage in the property bag.

#Synced-CL 50352752

diff --git a/Engine/Plugins/PCG/Source/PCG/Private/Helpers/PCGGraphParameterExtension.cpp b/Engine/Plugins/PCG/Source/PCG/Private/Helpers/PCGGraphParameterExtension.cpp
--- a/Engine/Plugins/PCG/Source/PCG/Private/Helpers/PCGGraphParameterExtension.cpp
+++ b/Engine/Plugins/PCG/Source/PCG/Private/Helpers/PCGGraphParameterExtension.cpp
@@ -41,14 +41,36 @@ template<> TValueOrError<FString, EPropertyBagResult> FPCGGraphParameterExtensio
 }
 template<> TValueOrError<FSoftObjectPath, EPropertyBagResult> FPCGGraphParameterExtension::GetGraphParameter(const FInstancedPropertyBag& PropertyBag, const FName PropertyName)
 {
-	return PropertyBag.GetValueSoftPath(PropertyName);
+	TValueOrError<FSoftObjectPath, EPropertyBagResult> ValueOrError = PropertyBag.GetValueSoftPath(PropertyName);
+
+	// Fallback to struct approach
+	if (ValueOrError.HasError())
+	{
+		TValueOrError<FSoftObjectPath*, EPropertyBagResult> StructPtrOrError = PropertyBag.GetValueStruct<FSoftObjectPath>(PropertyName);
+		if (!StructPtrOrError.HasError() && StructPtrOrError.GetValue())
+		{
+			return MakeValue(*StructPtrOrError.GetValue());
+		}
+	}
+
+	return ValueOrError;
 }
 template<> TValueOrError<FSoftClassPath, EPropertyBagResult> FPCGGraphParameterExtension::GetGraphParameter(const FInstancedPropertyBag& PropertyBag, const FName PropertyName)
 {
 	TValueOrError<FSoftObjectPath, EPropertyBagResult> InitialResult = PropertyBag.GetValueSoftPath(PropertyName);
+
+	// Fallback to struct approach
 	if (InitialResult.HasError())
 	{
-		return MakeError(InitialResult.GetError());
+		TValueOrError<FSoftClassPath*, EPropertyBagResult> StructPtrOrError = PropertyBag.GetValueStruct<FSoftClassPath>(PropertyName);
+		if (!StructPtrOrError.HasError() && StructPtrOrError.GetValue())
+		{
+			return MakeValue(*StructPtrOrError.GetValue());
+		}
+		else
+		{
+			return MakeError(InitialResult.GetError());
+		}
 	}
 	else
 	{
@@ -157,11 +179,23 @@ template<> EPropertyBagResult FPCGGraphParameterExtension::SetGraphParameter(FIn
 }
 template<> EPropertyBagResult FPCGGraphParameterExtension::SetGraphParameter(FInstancedPropertyBag& PropertyBag, const FName PropertyName, const FSoftObjectPath& Value)
 {
-	return PropertyBag.SetValueSoftPath(PropertyName, Value);
+	EPropertyBagResult OpResult = PropertyBag.SetValueSoftPath(PropertyName, Value);
+	if (OpResult == EPropertyBagResult::TypeMismatch)
+	{
+		OpResult = PropertyBag.SetValueStruct(PropertyName, Value);
+	}
+
+	return OpResult;
 }
 template<> EPropertyBagResult FPCGGraphParameterExtension::SetGraphParameter(FInstancedPropertyBag& PropertyBag, const FName PropertyName, const FSoftClassPath& Value)
 {
-	return PropertyBag.SetValueSoftPath(PropertyName, Value);
+	EPropertyBagResult OpResult = PropertyBag.SetValueSoftPath(PropertyName, Value);
+	if (OpResult == EPropertyBagResult::TypeMismatch)
+	{
+		OpResult = PropertyBag.SetValueStruct(PropertyName, Value);
+	}
+
+	return OpResult;
 }
 template<> EPropertyBagResult FPCGGraphParameterExtension::SetGraphParameter(FInstancedPropertyBag& PropertyBag, const FName PropertyName, const TSoftObjectPtr<UObject>& Value)
 {
