[PCG] Fix PCG Builder that keep the graph grid size array growing.

#rnx

diff --git a/Engine/Plugins/PCG/Source/PCG/Private/PCGGraph.cpp b/Engine/Plugins/PCG/Source/PCG/Private/PCGGraph.cpp
--- a/Engine/Plugins/PCG/Source/PCG/Private/PCGGraph.cpp
+++ b/Engine/Plugins/PCG/Source/PCG/Private/PCGGraph.cpp
@@ -1473,7 +1473,7 @@ void UPCGGraph::CacheGridSizesInternalNoLock() const
 	{
 		if (PCGHiGenGrid::IsValidGridSize(GetDefaultGridSize()))
 		{
-			GridInfo.GridSizes.Add(GraphDefaultGridSize);
+			GridInfo.GridSizes.AddUnique(GraphDefaultGridSize);
 		}
 	}
 	else
@@ -1484,10 +1484,7 @@ void UPCGGraph::CacheGridSizesInternalNoLock() const
 			const uint32 GridSize = GetNodeGenerationGridSize(Node, GraphDefaultGridSize);
 			if (PCGHiGenGrid::IsValidGridSize(GridSize))
 			{
-				if (!GridInfo.GridSizes.Contains(GridSize))
-				{
-					GridInfo.GridSizes.Add(GridSize);
-				}
+				GridInfo.GridSizes.AddUnique(GridSize);
 			}
 			else if (GridSize == PCGHiGenGrid::UnboundedGridSize())
 			{
@@ -1503,7 +1500,7 @@ void UPCGGraph::CacheGridSizesInternalNoLock() const
 		if (bHasUninitialized)
 		{
 			// Nodes outside grid ranges will execute at graph default
-			GridInfo.GridSizes.Add(GraphDefaultGridSize);
+			GridInfo.GridSizes.AddUnique(GraphDefaultGridSize);
 		}
 
 		// Descending
diff --git a/Engine/Plugins/PCG/Source/PCGEditor/Private/WorldPartitionBuilder/PCGWorldPartitionBuilder.cpp b/Engine/Plugins/PCG/Source/PCGEditor/Private/WorldPartitionBuilder/PCGWorldPartitionBuilder.cpp
--- a/Engine/Plugins/PCG/Source/PCGEditor/Private/WorldPartitionBuilder/PCGWorldPartitionBuilder.cpp
+++ b/Engine/Plugins/PCG/Source/PCGEditor/Private/WorldPartitionBuilder/PCGWorldPartitionBuilder.cpp
@@ -1386,13 +1386,15 @@ void FPCGWorldPartitionBuilder::GetGenerationGridSizes(const UPCGGraph* InGraph,
 		OutGridSizes = FoundGridSizeInfo->GridSizes;
 		bOutHasUnbounded = FoundGridSizeInfo->bHasUnbounded;
 	}
-		
-	check(PCGHelpers::GetGenerationGridSizes(InGraph, InWorldActor, OutGridSizes, bOutHasUnbounded));
+	else
+	{
+		check(PCGHelpers::GetGenerationGridSizes(InGraph, InWorldActor, OutGridSizes, bOutHasUnbounded));
 
-	// Store
-	FGridSizeInfo& NewGridSizeInfo = GraphToGridSizeInfo.Add(InGraph);
-	NewGridSizeInfo.bHasUnbounded = bOutHasUnbounded;
-	NewGridSizeInfo.GridSizes = OutGridSizes;
+		// Store
+		FGridSizeInfo& NewGridSizeInfo = GraphToGridSizeInfo.Add(InGraph);
+		NewGridSizeInfo.bHasUnbounded = bOutHasUnbounded;
+		NewGridSizeInfo.GridSizes = OutGridSizes;
+	}
 }
 
 void FPCGWorldPartitionBuilder::WaitOnGraphExecutor(UWorld* InWorld)
