[PCG] Fix issue where Make Rotator Attribute was not working with inline constants on the Make Rot from Euler Angle op.

[RN] Make Rotator Attribute is now working correctly for the Make Rot from Euler Angles operation with Inline Constants

diff --git a/Engine/Plugins/PCG/Source/PCG/Private/Elements/Metadata/PCGMetadataMakeRotator.cpp b/Engine/Plugins/PCG/Source/PCG/Private/Elements/Metadata/PCGMetadataMakeRotator.cpp
--- a/Engine/Plugins/PCG/Source/PCG/Private/Elements/Metadata/PCGMetadataMakeRotator.cpp
+++ b/Engine/Plugins/PCG/Source/PCG/Private/Elements/Metadata/PCGMetadataMakeRotator.cpp
@@ -216,8 +216,15 @@ FString UPCGMetadataMakeRotatorSettings::GetPinInitialDefaultValueString(FName P
 
 bool UPCGMetadataMakeRotatorSettings::CreateInitialDefaultValueAttribute(const FName PinLabel, UPCGMetadata* OutMetadata) const
 {
-	const FVector Value = PCGMetadataMakeRotator::Helpers::GetAxisAlignedVector(PinLabel);
-	return nullptr != OutMetadata->CreateAttribute(NAME_None, Value, /*bAllowsInterpolation=*/true, /*bOverrideParent=*/false);
+	if (Operation != EPCGMetadataMakeRotatorOp::MakeRotFromAngles)
+	{
+		const FVector Value = PCGMetadataMakeRotator::Helpers::GetAxisAlignedVector(PinLabel);
+		return nullptr != OutMetadata->CreateAttribute(NAME_None, Value, /*bAllowsInterpolation=*/true, /*bOverrideParent=*/false);
+	}
+	else
+	{
+		return Super::CreateInitialDefaultValueAttribute(PinLabel, OutMetadata);
+	}
 }
 
 EPCGMetadataTypes UPCGMetadataMakeRotatorSettings::GetPinInitialDefaultValueType(FName PinLabel) const
