[PCG] Fix issue related to transient flags in PCG Editor Tool mode

* When using the painting tool without any actor selected, if the graph was changed, the box component on the spawned actor was not saved and lost on map reload
* When using volumes without any actor selected, applying the tool was not dirtying the map, and the volume was lost on reload

[RN] Fix issue in PCG Editor Tool mode where actors/components spawned by the tool would be lost on save/reload.

diff --git a/Engine/Plugins/PCG/Source/PCGEditor/Private/EditorMode/Tools/PCGInteractiveToolSettings.cpp b/Engine/Plugins/PCG/Source/PCGEditor/Private/EditorMode/Tools/PCGInteractiveToolSettings.cpp
--- a/Engine/Plugins/PCG/Source/PCGEditor/Private/EditorMode/Tools/PCGInteractiveToolSettings.cpp
+++ b/Engine/Plugins/PCG/Source/PCGEditor/Private/EditorMode/Tools/PCGInteractiveToolSettings.cpp
@@ -192,6 +192,31 @@ void UPCGInteractiveToolSettings::Apply(UInteractiveTool* OwningTool)
 	if (GeneratedResources.SpawnedActor.IsValid())
 	{
 		GeneratedResources.SpawnedActor->ClearFlags(RF_Transient);
+
+		// @todo_pcg: Revisit in 5.8
+		// When we create a tool working data, the data might create some generated resources (like a component), and it is
+		// likely it is created as transient. Clearing the flag is done in the `OnToolApply` just below.
+		// The problem is when we change graphs that don't have the same id for the data, we create "new" working data that re-use the newly spawned root component.
+		// But the "old" working data of the first graph is still there, with its generated component still marked as transient. Since its id might be different, it is not
+		// gathered by the `GetMutableTypedWorkingDataMap` and thus the transient flag is never cleared.
+		// We definitely have a flow issue there.
+		// Probably that the tool data should indicate to the tool that it needs a given component to work, and if it doesn't exist, tell the tool
+		// to create it, and it is the tool that keeps the reference and is responsible to clear the transient flag.
+		// In the mean time for 5.7.4, we will just look all outered objects to the actor and clear their flags. The Root component is assumed to be
+		// outered to the actor.
+		const UPCGComponent* WorkingPCGComponent = GetWorkingPCGComponent();
+		ForEachObjectWithOuter(GeneratedResources.SpawnedActor.Get(),
+			[WorkingPCGComponent](UObject* Obj) 
+			{
+				// To be sure we do not override what PCG is doing, we will ignore any object that was generated by PCG
+				const UObject* ConstObj = Obj;
+				if (WorkingPCGComponent && WorkingPCGComponent->IsAnyObjectManagedByResource(MakeArrayView(&ConstObj, 1)))
+				{
+					return;
+				}
+
+				Obj->ClearFlags(RF_Transient);
+			}, /*bIncludeNestedObjects=*/true);
 	}
 	
 	if (GeneratedResources.GeneratedPCGComponent.IsValid())
@@ -254,7 +279,15 @@ void UPCGInteractiveToolSettings::Shutdown()
 		{
 			PCGSubsystem->FlushCache();
 		}
-		
+
+		// @todo_pcg: Revisit in 5.8
+		// If the working actor was spawned for the tool, but this actor has a PCG component by default (like PCGVolume)
+		// the default subobject component will also be flagged transient. In our case, the PCG Component will not be spawned by the tool
+		// but reused from the working actor, and thus will never have its transient flag cleared. We need to make sure that we keep enough information
+		// when spawning actors to clear all the transient flag.
+		// For 5.7.4, just clear the flag here before the generation.
+		Component->ClearFlags(RF_Transient);
+
 		Component->Generate(true);
 	}
 }
