[PCG] Fixed assert/bad behavior when trying to filter empty point data on a data-level attribute.

diff --git a/Engine/Plugins/PCG/Source/PCG/Private/Elements/PCGAttributeFilter.cpp b/Engine/Plugins/PCG/Source/PCG/Private/Elements/PCGAttributeFilter.cpp
--- a/Engine/Plugins/PCG/Source/PCG/Private/Elements/PCGAttributeFilter.cpp
+++ b/Engine/Plugins/PCG/Source/PCG/Private/Elements/PCGAttributeFilter.cpp
@@ -632,31 +632,35 @@ bool FPCGAttributeFilterElementBase::DoFiltering(FPCGContext* Context, EPCGAttri
 		}
 
 		OperationData.OriginalData = OriginalData;
+		int32 NumberOfEntries = 0;
 
 		if (OperationData.bIsInputPointData)
 		{
 			const UPCGBasePointData* OriginalPointData = CastChecked<UPCGBasePointData>(OriginalData);
 
 			OperationData.OriginalPointData = OriginalPointData;
-
-			// Will be set individually in batches
-			OperationData.FilterBitArray.SetNumUninitialized(OriginalPointData->GetNumPoints());
+			NumberOfEntries = OriginalPointData->GetNumPoints();
 		}
 		else
 		{
 			// Param data
 			const UPCGParamData* OriginalParamData = CastChecked<UPCGParamData>(OriginalData);
+			NumberOfEntries = OriginalParamData->Metadata ? OriginalParamData->Metadata->GetItemCountForChild() : 0;
+		}
 
-			// Will be set individually in batches
-			OperationData.FilterBitArray.SetNumUninitialized(OriginalParamData->Metadata ? OriginalParamData->Metadata->GetItemCountForChild() : 0);
+		if (NumberOfEntries == 0)
+		{
+			ForwardInputToInFilterPin();
+			continue;
 		}
 
-		auto Operation = [&Operator, &FirstThresholdInfo, &SecondThresholdInfo, &TargetAccessor, &TargetKeys, &FirstThreshold, &SecondThreshold, &OperationData](auto Dummy) -> bool
+		// Will be set individually in batches
+		OperationData.FilterBitArray.SetNumUninitialized(NumberOfEntries);
+
+		auto Operation = [&Operator, &FirstThresholdInfo, &SecondThresholdInfo, &TargetAccessor, &TargetKeys, &NumberOfEntries, &FirstThreshold, &SecondThreshold, &OperationData](auto Dummy) -> bool
 		{
 			using Type = decltype(Dummy);
 
-			const int32 NumberOfEntries = TargetKeys->GetNum();
-
 			if (NumberOfEntries <= 0)
 			{
 				return false;
