PCG fixes to not have const FNames in headers.

(const FNames in headers get duplicated into every TU that includes the header.)

diff --git a/Engine/Plugins/PCG/Source/PCG/Private/Data/PCGBasePointData.cpp b/Engine/Plugins/PCG/Source/PCG/Private/Data/PCGBasePointData.cpp
--- a/Engine/Plugins/PCG/Source/PCG/Private/Data/PCGBasePointData.cpp
+++ b/Engine/Plugins/PCG/Source/PCG/Private/Data/PCGBasePointData.cpp
@@ -27,6 +27,12 @@ static TAutoConsoleVariable<bool> CVarCacheFullPointDataCrc(
 	true,
 	TEXT("Enable fine-grained CRC of point data for change tracking on elements that request it, rather than using data UID."));
 
+namespace PCGPointDataConstants
+{
+	const FName ActorReferenceAttribute = TEXT("ActorReference");
+	const FName ElementsDomainName = "Points";
+}
+
 PCG_DEFINE_TYPE_INFO(FPCGDataTypeInfoPoint, UPCGBasePointData)
 
 namespace PCGPointHelpers
diff --git a/Engine/Plugins/PCG/Source/PCG/Private/Metadata/PCGMetadataAttribute.cpp b/Engine/Plugins/PCG/Source/PCG/Private/Metadata/PCGMetadataAttribute.cpp
--- a/Engine/Plugins/PCG/Source/PCG/Private/Metadata/PCGMetadataAttribute.cpp
+++ b/Engine/Plugins/PCG/Source/PCG/Private/Metadata/PCGMetadataAttribute.cpp
@@ -3,6 +3,14 @@
 #include "Metadata/PCGMetadataAttribute.h"
 #include "Metadata/PCGMetadata.h"
 
+namespace PCGMetadataAttributeConstants
+{
+	const FName LastAttributeName = TEXT("@Last");
+	const FName LastCreatedAttributeName = TEXT("@LastCreated");
+	const FName SourceAttributeName = TEXT("@Source");
+	const FName SourceNameAttributeName = TEXT("@SourceName");
+}
+
 namespace PCGMetadataAttributeBase
 {
 	static constexpr TCHAR AllowedSpecialCharacters[4] = {' ', '_', '-', '/'};
diff --git a/Engine/Plugins/PCG/Source/PCG/Private/Metadata/PCGMetadataCommon.cpp b/Engine/Plugins/PCG/Source/PCG/Private/Metadata/PCGMetadataCommon.cpp
new file mode 100644
--- /dev/null
+++ b/Engine/Plugins/PCG/Source/PCG/Private/Metadata/PCGMetadataCommon.cpp
@@ -0,0 +1,11 @@
+// Copyright Epic Games, Inc. All Rights Reserved.
+
+#include "Metadata/PCGMetadataCommon.h"
+
+namespace PCGMetadataDomainID
+{
+	const FPCGMetadataDomainID Default{EPCGMetadataDomainFlag::Default, -1, TEXT("Default")};
+	const FPCGMetadataDomainID Elements{EPCGMetadataDomainFlag::Elements, -1, TEXT("Elements")};
+	const FPCGMetadataDomainID Data{EPCGMetadataDomainFlag::Data, -1, TEXT("Data")};
+	const FPCGMetadataDomainID Invalid{EPCGMetadataDomainFlag::Invalid, -1, TEXT("Invalid")};
+}
diff --git a/Engine/Plugins/PCG/Source/PCG/Private/PCGCommon.cpp b/Engine/Plugins/PCG/Source/PCG/Private/PCGCommon.cpp
--- a/Engine/Plugins/PCG/Source/PCG/Private/PCGCommon.cpp
+++ b/Engine/Plugins/PCG/Source/PCG/Private/PCGCommon.cpp
@@ -278,10 +278,78 @@ namespace PCGPointCustomPropertyNames
 			Name == PCGPointCustomPropertyNames::LocalSizeName ||
 			Name == PCGPointCustomPropertyNames::ScaledLocalSizeName;
 	}
+
+	const FName ExtentsName = TEXT("Extents");
+	const FName LocalCenterName = TEXT("LocalCenter");
+	const FName PositionName = TEXT("Position");
+	const FName RotationName = TEXT("Rotation");
+	const FName ScaleName = TEXT("Scale");
+	const FName LocalSizeName = TEXT("LocalSize");
+	const FName ScaledLocalSizeName = TEXT("ScaledLocalSize");
+}
+
+namespace PCGPinConstants
+{
+	const FName DefaultInputLabel = TEXT("In");
+	const FName DefaultOutputLabel = TEXT("Out");
+	const FName DefaultParamsLabel = TEXT("Overrides");
+	UE_DEPRECATED(5.6, "Please use 'DefaultExecutionDependencyLabel' instead.")
+	const FName DefaultDependencyOnlyLabel = TEXT("Dependency Only");
+	const FName DefaultExecutionDependencyLabel = TEXT("Execution Dependency");
+
+	const FName DefaultInFilterLabel = TEXT("InsideFilter");
+	const FName DefaultOutFilterLabel = TEXT("OutsideFilter");
+
+	namespace Private
+	{
+		const FName OldDefaultParamsLabel = TEXT("Params");
+	}
+
+	namespace Icons
+	{
+		const FName LoopPinIcon = TEXT("GraphEditor.Macro.Loop_16x");
+		const FName FeedbackPinIcon = TEXT("GraphEditor.GetSequenceBinding");
+	}
+
+#if WITH_EDITOR
+	namespace Tooltips
+	{
+		const FText ExecutionDependencyTooltip = NSLOCTEXT("PCGCommon", "ExecutionDependencyPinTooltip", "Data passed to this pin will be used to order execution but will otherwise not contribute to the results of this node.");
+	}
+#endif // WITH_EDITOR
+} // namespace PCGPinConstants
+
+namespace PCGNodeConstants::Icons
+{
+	const FName CompactNodeConvert = TEXT("PCG.Node.Compact.Convert");
+	const FName CompactNodeFilter = TEXT("PCG.Node.Compact.Filter");
+}
+
+// Metadata used by PCG
+namespace PCGObjectMetadata
+{
+	const FLazyName Overridable = TEXT("PCG_Overridable");
+	const FLazyName OverridableCPUAndGPU = TEXT("PCG_OverridableCPUAndGPU");
+	const FLazyName OverridableCPUAndGPUWithReadback = TEXT("PCG_OverridableCPUAndGPUWithReadback");
+	const FLazyName NotOverridable = TEXT("PCG_NotOverridable");
+	const FLazyName OverridableChildProperties = TEXT("PCG_OverridableChildProperties");
+	const FLazyName OverrideAliases = TEXT("PCG_OverrideAliases");
+	const FLazyName DiscardPropertySelection = TEXT("PCG_DiscardPropertySelection");
+	const FLazyName DiscardExtraSelection = TEXT("PCG_DiscardExtraSelection");
+	const FLazyName EnumMetadataDomain = TEXT("PCG_MetadataDomain");
+	const FLazyName PropertyReadOnly = TEXT("PCG_PropertyReadOnly");
+	const FLazyName DataTypeIdentifierSupportsComposition = TEXT("PCG_SupportsComposition");
 }
 
 namespace PCGQualityHelpers
 {
+	const FName PinLabelDefault = TEXT("Default");
+	const FName PinLabelLow = TEXT("Low");
+	const FName PinLabelMedium = TEXT("Medium");
+	const FName PinLabelHigh = TEXT("High");
+	const FName PinLabelEpic = TEXT("Epic");
+	const FName PinLabelCinematic = TEXT("Cinematic");
+
 	FName GetQualityPinLabel()
 	{
 		const int32 QualityLevel = UPCGSubsystem::GetPCGQualityLevel();
@@ -327,3 +395,4 @@ const FPCGStack* FPCGStackHandle::GetStack() const
 {
 	return IsValid() ? StackContext->GetStack(StackIndex) : nullptr;
 }
+
diff --git a/Engine/Plugins/PCG/Source/PCG/Private/PCGData.cpp b/Engine/Plugins/PCG/Source/PCG/Private/PCGData.cpp
--- a/Engine/Plugins/PCG/Source/PCG/Private/PCGData.cpp
+++ b/Engine/Plugins/PCG/Source/PCG/Private/PCGData.cpp
@@ -18,6 +18,12 @@
 
 #include UE_INLINE_GENERATED_CPP_BY_NAME(PCGData)
 
+namespace PCGDataConstants
+{
+	const FName DataDomainName = "Data";
+	const FName DefaultDomainName = "Default";
+}
+
 //#define PCG_DATA_USAGE_LOGGING
 
 static TAutoConsoleVariable<bool> CVarPCGFullDataCrcInCommandlets(
diff --git a/Engine/Plugins/PCG/Source/PCG/Public/Data/PCGBasePointData.h b/Engine/Plugins/PCG/Source/PCG/Public/Data/PCGBasePointData.h
--- a/Engine/Plugins/PCG/Source/PCG/Public/Data/PCGBasePointData.h
+++ b/Engine/Plugins/PCG/Source/PCG/Public/Data/PCGBasePointData.h
@@ -14,8 +14,8 @@
 
 namespace PCGPointDataConstants
 {
-	const FName ActorReferenceAttribute = TEXT("ActorReference");
-	const FName ElementsDomainName = "Points";
+	extern PCG_API const FName ActorReferenceAttribute;
+	extern PCG_API const FName ElementsDomainName;
 }
 
 USTRUCT(BlueprintType)
diff --git a/Engine/Plugins/PCG/Source/PCG/Public/Metadata/PCGMetadataAttribute.h b/Engine/Plugins/PCG/Source/PCG/Public/Metadata/PCGMetadataAttribute.h
--- a/Engine/Plugins/PCG/Source/PCG/Public/Metadata/PCGMetadataAttribute.h
+++ b/Engine/Plugins/PCG/Source/PCG/Public/Metadata/PCGMetadataAttribute.h
@@ -13,10 +13,10 @@ class UPCGMetadata;
 
 namespace PCGMetadataAttributeConstants
 {
-	const FName LastAttributeName = TEXT("@Last");
-	const FName LastCreatedAttributeName = TEXT("@LastCreated");
-	const FName SourceAttributeName = TEXT("@Source");
-	const FName SourceNameAttributeName = TEXT("@SourceName");
+	extern PCG_API const FName LastAttributeName;
+	extern PCG_API const FName LastCreatedAttributeName;
+	extern PCG_API const FName SourceAttributeName;
+	extern PCG_API const FName SourceNameAttributeName;
 }
 
 class FPCGMetadataAttributeBase
diff --git a/Engine/Plugins/PCG/Source/PCG/Public/Metadata/PCGMetadataCommon.h b/Engine/Plugins/PCG/Source/PCG/Public/Metadata/PCGMetadataCommon.h
--- a/Engine/Plugins/PCG/Source/PCG/Public/Metadata/PCGMetadataCommon.h
+++ b/Engine/Plugins/PCG/Source/PCG/Public/Metadata/PCGMetadataCommon.h
@@ -115,10 +115,10 @@ struct FPCGMetadataDomainID
 
 namespace PCGMetadataDomainID
 {
-	static const FPCGMetadataDomainID Default{EPCGMetadataDomainFlag::Default, -1, TEXT("Default")};
-	static const FPCGMetadataDomainID Elements{EPCGMetadataDomainFlag::Elements, -1, TEXT("Elements")};
-	static const FPCGMetadataDomainID Data{EPCGMetadataDomainFlag::Data, -1, TEXT("Data")};
-	static const FPCGMetadataDomainID Invalid{EPCGMetadataDomainFlag::Invalid, -1, TEXT("Invalid")};
+	extern PCG_API const FPCGMetadataDomainID Default;
+	extern PCG_API const FPCGMetadataDomainID Elements;
+	extern PCG_API const FPCGMetadataDomainID Data;
+	extern PCG_API const FPCGMetadataDomainID Invalid;
 }
 
 USTRUCT(BlueprintType)
diff --git a/Engine/Plugins/PCG/Source/PCG/Public/PCGCommon.h b/Engine/Plugins/PCG/Source/PCG/Public/PCGCommon.h
--- a/Engine/Plugins/PCG/Source/PCG/Public/PCGCommon.h
+++ b/Engine/Plugins/PCG/Source/PCG/Public/PCGCommon.h
@@ -57,13 +57,13 @@ namespace PCGPointCustomPropertyNames
 {
 	bool IsCustomPropertyName(FName Name);
 
-	const FName ExtentsName = TEXT("Extents");
-	const FName LocalCenterName = TEXT("LocalCenter");
-	const FName PositionName = TEXT("Position");
-	const FName RotationName = TEXT("Rotation");
-	const FName ScaleName = TEXT("Scale");
-	const FName LocalSizeName = TEXT("LocalSize");
-	const FName ScaledLocalSizeName = TEXT("ScaledLocalSize");
+	extern PCG_API const FName ExtentsName;
+	extern PCG_API const FName LocalCenterName;
+	extern PCG_API const FName PositionName;
+	extern PCG_API const FName RotationName;
+	extern PCG_API const FName ScaleName;
+	extern PCG_API const FName LocalSizeName;
+	extern PCG_API const FName ScaledLocalSizeName;
 }
 
 UENUM(meta = (Bitflags))
@@ -197,31 +197,31 @@ namespace PCGValueConstants
 
 namespace PCGPinConstants
 {
-	const FName DefaultInputLabel = TEXT("In");
-	const FName DefaultOutputLabel = TEXT("Out");
-	const FName DefaultParamsLabel = TEXT("Overrides");
+	extern PCG_API const FName DefaultInputLabel;
+	extern PCG_API const FName DefaultOutputLabel;
+	extern PCG_API const FName DefaultParamsLabel;
 	UE_DEPRECATED(5.6, "Please use 'DefaultExecutionDependencyLabel' instead.")
-	const FName DefaultDependencyOnlyLabel = TEXT("Dependency Only");
-	const FName DefaultExecutionDependencyLabel = TEXT("Execution Dependency");
+	extern PCG_API const FName DefaultDependencyOnlyLabel;
+	extern PCG_API const FName DefaultExecutionDependencyLabel;
 
-	const FName DefaultInFilterLabel = TEXT("InsideFilter");
-	const FName DefaultOutFilterLabel = TEXT("OutsideFilter");
+	extern PCG_API const FName DefaultInFilterLabel;
+	extern PCG_API const FName DefaultOutFilterLabel;
 
 namespace Private
 {
-	const FName OldDefaultParamsLabel = TEXT("Params");
+	extern PCG_API const FName OldDefaultParamsLabel;
 }
 
 namespace Icons
 {
-	const FName LoopPinIcon = TEXT("GraphEditor.Macro.Loop_16x");
-	const FName FeedbackPinIcon = TEXT("GraphEditor.GetSequenceBinding");
+	extern PCG_API const FName LoopPinIcon;
+	extern PCG_API const FName FeedbackPinIcon;
 }
 
 #if WITH_EDITOR
 namespace Tooltips
 {
-	const FText ExecutionDependencyTooltip = NSLOCTEXT("PCGCommon", "ExecutionDependencyPinTooltip", "Data passed to this pin will be used to order execution but will otherwise not contribute to the results of this node.");
+	extern PCG_API const FText ExecutionDependencyTooltip;
 }
 #endif // WITH_EDITOR
 } // namespace PCGPinConstants
@@ -230,25 +230,25 @@ namespace PCGNodeConstants
 {
 namespace Icons
 {
-	const FName CompactNodeConvert = TEXT("PCG.Node.Compact.Convert");
-	const FName CompactNodeFilter = TEXT("PCG.Node.Compact.Filter");
+	extern PCG_API const FName CompactNodeConvert;
+	extern PCG_API const FName CompactNodeFilter;
 }
 }
 
 // Metadata used by PCG
 namespace PCGObjectMetadata
 {
-	const FLazyName Overridable = TEXT("PCG_Overridable");
-	const FLazyName OverridableCPUAndGPU = TEXT("PCG_OverridableCPUAndGPU");
-	const FLazyName OverridableCPUAndGPUWithReadback = TEXT("PCG_OverridableCPUAndGPUWithReadback");
-	const FLazyName NotOverridable = TEXT("PCG_NotOverridable");
-	const FLazyName OverridableChildProperties = TEXT("PCG_OverridableChildProperties");
-	const FLazyName OverrideAliases = TEXT("PCG_OverrideAliases");
-	const FLazyName DiscardPropertySelection = TEXT("PCG_DiscardPropertySelection");
-	const FLazyName DiscardExtraSelection = TEXT("PCG_DiscardExtraSelection");
-	const FLazyName EnumMetadataDomain = TEXT("PCG_MetadataDomain");
-	const FLazyName PropertyReadOnly = TEXT("PCG_PropertyReadOnly");
-	const FLazyName DataTypeIdentifierSupportsComposition = TEXT("PCG_SupportsComposition");
+	extern PCG_API const FLazyName Overridable;
+	extern PCG_API const FLazyName OverridableCPUAndGPU;
+	extern PCG_API const FLazyName OverridableCPUAndGPUWithReadback;
+	extern PCG_API const FLazyName NotOverridable;
+	extern PCG_API const FLazyName OverridableChildProperties;
+	extern PCG_API const FLazyName OverrideAliases;
+	extern PCG_API const FLazyName DiscardPropertySelection;
+	extern PCG_API const FLazyName DiscardExtraSelection;
+	extern PCG_API const FLazyName EnumMetadataDomain;
+	extern PCG_API const FLazyName PropertyReadOnly;
+	extern PCG_API const FLazyName DataTypeIdentifierSupportsComposition;
 
 	// Metadata usable in UPROPERTY for customizing the behavior when displaying the property in a property panel or graph node
 	enum
@@ -499,12 +499,12 @@ enum class EPCGNodeTitleType : uint8
 namespace PCGQualityHelpers
 {
 	constexpr int32 NumPins = 6;
-	const FName PinLabelDefault = TEXT("Default");
-	const FName PinLabelLow = TEXT("Low");
-	const FName PinLabelMedium = TEXT("Medium");
-	const FName PinLabelHigh = TEXT("High");
-	const FName PinLabelEpic = TEXT("Epic");
-	const FName PinLabelCinematic = TEXT("Cinematic");
+	extern PCG_API const FName PinLabelDefault;
+	extern PCG_API const FName PinLabelLow;
+	extern PCG_API const FName PinLabelMedium;
+	extern PCG_API const FName PinLabelHigh;
+	extern PCG_API const FName PinLabelEpic;
+	extern PCG_API const FName PinLabelCinematic;
 
 	/** Get the pin label associated with the current 'pcg.Quality' value. If the quality level is invalid, it will return the default pin label. */
 	PCG_API FName GetQualityPinLabel();
diff --git a/Engine/Plugins/PCG/Source/PCG/Public/PCGData.h b/Engine/Plugins/PCG/Source/PCG/Public/PCGData.h
--- a/Engine/Plugins/PCG/Source/PCG/Public/PCGData.h
+++ b/Engine/Plugins/PCG/Source/PCG/Public/PCGData.h
@@ -38,8 +38,8 @@ ENUM_CLASS_FLAGS(EPCGDataUsage);
 
 namespace PCGDataConstants
 {
-	const FName DataDomainName = "Data";
-	const FName DefaultDomainName = "Default";
+	extern PCG_API const FName DataDomainName;
+	extern PCG_API const FName DefaultDomainName;
 }
 
 /**
