[PCG] Add advanced toggle to disable the grammar will not fit error

diff --git a/Engine/Plugins/PCG/Source/PCG/Private/Elements/Grammar/PCGDuplicateCrossSections.cpp b/Engine/Plugins/PCG/Source/PCG/Private/Elements/Grammar/PCGDuplicateCrossSections.cpp
--- a/Engine/Plugins/PCG/Source/PCG/Private/Elements/Grammar/PCGDuplicateCrossSections.cpp
+++ b/Engine/Plugins/PCG/Source/PCG/Private/Elements/Grammar/PCGDuplicateCrossSections.cpp
@@ -154,7 +154,11 @@ bool FPCGDuplicateCrossSectionsElement::ExecuteInternal(FPCGContext* InContext)
 
 		if (!bHeightSubdivideSuccess)
 		{
-			PCGLog::LogErrorOnGraph(LOCTEXT("SubdivideFailed", "Grammar doesn't fit."), InContext);
+			if (!Settings->bSilenceGrammarDoesNotFitErrors)
+			{
+				PCGLog::LogErrorOnGraph(LOCTEXT("GrammarDoesNotFit", "Grammar does not fit for this segment."), InContext);
+			}
+
 			continue;
 		}
 
diff --git a/Engine/Plugins/PCG/Source/PCG/Private/Elements/Grammar/PCGSubdivideSegment.cpp b/Engine/Plugins/PCG/Source/PCG/Private/Elements/Grammar/PCGSubdivideSegment.cpp
--- a/Engine/Plugins/PCG/Source/PCG/Private/Elements/Grammar/PCGSubdivideSegment.cpp
+++ b/Engine/Plugins/PCG/Source/PCG/Private/Elements/Grammar/PCGSubdivideSegment.cpp
@@ -85,12 +85,17 @@ public:
 
 		if (!bSubdivideSuccess)
 		{
+			if (!InOutParameters.Settings->bSilenceGrammarDoesNotFitErrors)
+			{
+				PCGLog::LogErrorOnGraph(LOCTEXT("GrammarDoesNotFit", "Grammar does not fit for this segment."), InOutParameters.Context);
+			}
+
 			return;
 		}
 
 		if (!InOutParameters.Settings->bAcceptIncompleteSubdivision && !FMath::IsNearlyZero(RemainingSubdivide))
 		{
-			PCGLog::LogWarningOnGraph(LOCTEXT("FailSubdivisionFullLength", "One segment has an incomplete subdivision (grammar doesn't fit the whole segment)."), InOutParameters.Context);
+			PCGLog::LogWarningOnGraph(LOCTEXT("FailSubdivisionFullLength", "One segment has an incomplete subdivision (grammar does not fit the whole segment)."), InOutParameters.Context);
 			return;
 		}
 
diff --git a/Engine/Plugins/PCG/Source/PCG/Private/Elements/Grammar/PCGSubdivideSpline.cpp b/Engine/Plugins/PCG/Source/PCG/Private/Elements/Grammar/PCGSubdivideSpline.cpp
--- a/Engine/Plugins/PCG/Source/PCG/Private/Elements/Grammar/PCGSubdivideSpline.cpp
+++ b/Engine/Plugins/PCG/Source/PCG/Private/Elements/Grammar/PCGSubdivideSpline.cpp
@@ -19,6 +19,7 @@ namespace PCGSubdivideSplineHelpers
 	struct FParameters
 	{
 		FPCGContext* Context = nullptr;
+		const UPCGSubdivideSplineSettings* Settings = nullptr;
 
 		PCGSubdivisionBase::FModuleInfoMap ModulesInfo;
 		TMap<FString, PCGGrammar::FTokenizedGrammar> CachedModules;
@@ -110,6 +111,11 @@ namespace PCGSubdivideSplineHelpers
 		double RemainingSubdivide;
 		if (!Subdivide(*CurrentTokenizedGrammar.ModuleGrammar, SplineLength, ModulesInstances, RemainingSubdivide, InOutParameters.Context, InOutParameters.AdditionalSeed))
 		{
+			if (!InOutParameters.Settings->bSilenceGrammarDoesNotFitErrors)
+			{
+				PCGLog::LogErrorOnGraph(LOCTEXT("GrammarDoesNotFit", "Grammar does not fit for this spline."), InOutParameters.Context);
+			}
+
 			return;
 		}
 
@@ -339,6 +345,7 @@ bool FPCGSubdivideSplineElement::ExecuteInternal(FPCGContext* InContext) const
 	PCGSubdivideSplineHelpers::FParameters Parameters
 	{
 		.Context = InContext,
+		.Settings = Settings,
 		.ModulesInfo = ModulesInfo,
 		.bAcceptIncompleteSubdivision = Settings->bAcceptIncompleteSubdivision,
 		.BisectionTolerance = Settings->ModulePlacementTolerance
diff --git a/Engine/Plugins/PCG/Source/PCG/Public/Elements/Grammar/PCGDuplicateCrossSections.h b/Engine/Plugins/PCG/Source/PCG/Public/Elements/Grammar/PCGDuplicateCrossSections.h
--- a/Engine/Plugins/PCG/Source/PCG/Public/Elements/Grammar/PCGDuplicateCrossSections.h
+++ b/Engine/Plugins/PCG/Source/PCG/Public/Elements/Grammar/PCGDuplicateCrossSections.h
@@ -44,6 +44,10 @@ public:
 	/** Name of the spline index output attribute name. */
 	UPROPERTY(BlueprintReadWrite, EditAnywhere, Category = "Settings|Extra Output Attributes", meta = (PCG_Overridable, EditCondition = "bOutputSplineIndexAttribute"))
 	FName SplineIndexAttributeName = TEXT("SplineIndex");
+
+	/** Do not post node errors when the grammar does not fit the intended cross-sections. */
+	UPROPERTY(BlueprintReadWrite, EditAnywhere, Category = Settings, AdvancedDisplay, meta = (PCG_Overridable))
+	bool bSilenceGrammarDoesNotFitErrors = false;
 };
 
 class FPCGDuplicateCrossSectionsElement : public FPCGSubdivisionBaseElement
diff --git a/Engine/Plugins/PCG/Source/PCG/Public/Elements/Grammar/PCGSubdivideSegment.h b/Engine/Plugins/PCG/Source/PCG/Public/Elements/Grammar/PCGSubdivideSegment.h
--- a/Engine/Plugins/PCG/Source/PCG/Public/Elements/Grammar/PCGSubdivideSegment.h
+++ b/Engine/Plugins/PCG/Source/PCG/Public/Elements/Grammar/PCGSubdivideSegment.h
@@ -76,6 +76,9 @@ public:
 	UPROPERTY(BlueprintReadWrite, EditAnywhere, Category = "Settings|Extra Output Attributes", meta = (PCG_Overridable, EditCondition = "bOutputExtremityNeighborIndexAttribute"))
 	FName ExtremityNeighborIndexAttributeName = TEXT("ExtremityNeighborIndex");
 
+	/** Do not post node errors when the grammar does not fit the segment. */
+	UPROPERTY(BlueprintReadWrite, EditAnywhere, Category = Settings, AdvancedDisplay, meta = (PCG_Overridable))
+	bool bSilenceGrammarDoesNotFitErrors = false;
 };
 
 class PCGSubdivideSegmentHelpers;
diff --git a/Engine/Plugins/PCG/Source/PCG/Public/Elements/Grammar/PCGSubdivideSpline.h b/Engine/Plugins/PCG/Source/PCG/Public/Elements/Grammar/PCGSubdivideSpline.h
--- a/Engine/Plugins/PCG/Source/PCG/Public/Elements/Grammar/PCGSubdivideSpline.h
+++ b/Engine/Plugins/PCG/Source/PCG/Public/Elements/Grammar/PCGSubdivideSpline.h
@@ -67,9 +67,13 @@ public:
 	UPROPERTY(BlueprintReadWrite, EditAnywhere, Category = "Settings|Extra Output Attributes", meta = (PCG_Overridable, EditCondition = "bOutputExtremityAttributes", EditConditionHides))
 	FName IsFinalAttributeName = TEXT("IsFinal");
 
+	/** Do not post node errors when the grammar does not fit the spline. */
+	UPROPERTY(BlueprintReadWrite, EditAnywhere, Category = Settings, AdvancedDisplay, meta = (PCG_Overridable))
+	bool bSilenceGrammarDoesNotFitErrors = false;
+
 	/** Maximum acceptable distance tolerance between placed modules (overlap or gap) - a smaller number is more precise. Affects performance if the number is very small. */
 	UPROPERTY(BlueprintReadWrite, EditAnywhere, Category = Settings, AdvancedDisplay, meta = (PCG_Overridable, ClampMin="0.000001"))
-	double ModulePlacementTolerance = 0.01; 
+	double ModulePlacementTolerance = 0.01;
 };
 
 class FPCGSubdivideSplineElement : public FPCGSubdivisionBaseElement
diff --git a/Engine/Plugins/PCG/Source/PCG/Public/Elements/Grammar/PCGSubdivisionBase.h b/Engine/Plugins/PCG/Source/PCG/Public/Elements/Grammar/PCGSubdivisionBase.h
--- a/Engine/Plugins/PCG/Source/PCG/Public/Elements/Grammar/PCGSubdivisionBase.h
+++ b/Engine/Plugins/PCG/Source/PCG/Public/Elements/Grammar/PCGSubdivisionBase.h
@@ -176,7 +176,7 @@ namespace PCGSubdivisionBase
 	PCGGrammar::FTokenizedGrammar GetTokenizedGrammar(FPCGContext* InContext, const FString& InGrammar, const FModuleInfoMap& InModulesInfo, double& OutMinSize);
 
 	template <typename T>
-	bool Subdivide(const T& Root, double Length, TArray<TModuleInstance<T>>& OutModuleInstances, double& RemainingLength, FPCGContext* InOptionalContext = nullptr, int32 InOptionalAdditionalSeed = 0)
+	bool Subdivide(const T& Root, double Length, TArray<TModuleInstance<T>>& OutModuleInstances, double& RemainingLength, FPCGContext* InOptionalContext = nullptr, int32 InOptionalAdditionalSeed = 0, const bool bQuiet = false)
 	{
 		OutModuleInstances.Reset();
 		RemainingLength = Length;
@@ -195,7 +195,6 @@ namespace PCGSubdivisionBase
 
 		if (RemainingLength < 0)
 		{
-			PCGLog::LogErrorOnGraph(NSLOCTEXT("PCGSubdivisionBase", "SegmentCutFail", "Grammar doesn't fit for this segment."), InOptionalContext);
 			return false;
 		}
 
