[PCG] Added missing context to the export attributes error logging.

#rnx

diff --git a/Engine/Plugins/PCG/Source/PCG/Private/Elements/IO/PCGExportSelectedAttributes.cpp b/Engine/Plugins/PCG/Source/PCG/Private/Elements/IO/PCGExportSelectedAttributes.cpp
--- a/Engine/Plugins/PCG/Source/PCG/Private/Elements/IO/PCGExportSelectedAttributes.cpp
+++ b/Engine/Plugins/PCG/Source/PCG/Private/Elements/IO/PCGExportSelectedAttributes.cpp
@@ -171,14 +171,14 @@ bool FPCGExportSelectedAttributesElement::ExecuteInternal(FPCGContext* Context)
 		if (!bSuccess || OutputDirectory.IsEmpty())
 #endif // WITH_EDITOR
 		{
-			PCGLog::LogErrorOnGraph(FText::Format(LOCTEXT("InvalidDirectoryError", "Invalid or empty export directory: {0}"), FText::FromString(OutputDirectory)));
+			PCGLog::LogErrorOnGraph(FText::Format(LOCTEXT("InvalidDirectoryError", "Invalid or empty export directory: {0}"), FText::FromString(OutputDirectory)), Context);
 			return true;
 		}
 	}
 
 	if (Settings->FileName.IsEmpty())
 	{
-		PCGLog::LogErrorOnGraph(LOCTEXT("EmptyFileName", "A valid file name is required."));
+		PCGLog::LogErrorOnGraph(LOCTEXT("EmptyFileName", "A valid file name is required."), Context);
 		return true;
 	}
 
@@ -194,7 +194,7 @@ bool FPCGExportSelectedAttributesElement::ExecuteInternal(FPCGContext* Context)
 		Archive = IFileManager::Get().CreateFileWriter(*FinalPath);
 		if (!Archive)
 		{
-			PCGLog::LogErrorOnGraph(FText::Format(LOCTEXT("FileWriterInvalidPath", "Could not create a valid archive writer with path: {0}"), FText::FromString(FinalPath)));
+			PCGLog::LogErrorOnGraph(FText::Format(LOCTEXT("FileWriterInvalidPath", "Could not create a valid archive writer with path: {0}"), FText::FromString(FinalPath)), Context);
 			return true;
 		}
 
@@ -406,14 +406,14 @@ bool FPCGExportSelectedAttributesElement::ExecuteInternal(FPCGContext* Context)
 		const TSharedRef<TJsonWriter<>> JsonWriter = TJsonWriterFactory<>::Create(&JsonString);
 		if (!FJsonSerializer::Serialize(RootJsonObject.ToSharedRef(), JsonWriter))
 		{
-			PCGLog::LogErrorOnGraph(LOCTEXT("JsonSerializationError", "Serialization of Json data failed."));
+			PCGLog::LogErrorOnGraph(LOCTEXT("JsonSerializationError", "Serialization of Json data failed."), Context);
 			return true;
 		}
 
 		const FString FinalPath = Helpers::String::MergeFileName(OutputDirectory, Settings->FileName, TEXT(".json"));
 		if (!FFileHelper::SaveStringToFile(JsonString, *FinalPath))
 		{
-			PCGLog::LogErrorOnGraph(FText::Format(LOCTEXT("JsonInvalidPath", "Could not create a valid Json file with path: {0}"), FText::FromString(FinalPath)));
+			PCGLog::LogErrorOnGraph(FText::Format(LOCTEXT("JsonInvalidPath", "Could not create a valid Json file with path: {0}"), FText::FromString(FinalPath)), Context);
 			return true;
 		}
 	}
