[PCG] Fix for incorrectly transacting the reference to the scheduling policy when a PCG component is instanced by a BP class via SCS.

diff --git a/Engine/Plugins/PCG/Source/PCG/Private/PCGComponent.cpp b/Engine/Plugins/PCG/Source/PCG/Private/PCGComponent.cpp
--- a/Engine/Plugins/PCG/Source/PCG/Private/PCGComponent.cpp
+++ b/Engine/Plugins/PCG/Source/PCG/Private/PCGComponent.cpp
@@ -2218,7 +2218,10 @@ void UPCGComponent::PostLoad()
 	}
 	else
 	{
-		const EObjectFlags Flags = GetMaskedFlags(RF_PropagateToSubObjects) | RF_Transactional;
+		// Note: Should not be transactional *unless* we are - this flag is already included in RF_PropagateToSubobjects.
+		// Clearing it here in case it was incorrectly set on a previous save. This is needed to resolve some undo issues.
+		SchedulingPolicy->ClearFlags(RF_Transactional);
+		const EObjectFlags Flags = GetMaskedFlags(RF_PropagateToSubObjects);
 		SchedulingPolicy->SetFlags(Flags);
 #if WITH_EDITOR
 		SchedulingPolicy->SetShouldDisplayProperties(IsManagedByRuntimeGenSystem());
