[PCG] PCG Cache implement UPCGTextureData::GetResourceSizeEx
- TLS enabled mode which considers Mip 0 on the DuplicateTexture as system memory to help with cache size reporting
- Todo: revisit usage of EResourceSizeMode::Exclusive

#rnx

diff --git a/Engine/Plugins/PCG/Source/PCG/Private/Data/PCGTextureData.cpp b/Engine/Plugins/PCG/Source/PCG/Private/Data/PCGTextureData.cpp
--- a/Engine/Plugins/PCG/Source/PCG/Private/Data/PCGTextureData.cpp
+++ b/Engine/Plugins/PCG/Source/PCG/Private/Data/PCGTextureData.cpp
@@ -6,6 +6,7 @@
 #include "PCGTextureReadback.h"
 #include "Data/PCGPointArrayData.h"
 #include "Data/PCGPointData.h"
+#include "Graph/PCGGraphCache.h"
 #include "Helpers/PCGAsync.h"
 #include "Helpers/PCGHelpers.h"
 
@@ -665,6 +666,20 @@ bool UPCGTextureData::Initialize(TRefCountPtr<IPooledRenderTarget> InTextureHand
 	return bSuccessfullyInitialized;
 }
 
+void UPCGTextureData::GetResourceSizeEx(FResourceSizeEx& CumulativeResourceSize)
+{
+	Super::GetResourceSizeEx(CumulativeResourceSize);
+
+#if WITH_EDITORONLY_DATA
+	if (FPCGGraphCache::IsGatheringCacheSize() && ::IsValid(DuplicateTexture) && DuplicateTexture->GetNumMips() > 0)
+	{
+		DuplicateTexture->Source.LockMipReadOnly(0);
+		CumulativeResourceSize.AddDedicatedSystemMemoryBytes(DuplicateTexture->Source.CalcMipSize(0));
+		DuplicateTexture->Source.UnlockMip(0);
+	}
+#endif
+}
+
 void UPCGTextureData::AddToCrc(FArchiveCrc32& Ar, bool bFullDataCrc) const
 {
 	Super::AddToCrc(Ar, bFullDataCrc);
diff --git a/Engine/Plugins/PCG/Source/PCG/Private/Graph/PCGGraphCache.cpp b/Engine/Plugins/PCG/Source/PCG/Private/Graph/PCGGraphCache.cpp
--- a/Engine/Plugins/PCG/Source/PCG/Private/Graph/PCGGraphCache.cpp
+++ b/Engine/Plugins/PCG/Source/PCG/Private/Graph/PCGGraphCache.cpp
@@ -39,9 +39,24 @@ static TAutoConsoleVariable<bool> CVarValidateElementToCacheEntryKeys(
 	false,
 	TEXT("Validate ElementToCacheEntryKeys acceleration table (debug)."));
 
+static TAutoConsoleVariable<bool> CVarEnableTLSGatherCacheSize(
+	TEXT("pcg.Cache.EnableTLSGatherCacheSize"),
+	false,
+	TEXT("Whether we give a hint to UPCGData subclass that we are calling GetResourceSizeEx for cache purposes."));
+
 // Initial max number of entries graph cache
 static const int32 GPCGGraphCacheInitialCapacity = 65536;
 
+namespace PCGGraphCacheHelpers
+{
+	thread_local bool bIsGatheringCacheSize = false;
+}
+
+bool FPCGGraphCache::IsGatheringCacheSize()
+{
+	return PCGGraphCacheHelpers::bIsGatheringCacheSize;
+}
+
 FPCGGraphCache::FPCGGraphCache()
 	: CacheData(CVarCacheEnabled.GetValueOnGameThread() ? GPCGGraphCacheInitialCapacity : 0)
 {
@@ -339,6 +354,13 @@ void FPCGGraphCache::AddDataToAccountedMemory(const FPCGDataCollection& InCollec
 					}
 					else
 					{
+						PCGGraphCacheHelpers::bIsGatheringCacheSize = CVarEnableTLSGatherCacheSize.GetValueOnAnyThread();
+						ON_SCOPE_EXIT
+						{
+							PCGGraphCacheHelpers::bIsGatheringCacheSize = false;
+						};
+
+						// @todo_pcg: revisit usage of EResourceSizeMode::Exclusive instead of EResourceSizeMode::EstimatedTotals
 						FResourceSizeEx ResSize = FResourceSizeEx(EResourceSizeMode::Exclusive);
 						// Calculate data size. Function is non-const but is const-like, especially when
 						// resource mode is Exclusive. The other mode calls a function to find all outer'd
diff --git a/Engine/Plugins/PCG/Source/PCG/Private/Graph/PCGGraphCache.h b/Engine/Plugins/PCG/Source/PCG/Private/Graph/PCGGraphCache.h
--- a/Engine/Plugins/PCG/Source/PCG/Private/Graph/PCGGraphCache.h
+++ b/Engine/Plugins/PCG/Source/PCG/Private/Graph/PCGGraphCache.h
@@ -68,6 +68,8 @@ public:
 	FPCGGraphCache();
 	~FPCGGraphCache();
 
+	static bool IsGatheringCacheSize();
+
 	/** Returns true if data was found from the cache, in which case the outputs are written in OutOutput. InNode is optional and for logging only. */
 	virtual bool GetFromCache(const FPCGGetFromCacheParams& Params, FPCGDataCollection& OutCollection) const override;
 
diff --git a/Engine/Plugins/PCG/Source/PCG/Public/Data/PCGTextureData.h b/Engine/Plugins/PCG/Source/PCG/Public/Data/PCGTextureData.h
--- a/Engine/Plugins/PCG/Source/PCG/Public/Data/PCGTextureData.h
+++ b/Engine/Plugins/PCG/Source/PCG/Public/Data/PCGTextureData.h
@@ -212,6 +212,10 @@ public:
 
 	EPCGTextureResourceType GetTextureResourceType() const override { return ResourceType; }
 
+	//~Begin UObject Interface
+	virtual void GetResourceSizeEx(FResourceSizeEx& CumulativeResourceSize) override;
+	//~End UObject Interface
+
 	//~Begin UPCGData interface
 	PCG_ASSIGN_DEFAULT_TYPE_INFO(FPCGDataTypeInfoTexture2D)
 	virtual void AddToCrc(FArchiveCrc32& Ar, bool bFullDataCrc) const override;
