[PCG] Bughawk fix 46125523: Check Source for num mips

#rnx

diff --git a/Engine/Plugins/PCG/Source/PCG/Private/Data/PCGTextureData.cpp b/Engine/Plugins/PCG/Source/PCG/Private/Data/PCGTextureData.cpp
--- a/Engine/Plugins/PCG/Source/PCG/Private/Data/PCGTextureData.cpp
+++ b/Engine/Plugins/PCG/Source/PCG/Private/Data/PCGTextureData.cpp
@@ -671,7 +671,7 @@ void UPCGTextureData::GetResourceSizeEx(FResourceSizeEx& CumulativeResourceSize)
 	Super::GetResourceSizeEx(CumulativeResourceSize);
 
 #if WITH_EDITORONLY_DATA
-	if (FPCGGraphCache::IsGatheringCacheSize() && ::IsValid(DuplicateTexture) && DuplicateTexture->GetNumMips() > 0)
+	if (FPCGGraphCache::IsGatheringCacheSize() && ::IsValid(DuplicateTexture) && DuplicateTexture->Source.IsValid() && DuplicateTexture->Source.GetNumMips() > 0)
 	{
 		DuplicateTexture->Source.LockMipReadOnly(0);
 		CumulativeResourceSize.AddDedicatedSystemMemoryBytes(DuplicateTexture->Source.CalcMipSize(0));
