[PCG] Add node to statically switch based on platform group.

[RN] [PCG] Added platform switch node to enable/disable parts of graph based on target platform.

diff --git a/Engine/Plugins/PCG/Source/PCG/Private/Elements/ControlFlow/PCGPlatformSwitch.cpp b/Engine/Plugins/PCG/Source/PCG/Private/Elements/ControlFlow/PCGPlatformSwitch.cpp
new file mode 100644
--- /dev/null
+++ b/Engine/Plugins/PCG/Source/PCG/Private/Elements/ControlFlow/PCGPlatformSwitch.cpp
@@ -0,0 +1,276 @@
+// Copyright Epic Games, Inc. All Rights Reserved.
+
+#include "Elements/ControlFlow/PCGPlatformSwitch.h"
+
+#include "PCGContext.h"
+#include "PCGGraph.h"
+#include "PCGModule.h"
+#include "Elements/PCGGather.h"
+#include "Helpers/PCGHelpers.h"
+#include "Subsystems/PCGSubsystem.h"
+
+#include "DeviceProfiles/DeviceProfile.h"
+#include "DeviceProfiles/DeviceProfileManager.h"
+#include "Interfaces/ITargetPlatform.h"
+#include "Interfaces/ITargetPlatformManagerModule.h"
+#include "Misc/DataDrivenPlatformInfoRegistry.h"
+#include "UObject/ObjectSaveContext.h"
+
+#if WITH_EDITOR
+#include "Editor.h"
+#endif
+
+#include UE_INLINE_GENERATED_CPP_BY_NAME(PCGPlatformSwitch)
+
+#define LOCTEXT_NAMESPACE "FPCGPlatformSwitchElement"
+
+namespace PCGPlatformSwitch
+{
+	static const FName DefaultLabel = TEXT("Default");
+
+#if WITH_EDITOR
+	bool GetCurrentPlatformInfoFromEditor(FName& OutPlatform, FName& OutPlatformGroup)
+	{
+		// Ask the editor for the active preview platform (PS5, XSX, etc.)
+		FName IniPlatform;
+		if (GEditor && GEditor->GetPreviewPlatformName(IniPlatform))
+		{
+			OutPlatform = IniPlatform;
+			OutPlatformGroup = FDataDrivenPlatformInfoRegistry::GetPlatformInfo(IniPlatform).PlatformGroupName;
+			return true;
+		}
+
+		// Single active target platform selected by the user (Project Launcher / Platforms UI)
+		const TArray<ITargetPlatform*>& ActiveTPs = GetTargetPlatformManagerRef().GetActiveTargetPlatforms();
+		if (ActiveTPs.Num() == 1)
+		{
+			OutPlatform = ActiveTPs[0]->GetPlatformInfo().IniPlatformName;
+			OutPlatformGroup = ActiveTPs[0]->GetPlatformInfo().PlatformGroupName;
+			return true;
+		}
+
+		// Device Profile (covers Mobile Preview / device-specific previews in-editor)
+		if (UDeviceProfile* ActiveProfile = UDeviceProfileManager::Get().GetActiveProfile())
+		{
+			// DeviceType typically matches the Ini platform name ("Windows", "Android", "PS5", ...)
+			OutPlatform = *ActiveProfile->DeviceType;
+			OutPlatformGroup = FDataDrivenPlatformInfoRegistry::GetPlatformInfo(OutPlatform).PlatformGroupName;
+			return true;
+		}
+
+		// Fallback: the platform the editor is running on
+		if (const ITargetPlatform* Running = GetTargetPlatformManagerRef().GetRunningTargetPlatform())
+		{
+			OutPlatform = Running->GetPlatformInfo().IniPlatformName;
+			OutPlatformGroup = Running->GetPlatformInfo().PlatformGroupName;
+			return true;
+		}
+
+		return ensure(false);
+	}
+#endif // WITH_EDITOR
+}
+
+bool UPCGPlatformSwitchSettings::IsPinStaticallyActive(const FName& PinLabel) const
+{
+	return PinLabel == PCGPlatformSwitch::DefaultLabel || IsPlatformSelected(PinLabel);
+}
+
+FString UPCGPlatformSwitchSettings::GetAdditionalTitleInformation() const
+{
+#if WITH_EDITOR
+	FString Selection;
+	if (PlatformOutputs.IsEmpty())
+	{
+		Selection = TEXT("None");
+	}
+	else if (PlatformOutputs.Num() == 1)
+	{
+		Selection = PlatformOutputs[0].ToString();
+	}
+	else
+	{
+		Selection = TEXT("Multiple");
+	}
+
+	return FString::Format(TEXT("Platform: {0}"), { Selection });
+#else
+	return {};
+#endif
+}
+
+#if WITH_EDITOR
+void UPCGPlatformSwitchSettings::PreSave(FObjectPreSaveContext ObjectSaveContext)
+{
+	if (ObjectSaveContext.IsCooking())
+	{
+		ensure(PCGPlatformSwitch::GetCurrentPlatformInfoFromEditor(CookedPlatform, CookedPlatformGroup));
+	}
+}
+
+FText UPCGPlatformSwitchSettings::GetNodeTooltipText() const
+{
+	return LOCTEXT("NodeTooltip", "Statically activates/deactivates output pins based on the current platform or platform group (desktop/console/mobile), useful for culling downstream nodes.");
+}
+
+EPCGChangeType UPCGPlatformSwitchSettings::GetChangeTypeForProperty(const FName& InPropertyName) const
+{
+	EPCGChangeType ChangeType = Super::GetChangeTypeForProperty(InPropertyName) | EPCGChangeType::Cosmetic;
+
+	if (InPropertyName == GET_MEMBER_NAME_CHECKED(UPCGPlatformSwitchSettings, bEnabled)
+		|| InPropertyName == GET_MEMBER_NAME_CHECKED(UPCGPlatformSwitchSettings, PlatformOutputs))
+	{
+		ChangeType |= EPCGChangeType::Structural;
+	}
+
+	return ChangeType;
+}
+
+TArray<FName> UPCGPlatformSwitchSettings::GetPlatformOptions() const
+{
+	const TArray<const FDataDrivenPlatformInfo*>& SortedPlatforms = FDataDrivenPlatformInfoRegistry::GetSortedPlatformInfos(EPlatformInfoType::TruePlatformsOnly);
+
+	// Platform names and platform group names.
+	TArray<FName> BasePlatformNameArray;
+	TArray<FName> PlatformGroupNameArray;
+
+	// Create mapping from platform to platform groups and remove postfixes and invalid platform names
+	for (const FDataDrivenPlatformInfo* DDPI : SortedPlatforms)
+	{
+		// Add platform name if it isn't already set, and also add to group mapping
+		if (!BasePlatformNameArray.Contains(DDPI->IniPlatformName))
+		{
+			BasePlatformNameArray.Add(DDPI->IniPlatformName);
+			PlatformGroupNameArray.AddUnique(DDPI->PlatformGroupName);
+		}
+	}
+
+	TArray<FName> PlatformOptions = MoveTemp(PlatformGroupNameArray);
+	PlatformOptions.Append(BasePlatformNameArray);
+
+	return PlatformOptions;
+}
+#endif // WITH_EDITOR
+
+TArray<FPCGPinProperties> UPCGPlatformSwitchSettings::InputPinProperties() const
+{
+	TArray<FPCGPinProperties> PinProperties;
+	PinProperties.Emplace(PCGPinConstants::DefaultInputLabel, EPCGDataType::Any, /*bInAllowMultipleConnections=*/true, /*bAllowMultipleData=*/true,
+		LOCTEXT("OutputPinTooltip", "All input will be forwarded directly to the active output pin."));
+
+	return PinProperties;
+}
+
+TArray<FPCGPinProperties> UPCGPlatformSwitchSettings::OutputPinProperties() const
+{
+	TArray<FPCGPinProperties> PinProperties;
+
+	PinProperties.Emplace(TEXT("Default"), EPCGDataType::Any);
+
+	for (FName Platform : GetSanitizedPlatforms())
+	{
+		PinProperties.Emplace(Platform, EPCGDataType::Any);
+	}
+
+	return PinProperties;
+}
+
+FPCGElementPtr UPCGPlatformSwitchSettings::CreateElement() const
+{
+	return MakeShared<FPCGPlatformSwitchElement>();
+}
+
+bool UPCGPlatformSwitchSettings::IsPlatformSelected(FName InPlatform) const
+{
+	FName CurrentPlatform, CurrentPlatformGroup;
+
+	if (GetCurrentPlatformInfo(CurrentPlatform, CurrentPlatformGroup))
+	{
+		return IsPlatformSelected(InPlatform, CurrentPlatform, CurrentPlatformGroup);
+	}
+
+	return false;
+}
+
+bool UPCGPlatformSwitchSettings::IsPlatformSelected(FName InPlatform, FName InCurrentPlatform, FName InCurrentPlatformGroup) const
+{
+	ensure(!InCurrentPlatform.IsNone());
+	ensure(!InCurrentPlatformGroup.IsNone());
+
+	return InPlatform == InCurrentPlatform || InPlatform == InCurrentPlatformGroup;
+}
+
+bool UPCGPlatformSwitchSettings::GetCurrentPlatformInfo(FName& OutPlatform, FName& OutPlatformGroup) const
+{
+#if WITH_EDITOR
+	return ensure(PCGPlatformSwitch::GetCurrentPlatformInfoFromEditor(OutPlatform, OutPlatformGroup));
+#else
+	OutPlatform = CookedPlatform;
+	OutPlatformGroup = CookedPlatformGroup;
+	return true;
+#endif
+}
+
+TArray<FName> UPCGPlatformSwitchSettings::GetSanitizedPlatforms() const
+{
+	TArray<FName> Platforms;
+	Platforms.Reserve(PlatformOutputs.Num());
+
+	for (FName Platform : PlatformOutputs)
+	{
+		if (!Platform.IsNone() && !Platforms.Contains(Platform))
+		{
+			Platforms.Add(Platform);
+		}
+	}
+
+	return Platforms;
+}
+
+bool UPCGPlatformSwitchSettings::IsPinUsedByNodeExecution(const UPCGPin* InPin) const
+{
+	if (!InPin->IsOutputPin())
+	{
+		return Super::IsPinUsedByNodeExecution(InPin);
+	}
+
+	return IsPinStaticallyActive(InPin->Properties.Label);
+}
+
+bool FPCGPlatformSwitchElement::ExecuteInternal(FPCGContext* InContext) const
+{
+	TRACE_CPUPROFILER_EVENT_SCOPE(FPCGPlatformSwitchElement::ExecuteInternal);
+
+	const UPCGPlatformSwitchSettings* Settings = InContext->GetInputSettings<UPCGPlatformSwitchSettings>();
+	check(Settings);
+
+	const FPCGDataCollection InputData = PCGGather::GatherDataForPin(InContext->InputData, PCGPinConstants::DefaultInputLabel, PCGPlatformSwitch::DefaultLabel);
+
+	const TArray<FName> Platforms = Settings->GetSanitizedPlatforms();
+	InContext->OutputData.TaggedData.Reserve(InputData.TaggedData.Num() * (1 + Platforms.Num()));
+
+	// Data for Default pin.
+	InContext->OutputData.TaggedData.Append(InputData.TaggedData);
+
+	FName CurrentPlatform, CurrentPlatformGroup;
+
+	if (Settings->GetCurrentPlatformInfo(CurrentPlatform, CurrentPlatformGroup))
+	{
+		// Copy of data for each platform output.
+		for (FName Platform : Platforms)
+		{
+			if (Settings->IsPlatformSelected(Platform, CurrentPlatform, CurrentPlatformGroup))
+			{
+				for (FPCGTaggedData TaggedData : InputData.TaggedData)
+				{
+					TaggedData.Pin = Platform;
+					InContext->OutputData.TaggedData.Add(TaggedData);
+				}
+			}
+		}
+	}
+
+	return true;
+}
+
+#undef LOCTEXT_NAMESPACE
diff --git a/Engine/Plugins/PCG/Source/PCG/Private/Subsystems/PCGSubsystem.cpp b/Engine/Plugins/PCG/Source/PCG/Private/Subsystems/PCGSubsystem.cpp
--- a/Engine/Plugins/PCG/Source/PCG/Private/Subsystems/PCGSubsystem.cpp
+++ b/Engine/Plugins/PCG/Source/PCG/Private/Subsystems/PCGSubsystem.cpp
@@ -211,6 +211,12 @@ void UPCGSubsystem::Deinitialize()
 
 #if WITH_EDITOR
 	FCoreUObjectDelegates::OnObjectsReplaced.RemoveAll(this);
+
+	if (GEditor)
+	{
+		GEditor->OnPreviewPlatformChanged().RemoveAll(this);
+		GEditor->OnPreviewFeatureLevelChanged().RemoveAll(this);
+	}
 #endif
 
 	Super::Deinitialize();
@@ -245,6 +251,12 @@ void UPCGSubsystem::PostInitialize()
 
 #if WITH_EDITOR
 	FCoreUObjectDelegates::OnObjectsReplaced.AddUObject(this, &UPCGSubsystem::OnObjectsReplaced);
+
+	if (GEditor)
+	{
+		GEditor->OnPreviewPlatformChanged().AddUObject(this, &UPCGSubsystem::OnPreviewPlatformChanged);
+		GEditor->OnPreviewFeatureLevelChanged().AddUObject(this, &UPCGSubsystem::OnPreviewFeatureLevelChanged);
+	}
 #endif
 }
 
@@ -1885,6 +1897,20 @@ void UPCGSubsystem::CreateMissingPartitionActors()
 	}
 }
 
+void UPCGSubsystem::OnPreviewPlatformChanged()
+{
+	FlushCache();
+
+	RefreshAllRuntimeGenComponents(EPCGChangeType::GenerationGrid);
+}
+
+void UPCGSubsystem::OnPreviewFeatureLevelChanged(ERHIFeatureLevel::Type)
+{
+	FlushCache();
+
+	RefreshAllRuntimeGenComponents(EPCGChangeType::GenerationGrid);
+}
+
 void UPCGSubsystem::CreatePartitionActorsWithinBounds(UPCGComponent* InComponent, const FBox& InBounds, const PCGHiGenGrid::FSizeArray& InGridSizes)
 {
 	UWorld* World = GetWorld();
diff --git a/Engine/Plugins/PCG/Source/PCG/Public/Elements/ControlFlow/PCGPlatformSwitch.h b/Engine/Plugins/PCG/Source/PCG/Public/Elements/ControlFlow/PCGPlatformSwitch.h
new file mode 100644
--- /dev/null
+++ b/Engine/Plugins/PCG/Source/PCG/Public/Elements/ControlFlow/PCGPlatformSwitch.h
@@ -0,0 +1,82 @@
+// Copyright Epic Games, Inc. All Rights Reserved.
+
+#pragma once
+
+#include "PCGControlFlow.h"
+
+#include "PCGPlatformSwitch.generated.h"
+
+/** Statically activates/deactivates output pins based on the current platform/platform group. */
+UCLASS(BlueprintType, ClassGroup = (Procedural), meta=(Keywords = "if branch platform"))
+class UPCGPlatformSwitchSettings : public UPCGControlFlowSettings
+{
+	GENERATED_BODY()
+
+public:
+	//~ Begin UObject interface
+#if WITH_EDITOR
+	virtual void PreSave(FObjectPreSaveContext ObjectSaveContext) override;
+#endif
+	//~ End UObject interface
+
+	//~ Begin UPCGSettings interface
+#if WITH_EDITOR
+	virtual FName GetDefaultNodeName() const override { return FName(TEXT("PlatformSwitch")); }
+	virtual FText GetDefaultNodeTitle() const override { return NSLOCTEXT("FPCGPlatformSwitchElement", "NodeTitle", "Platform Switch"); }
+	virtual FText GetNodeTooltipText() const override;
+	virtual EPCGSettingsType GetType() const override { return EPCGSettingsType::ControlFlow; }
+#endif // WITH_EDITOR
+	virtual bool HasDynamicPins() const override { return true; }
+	virtual bool IsPinStaticallyActive(const FName& PinLabel) const override;
+	virtual bool IsPinUsedByNodeExecution(const UPCGPin* InPin) const override;
+	virtual FString GetAdditionalTitleInformation() const override;
+	virtual bool HasFlippedTitleLines() const override { return true; }
+	virtual bool CanCullTaskIfUnwired() const override { return false; } // Useful to run unwired to only cull downstream nodes.
+
+protected:
+#if WITH_EDITOR
+	virtual EPCGChangeType GetChangeTypeForProperty(const FName& InPropertyName) const override;
+#endif
+	virtual TArray<FPCGPinProperties> InputPinProperties() const override;
+	virtual TArray<FPCGPinProperties> OutputPinProperties() const override;
+	virtual FPCGElementPtr CreateElement() const override;
+	//~ End UPCGSettings interface
+
+public:
+	/** Get the users selection, sanitized to remove None and duplicates. */
+	TArray<FName> GetSanitizedPlatforms() const;
+
+	/** Returns current platform and platform group, either by querying from editor or by using cooked values. */
+	bool GetCurrentPlatformInfo(FName& OutPlatform, FName& OutPlatformGroup) const;
+
+	/** Returns true if platform or platform group included in user's selection. */
+	bool IsPlatformSelected(FName InPlatform) const;
+	bool IsPlatformSelected(FName InPlatform, FName InCurrentPlatform, FName InCurrentPlatformGroup) const;
+
+protected:
+#if WITH_EDITOR
+	/** List of all non-advanced input pin names. */
+	UFUNCTION()
+	TArray<FName> GetPlatformOptions() const;
+#endif
+
+protected:
+	UPROPERTY(EditAnywhere, Category=Settings, meta = (GetOptions = "GetPlatformOptions"))
+	TArray<FName> PlatformOutputs;
+
+	/** The current platform used in cooked builds. */
+	UPROPERTY()
+	FName CookedPlatform;
+
+	/** The current platform group used in cooked builds. */
+	UPROPERTY()
+	FName CookedPlatformGroup;
+};
+
+class FPCGPlatformSwitchElement : public IPCGElement
+{
+protected:
+	virtual bool ExecuteInternal(FPCGContext* InContext) const override;
+	virtual bool SupportsGPUResidentData(FPCGContext* InContext) const override { return true; }
+	virtual bool SupportsBasePointDataInputs(FPCGContext* InContext) const { return true; }
+};
diff --git a/Engine/Plugins/PCG/Source/PCG/Public/Subsystems/IPCGBaseSubsystem.h b/Engine/Plugins/PCG/Source/PCG/Public/Subsystems/IPCGBaseSubsystem.h
--- a/Engine/Plugins/PCG/Source/PCG/Public/Subsystems/IPCGBaseSubsystem.h
+++ b/Engine/Plugins/PCG/Source/PCG/Public/Subsystems/IPCGBaseSubsystem.h
@@ -86,7 +86,7 @@ public:
 	/** Returns the interface to the cache, required for element per-data caching */
 	PCG_API IPCGGraphCache* GetCache();
 
-	/** Flushes the graph cache completely, use only for debugging */
+	/** Flushes the graph cache and graph compiler cache, and in editor trigger a GC. */
 	PCG_API void FlushCache();
 
 	/** True if graph cache debugging is enabled. */
diff --git a/Engine/Plugins/PCG/Source/PCG/Public/Subsystems/PCGSubsystem.h b/Engine/Plugins/PCG/Source/PCG/Public/Subsystems/PCGSubsystem.h
--- a/Engine/Plugins/PCG/Source/PCG/Public/Subsystems/PCGSubsystem.h
+++ b/Engine/Plugins/PCG/Source/PCG/Public/Subsystems/PCGSubsystem.h
@@ -29,6 +29,7 @@ class UPCGLandscapeCache;
 
 enum class EPCGComponentDirtyFlag : uint8;
 enum class ETickableTickType : uint8;
+namespace ERHIFeatureLevel { enum Type : int; }
 
 class IPCGGraphCache;
 class FPCGGraphCompiler;
@@ -355,6 +356,10 @@ public:
 	FPCGOnPCGComponentGenerationDone OnPCGComponentGenerationDone;
 
 	UE_API void CreateMissingPartitionActors();
+
+	void OnPreviewPlatformChanged();
+	void OnPreviewFeatureLevelChanged(ERHIFeatureLevel::Type);
+
 private:
 	UE_API void OnPCGGraphCancelled(UPCGComponent* InComponent);
 	UE_API void OnPCGGraphStartGenerating(UPCGComponent* InComponent);
