UE: CIS quick fix

diff --git a/Engine/Plugins/PCG/Source/PCG/Private/Elements/ControlFlow/PCGPlatformSwitch.cpp b/Engine/Plugins/PCG/Source/PCG/Private/Elements/ControlFlow/PCGPlatformSwitch.cpp
--- a/Engine/Plugins/PCG/Source/PCG/Private/Elements/ControlFlow/PCGPlatformSwitch.cpp
+++ b/Engine/Plugins/PCG/Source/PCG/Private/Elements/ControlFlow/PCGPlatformSwitch.cpp
@@ -102,6 +102,8 @@ FString UPCGPlatformSwitchSettings::GetAdditionalTitleInformation() const
 #if WITH_EDITOR
 void UPCGPlatformSwitchSettings::PreSave(FObjectPreSaveContext ObjectSaveContext)
 {
+	Super::PreSave(ObjectSaveContext);
+	
 	if (ObjectSaveContext.IsCooking())
 	{
 		ensure(PCGPlatformSwitch::GetCurrentPlatformInfoFromEditor(CookedPlatform, CookedPlatformGroup));
