[PCG] Component details: identify custom button row

#rnx

diff --git a/Engine/Plugins/PCG/Source/PCGEditor/Private/Details/PCGComponentDetails.cpp b/Engine/Plugins/PCG/Source/PCGEditor/Private/Details/PCGComponentDetails.cpp
--- a/Engine/Plugins/PCG/Source/PCGEditor/Private/Details/PCGComponentDetails.cpp
+++ b/Engine/Plugins/PCG/Source/PCGEditor/Private/Details/PCGComponentDetails.cpp
@@ -89,7 +89,7 @@ void FPCGComponentDetails::CustomizeDetails(IDetailLayoutBuilder& DetailBuilder)
 		PCGCategory.AddExternalObjectProperty(Objects, TEXT("ParametersOverrides"));
 	}
 
-	FDetailWidgetRow& NewRow = PCGCategory.AddCustomRow(FText::GetEmpty());
+	FDetailWidgetRow& NewRow = PCGCategory.AddCustomRow(FText::GetEmpty()).RowTag("PCG_Component_Actions");
 
 	NewRow.ValueContent()
 		.MaxDesiredWidth(120.f)
