[PCG] Platform Switch: make Default pin active only if no other pin is selected.

#rnx

diff --git a/Engine/Plugins/PCG/Source/PCG/Private/Elements/ControlFlow/PCGPlatformSwitch.cpp b/Engine/Plugins/PCG/Source/PCG/Private/Elements/ControlFlow/PCGPlatformSwitch.cpp
--- a/Engine/Plugins/PCG/Source/PCG/Private/Elements/ControlFlow/PCGPlatformSwitch.cpp
+++ b/Engine/Plugins/PCG/Source/PCG/Private/Elements/ControlFlow/PCGPlatformSwitch.cpp
@@ -73,7 +73,31 @@ namespace PCGPlatformSwitch
 
 bool UPCGPlatformSwitchSettings::IsPinStaticallyActive(const FName& PinLabel) const
 {
-	return PinLabel == PCGPlatformSwitch::DefaultLabel || IsPlatformSelected(PinLabel);
+	FName CurrentPlatform, CurrentPlatformGroup;
+
+	if (GetCurrentPlatformInfo(CurrentPlatform, CurrentPlatformGroup))
+	{
+		if (PinLabel != PCGPlatformSwitch::DefaultLabel)
+		{
+			// Active if the platform for this pin is selected.
+			return IsPlatformSelected(PinLabel, CurrentPlatform, CurrentPlatformGroup);
+		}
+		else
+		{
+			// Default only active if no other pin is active.
+			for (FName Platform : GetSanitizedPlatforms())
+			{
+				if (IsPlatformSelected(Platform, CurrentPlatform, CurrentPlatformGroup))
+				{
+					return false;
+				}
+			}
+
+			return true;
+		}
+	}
+
+	return false;
 }
 
 FString UPCGPlatformSwitchSettings::GetAdditionalTitleInformation() const
@@ -251,27 +275,35 @@ bool FPCGPlatformSwitchElement::ExecuteInternal(FPCGContext* InContext) const
 	const TArray<FName> Platforms = Settings->GetSanitizedPlatforms();
 	InContext->OutputData.TaggedData.Reserve(InputData.TaggedData.Num() * (1 + Platforms.Num()));
 
-	// Data for Default pin.
-	InContext->OutputData.TaggedData.Append(InputData.TaggedData);
-
 	FName CurrentPlatform, CurrentPlatformGroup;
+	if (!Settings->GetCurrentPlatformInfo(CurrentPlatform, CurrentPlatformGroup))
+	{
+		return true;
+	}
+
+	bool bAnyPlatformSelected = false;
 
-	if (Settings->GetCurrentPlatformInfo(CurrentPlatform, CurrentPlatformGroup))
+	// Copy of data for each platform output.
+	for (FName Platform : Platforms)
 	{
-		// Copy of data for each platform output.
-		for (FName Platform : Platforms)
+		if (Settings->IsPlatformSelected(Platform, CurrentPlatform, CurrentPlatformGroup))
 		{
-			if (Settings->IsPlatformSelected(Platform, CurrentPlatform, CurrentPlatformGroup))
+			for (FPCGTaggedData TaggedData : InputData.TaggedData)
 			{
-				for (FPCGTaggedData TaggedData : InputData.TaggedData)
-				{
-					TaggedData.Pin = Platform;
-					InContext->OutputData.TaggedData.Add(TaggedData);
-				}
+				TaggedData.Pin = Platform;
+				InContext->OutputData.TaggedData.Add(TaggedData);
 			}
+
+			bAnyPlatformSelected = true;
 		}
 	}
 
+	if (!bAnyPlatformSelected)
+	{
+		// Fall back to Default pin.
+		InContext->OutputData.TaggedData.Append(InputData.TaggedData);
+	}
+
 	return true;
 }
 
