[PCG] Remove most of the pin widget dynamic lambdas which can simply be baked onto the pin and updated on refresh.

#rnx

diff --git a/Engine/Plugins/PCG/Source/PCGEditor/Private/Pins/SPCGEditorGraphNodePin.cpp b/Engine/Plugins/PCG/Source/PCGEditor/Private/Pins/SPCGEditorGraphNodePin.cpp
--- a/Engine/Plugins/PCG/Source/PCGEditor/Private/Pins/SPCGEditorGraphNodePin.cpp
+++ b/Engine/Plugins/PCG/Source/PCGEditor/Private/Pins/SPCGEditorGraphNodePin.cpp
@@ -14,6 +14,7 @@
 #include "Algo/AnyOf.h"
 #include "Schema/PCGEditorGraphSchema.h"
 #include "Widgets/SToolTip.h"
+#include "Widgets/Images/SLayeredImage.h"
 #include "Widgets/Input/SButton.h"
 
 void SPCGEditorGraphNodePin::Construct(const FArguments& InArgs, UEdGraphPin* InPin)
@@ -25,7 +26,7 @@ void SPCGEditorGraphNodePin::Construct(const FArguments& InArgs, UEdGraphPin* In
 	bUsePinColorForText = InArgs._UsePinColorForText;
 	this->SetCursor(EMouseCursor::Default);
 
-	SetVisibility(MakeAttributeSP(this, &SPCGEditorGraphNodePin::GetPinVisiblity));
+	SetVisibility(GetPinVisiblity());
 
 	GraphPinObj = InPin;
 	check(GraphPinObj != NULL);
@@ -40,6 +41,8 @@ void SPCGEditorGraphNodePin::Construct(const FArguments& InArgs, UEdGraphPin* In
 		);
 
 	const bool bIsInput = (GetDirection() == EGPD_Input);
+	const FSlateColor PinColor = GetPinColor();
+	const FSlateColor PinTextColor = GetPinTextColor();
 
 	// A small marker to indicate the pin is required for the node to be executed.
 	TSharedPtr<SImage> RequiredPinIconWidget;
@@ -54,15 +57,14 @@ void SPCGEditorGraphNodePin::Construct(const FArguments& InArgs, UEdGraphPin* In
 		RequiredPinIconWidget =
 			SNew(SImage)
 			.Image(bDisplayPinMarker ? RequiredPinMarkerIcon : FAppStyle::GetNoBrush())
-			.ColorAndOpacity(this, &SPCGEditorGraphNodePin::GetPinColor);
+			.ColorAndOpacity(PinColor);
 	}
 
 	// Create the pin icon widget
-	TSharedRef<SWidget> PinWidgetRef = SPinTypeSelector::ConstructPinTypeImage(
-		MakeAttributeSP(this, &SPCGEditorGraphNodePin::GetPinIcon),
-		MakeAttributeSP(this, &SPCGEditorGraphNodePin::GetPinColor),
-		MakeAttributeSP(this, &SPCGEditorGraphNodePin::GetSecondaryPinIcon),
-		MakeAttributeSP(this, &SPCGEditorGraphNodePin::GetSecondaryPinColor));
+	TSharedRef<SWidget> PinWidgetRef =
+		SNew(SLayeredImage, GetSecondaryPinIcon(), GetSecondaryPinColor())
+		.Image(MakeAttributeSP(this, &SPCGEditorGraphNodePin::GetPinIcon))
+		.ColorAndOpacity(PinColor);
 	PinImage = PinWidgetRef;
 
 	PinWidgetRef->SetCursor(
@@ -84,7 +86,12 @@ void SPCGEditorGraphNodePin::Construct(const FArguments& InArgs, UEdGraphPin* In
 			.Image(this, &SPCGEditorGraphNodePin::GetPinStatusIcon)
 		];
 
-	TSharedRef<SWidget> LabelWidget = GetLabelWidget(GetLabelStyle(InArgs._PinLabelStyle));
+	TSharedRef<SWidget> LabelWidget =
+		SNew(STextBlock)
+		.Text(GetPinLabel())
+		.TextStyle(FAppStyle::Get(), InArgs._PinLabelStyle)
+		.Visibility(GetPinLabelVisibility())
+		.ColorAndOpacity(PinTextColor);
 
 	// Create the widget used for the pin body (status indicator, label, and value)
 	LabelAndValue =
@@ -98,7 +105,7 @@ void SPCGEditorGraphNodePin::Construct(const FArguments& InArgs, UEdGraphPin* In
 	{
 		ExtraPinIconWidget = SNew(SImage)
 			.Image(FAppStyle::GetBrush(ExtraPinIcon))
-			.ColorAndOpacity(this, &SPCGEditorGraphNodePin::GetPinTextColor);
+			.ColorAndOpacity(PinTextColor);
 
 		if (!ExtraPinIconTooltip.IsEmpty())
 		{
@@ -165,7 +172,7 @@ void SPCGEditorGraphNodePin::Construct(const FArguments& InArgs, UEdGraphPin* In
 
 			if (!DoesWidgetHandleSettingEditingEnabled())
 			{
-				ValueBox->SetEnabled(TAttribute<bool>(this, &SPCGEditorGraphNodePin::IsEditingEnabled));
+				ValueBox->SetEnabled(IsEditingEnabled());
 			}
 		}
 
@@ -224,18 +231,18 @@ void SPCGEditorGraphNodePin::Construct(const FArguments& InArgs, UEdGraphPin* In
 	// Set up a hover for pins that is tinted the color of the pin.
 
 	SBorder::Construct(SBorder::FArguments()
-		.BorderImage(this, &SPCGEditorGraphNodePin::GetPinBorder)
-		.BorderBackgroundColor(this, &SPCGEditorGraphNodePin::GetHighlightColor)
+		.BorderImage(SPCGEditorGraphNodePin::GetPinBorder())
+		.BorderBackgroundColor(SPCGEditorGraphNodePin::GetHighlightColor())
 		.OnMouseButtonDown(this, &SPCGEditorGraphNodePin::OnPinNameMouseDown)
 		.Padding(0) // NOTE: This is different from base class implementation
 		[
 			SNew(SBorder)
 			.BorderImage(CachedImg_Pin_DiffOutline)
-			.BorderBackgroundColor(this, &SPCGEditorGraphNodePin::GetPinDiffColor)
+			.BorderBackgroundColor(GetPinDiffColor())
 			.Padding(0) // NOTE: This is different from base class implementation
 			[
 				SNew(SLevelOfDetailBranchNode)
-				.UseLowDetailSlot(this, &SPCGEditorGraphNodePin::UseLowDetailPinNames)
+				.UseLowDetailSlot(UseLowDetailPinNames())
 				.LowDetail()
 				[
 					//@TODO: Try creating a pin-colored line replacement that doesn't measure text / call delegates but still renders
@@ -248,8 +255,9 @@ void SPCGEditorGraphNodePin::Construct(const FArguments& InArgs, UEdGraphPin* In
 			]
 		]);
 
-	TSharedPtr<IToolTip> TooltipWidget = SNew(SToolTip)
-		.Text(this, &SPCGEditorGraphNodePin::GetTooltipText);
+	TSharedPtr<IToolTip> TooltipWidget =
+		SNew(SToolTip)
+		.Text(GetTooltipText());
 
 	SetToolTip(TooltipWidget);
 }
