[PCG] PCGGraph : ChangeTrackingIgnoredClasses
- Tracked changes to the specified classes will be ignored and will not trigger a refresh

diff --git a/Engine/Plugins/PCG/Source/PCG/Private/PCGActorAndComponentMapping.cpp b/Engine/Plugins/PCG/Source/PCG/Private/PCGActorAndComponentMapping.cpp
--- a/Engine/Plugins/PCG/Source/PCG/Private/PCGActorAndComponentMapping.cpp
+++ b/Engine/Plugins/PCG/Source/PCG/Private/PCGActorAndComponentMapping.cpp
@@ -1930,10 +1930,11 @@ void FPCGActorAndComponentMapping::OnSelectionKeyChanged(const FPCGSelectionKey&
 
 		const UObject* ObjectIgnored = nullptr;
 		const UPCGComponent* OriginalComponent = InComponent->GetConstOriginalComponent();
+		EPCGIgnoreChangeOriginReason Reason;
 
-		if (ensure(OriginalComponent) && OriginalComponent->IsIgnoringAnyChangeOrigins(ChangedObjects, ObjectIgnored))
+		if (ensure(OriginalComponent) && OriginalComponent->IsIgnoringAnyChangeOrigins(ChangedObjects, ObjectIgnored, Reason))
 		{
-			PCGGraphExecutionLogging::LogChangeOriginIgnoredForComponent(ObjectIgnored, InComponent);
+			PCGGraphExecutionLogging::LogChangeOriginIgnoredForComponent(ObjectIgnored, InComponent, Reason);
 			return true;
 		}
 
diff --git a/Engine/Plugins/PCG/Source/PCG/Private/PCGComponent.cpp b/Engine/Plugins/PCG/Source/PCG/Private/PCGComponent.cpp
--- a/Engine/Plugins/PCG/Source/PCG/Private/PCGComponent.cpp
+++ b/Engine/Plugins/PCG/Source/PCG/Private/PCGComponent.cpp
@@ -4029,24 +4029,52 @@ bool UPCGComponent::IsIgnoringChangeOrigin(const UObject* InChangeOrigin) const
 
 bool UPCGComponent::IsIgnoringAnyChangeOrigins(const TArrayView<const UObject*> InChangeOrigins, const UObject*& OutFirstObjectFound) const
 {
+	EPCGIgnoreChangeOriginReason Reason;
+	return IsIgnoringAnyChangeOrigins(InChangeOrigins, OutFirstObjectFound, Reason);
+}
+
+bool UPCGComponent::IsIgnoringAnyChangeOrigins(const TArrayView<const UObject*> InChangeOrigins, const UObject*& OutFirstObjectFound, EPCGIgnoreChangeOriginReason& OutReason) const
+{
+	OutReason = EPCGIgnoreChangeOriginReason::None;
+
 	if ((bIgnoreLandscapeTracking || (GetGraph() && GetGraph()->bIgnoreLandscapeTracking)))
 	{
 		const int32 Index = Algo::IndexOfByPredicate(InChangeOrigins, [](const UObject* ChangeOrigin) { return Cast<ALandscapeProxy>(ChangeOrigin); });
 		if (Index != INDEX_NONE)
 		{
 			OutFirstObjectFound = InChangeOrigins[Index];
+			OutReason = bIgnoreLandscapeTracking ? EPCGIgnoreChangeOriginReason::ComponentDisabledLandscapeTracking : EPCGIgnoreChangeOriginReason::GraphDisabledLandscapeTracking;
 			return true;
 		}
 	}
 
-	UE::TReadScopeLock Lock(IgnoredChangeOriginsLock);
-	for (const UObject* ChangeOrigin : InChangeOrigins)
 	{
-		const int32* Counter = IgnoredChangeOriginsToCounters.Find(ChangeOrigin);
-		if (Counter && ensure(*Counter > 0))
+		UE::TReadScopeLock Lock(IgnoredChangeOriginsLock);
+		for (const UObject* ChangeOrigin : InChangeOrigins)
 		{
-			OutFirstObjectFound = ChangeOrigin;
-			return true;
+			const int32* Counter = IgnoredChangeOriginsToCounters.Find(ChangeOrigin);
+			if (Counter && ensure(*Counter > 0))
+			{
+				OutFirstObjectFound = ChangeOrigin;
+				OutReason = EPCGIgnoreChangeOriginReason::GenerationIgnoredObject;
+				return true;
+			}
+		}
+	}
+
+	if (UPCGGraph* Graph = GetGraph())
+	{
+		for (UClass* ChangeTrackingIgnoredClass : Graph->GetChangeTrackingIgnoredClasses())
+		{
+			for (const UObject* ChangeOrigin : InChangeOrigins)
+			{
+				if(ChangeOrigin && ChangeOrigin->IsA(ChangeTrackingIgnoredClass))
+				{
+					OutFirstObjectFound = ChangeOrigin;
+					OutReason = EPCGIgnoreChangeOriginReason::GraphIgnoredClass;
+					return true;
+				}
+			}
 		}
 	}
 	
diff --git a/Engine/Plugins/PCG/Source/PCG/Private/Utils/PCGGraphExecutionLogging.cpp b/Engine/Plugins/PCG/Source/PCG/Private/Utils/PCGGraphExecutionLogging.cpp
--- a/Engine/Plugins/PCG/Source/PCG/Private/Utils/PCGGraphExecutionLogging.cpp
+++ b/Engine/Plugins/PCG/Source/PCG/Private/Utils/PCGGraphExecutionLogging.cpp
@@ -235,7 +235,7 @@ namespace PCGGraphExecutionLogging
 #endif // !(UE_BUILD_SHIPPING || UE_BUILD_TEST) || USE_LOGGING_IN_SHIPPING
 	}
 
-	void LogChangeOriginIgnoredForComponent(const UObject* InObject, const IPCGGraphExecutionSource* InExecutionSource)
+	void LogChangeOriginIgnoredForComponent(const UObject* InObject, const IPCGGraphExecutionSource* InExecutionSource, EPCGIgnoreChangeOriginReason InReason)
 	{
 #if !(UE_BUILD_SHIPPING || UE_BUILD_TEST) || USE_LOGGING_IN_SHIPPING
 		if (!LogEnabled())
@@ -243,10 +243,11 @@ namespace PCGGraphExecutionLogging
 			return;
 		}
 
-		UE_LOG(LogPCG, Log, TEXT("[%s/%s] Change origin ignored: '%s'"),
+		UE_LOG(LogPCG, Log, TEXT("[%s/%s] Change origin ignored: '%s' (%s)"),
 			*PCGLog::GetExecutionSourceName(InExecutionSource),
 			(InExecutionSource && InExecutionSource->GetExecutionState().GetGraph()) ? *InExecutionSource->GetExecutionState().GetGraph()->GetName() : TEXT("MISSINGGRAPH"),
-			InObject ? *InObject->GetName() : TEXT("MISSINGOBJECT"));
+			InObject ? *InObject->GetName() : TEXT("MISSINGOBJECT"),
+			*StaticEnum<EPCGIgnoreChangeOriginReason>()->GetValueAsString(InReason));
 #endif // !(UE_BUILD_SHIPPING || UE_BUILD_TEST) || USE_LOGGING_IN_SHIPPING
 	}
 
diff --git a/Engine/Plugins/PCG/Source/PCG/Private/Utils/PCGGraphExecutionLogging.h b/Engine/Plugins/PCG/Source/PCG/Private/Utils/PCGGraphExecutionLogging.h
--- a/Engine/Plugins/PCG/Source/PCG/Private/Utils/PCGGraphExecutionLogging.h
+++ b/Engine/Plugins/PCG/Source/PCG/Private/Utils/PCGGraphExecutionLogging.h
@@ -11,6 +11,7 @@ class IPCGGraphExecutionSource;
 class UPCGComponent;
 class UPCGGraph;
 enum class EPCGHiGenGrid : uint32;
+enum class EPCGIgnoreChangeOriginReason : uint8;
 struct FPCGContext;
 struct FPCGGraphActiveTask;
 struct FPCGGraphTask;
@@ -35,7 +36,7 @@ namespace PCGGraphExecutionLogging
 
 	void LogExecutionSourceCancellation(const TSet<IPCGGraphExecutionSource*>& CancelledComponents);
 
-	void LogChangeOriginIgnoredForComponent(const UObject* InObject, const IPCGGraphExecutionSource* InExecutionSource);
+	void LogChangeOriginIgnoredForComponent(const UObject* InObject, const IPCGGraphExecutionSource* InExecutionSource, EPCGIgnoreChangeOriginReason InReason);
 
 	void LogGraphExecuteFrameFinished();
 
diff --git a/Engine/Plugins/PCG/Source/PCG/Public/PCGComponent.h b/Engine/Plugins/PCG/Source/PCG/Public/PCGComponent.h
--- a/Engine/Plugins/PCG/Source/PCG/Public/PCGComponent.h
+++ b/Engine/Plugins/PCG/Source/PCG/Public/PCGComponent.h
@@ -87,6 +87,16 @@ enum class EPCGComponentDirtyFlag : uint8
 };
 ENUM_CLASS_FLAGS(EPCGComponentDirtyFlag);
 
+UENUM()
+enum class EPCGIgnoreChangeOriginReason : uint8
+{
+	None = 0,
+	ComponentDisabledLandscapeTracking,
+	GraphDisabledLandscapeTracking,
+	GenerationIgnoredObject,
+	GraphIgnoredClass
+};
+
 namespace PCGTestsCommon
 {
 	struct FTestData;
@@ -492,6 +502,7 @@ public:
 
 	UE_API bool IsIgnoringChangeOrigin(const UObject* InChangeOrigin) const;
 	UE_API bool IsIgnoringAnyChangeOrigins(const TArrayView<const UObject*> InChangeOrigins, const UObject*& OutFirstObjectFound) const;
+	UE_API bool IsIgnoringAnyChangeOrigins(const TArrayView<const UObject*> InChangeOrigins, const UObject*& OutFirstObjectFound, EPCGIgnoreChangeOriginReason& OutReason) const;
 	UE_API void ResetIgnoredChangeOrigins(bool bLogIfAnyPresent);
 #endif
 
diff --git a/Engine/Plugins/PCG/Source/PCG/Public/PCGGraph.h b/Engine/Plugins/PCG/Source/PCG/Public/PCGGraph.h
--- a/Engine/Plugins/PCG/Source/PCG/Public/PCGGraph.h
+++ b/Engine/Plugins/PCG/Source/PCG/Public/PCGGraph.h
@@ -468,6 +468,8 @@ public:
 
 	/** Overridable function for child classes to have a graph-wide node title sanitization when the title of a node changes. Returns true if the title have changed. */
 	virtual bool SanitizeNodeTitle(FName& InOutTitle, const UPCGNode* InNode) const { return false; }
+
+	const TArray<TSubclassOf<UObject>>& GetChangeTrackingIgnoredClasses() const { return ChangeTrackingIgnoredClasses; }
 #endif
 
 #if WITH_EDITOR
@@ -602,6 +604,9 @@ protected:
 	*/
 	UPROPERTY(EditAnywhere, Category = Debug, meta=(EditCondition="!bIsStandaloneGraph", EditConditionHides))
 	bool bDebugFlagAppliesToIndividualComponents = true;
+
+	UPROPERTY(EditAnywhere, Category = ChangeTracking, DisplayName = "Ignored Classes")
+	TArray<TSubclassOf<UObject>> ChangeTrackingIgnoredClasses;
 #endif // WITH_EDITORONLY_DATA
 
 	/**
