[Backout] - CL46461762 - broke graph parameter nodes

#rnx
Original CL Desc
-----------------------------------------------------------------
[PCG] Remove most of the pin widget dynamic lambdas which can simply be baked onto the pin and updated on refresh.

#rnx

diff --git a/Engine/Plugins/PCG/Source/PCGEditor/Private/Pins/SPCGEditorGraphNodePin.cpp b/Engine/Plugins/PCG/Source/PCGEditor/Private/Pins/SPCGEditorGraphNodePin.cpp
--- a/Engine/Plugins/PCG/Source/PCGEditor/Private/Pins/SPCGEditorGraphNodePin.cpp
+++ b/Engine/Plugins/PCG/Source/PCGEditor/Private/Pins/SPCGEditorGraphNodePin.cpp
@@ -14,7 +14,6 @@
 #include "Algo/AnyOf.h"
 #include "Schema/PCGEditorGraphSchema.h"
 #include "Widgets/SToolTip.h"
-#include "Widgets/Images/SLayeredImage.h"
 #include "Widgets/Input/SButton.h"
 
 void SPCGEditorGraphNodePin::Construct(const FArguments& InArgs, UEdGraphPin* InPin)
@@ -26,7 +25,7 @@ void SPCGEditorGraphNodePin::Construct(const FArguments& InArgs, UEdGraphPin* In
 	bUsePinColorForText = InArgs._UsePinColorForText;
 	this->SetCursor(EMouseCursor::Default);
 
-	SetVisibility(GetPinVisiblity());
+	SetVisibility(MakeAttributeSP(this, &SPCGEditorGraphNodePin::GetPinVisiblity));
 
 	GraphPinObj = InPin;
 	check(GraphPinObj != NULL);
@@ -41,8 +40,6 @@ void SPCGEditorGraphNodePin::Construct(const FArguments& InArgs, UEdGraphPin* In
 		);
 
 	const bool bIsInput = (GetDirection() == EGPD_Input);
-	const FSlateColor PinColor = GetPinColor();
-	const FSlateColor PinTextColor = GetPinTextColor();
 
 	// A small marker to indicate the pin is required for the node to be executed.
 	TSharedPtr<SImage> RequiredPinIconWidget;
@@ -57,14 +54,15 @@ void SPCGEditorGraphNodePin::Construct(const FArguments& InArgs, UEdGraphPin* In
 		RequiredPinIconWidget =
 			SNew(SImage)
 			.Image(bDisplayPinMarker ? RequiredPinMarkerIcon : FAppStyle::GetNoBrush())
-			.ColorAndOpacity(PinColor);
+			.ColorAndOpacity(this, &SPCGEditorGraphNodePin::GetPinColor);
 	}
 
 	// Create the pin icon widget
-	TSharedRef<SWidget> PinWidgetRef =
-		SNew(SLayeredImage, GetSecondaryPinIcon(), GetSecondaryPinColor())
-		.Image(MakeAttributeSP(this, &SPCGEditorGraphNodePin::GetPinIcon))
-		.ColorAndOpacity(PinColor);
+	TSharedRef<SWidget> PinWidgetRef = SPinTypeSelector::ConstructPinTypeImage(
+		MakeAttributeSP(this, &SPCGEditorGraphNodePin::GetPinIcon),
+		MakeAttributeSP(this, &SPCGEditorGraphNodePin::GetPinColor),
+		MakeAttributeSP(this, &SPCGEditorGraphNodePin::GetSecondaryPinIcon),
+		MakeAttributeSP(this, &SPCGEditorGraphNodePin::GetSecondaryPinColor));
 	PinImage = PinWidgetRef;
 
 	PinWidgetRef->SetCursor(
@@ -86,12 +84,7 @@ void SPCGEditorGraphNodePin::Construct(const FArguments& InArgs, UEdGraphPin* In
 			.Image(this, &SPCGEditorGraphNodePin::GetPinStatusIcon)
 		];
 
-	TSharedRef<SWidget> LabelWidget =
-		SNew(STextBlock)
-		.Text(GetPinLabel())
-		.TextStyle(FAppStyle::Get(), InArgs._PinLabelStyle)
-		.Visibility(GetPinLabelVisibility())
-		.ColorAndOpacity(PinTextColor);
+	TSharedRef<SWidget> LabelWidget = GetLabelWidget(GetLabelStyle(InArgs._PinLabelStyle));
 
 	// Create the widget used for the pin body (status indicator, label, and value)
 	LabelAndValue =
@@ -105,7 +98,7 @@ void SPCGEditorGraphNodePin::Construct(const FArguments& InArgs, UEdGraphPin* In
 	{
 		ExtraPinIconWidget = SNew(SImage)
 			.Image(FAppStyle::GetBrush(ExtraPinIcon))
-			.ColorAndOpacity(PinTextColor);
+			.ColorAndOpacity(this, &SPCGEditorGraphNodePin::GetPinTextColor);
 
 		if (!ExtraPinIconTooltip.IsEmpty())
 		{
@@ -172,7 +165,7 @@ void SPCGEditorGraphNodePin::Construct(const FArguments& InArgs, UEdGraphPin* In
 
 			if (!DoesWidgetHandleSettingEditingEnabled())
 			{
-				ValueBox->SetEnabled(IsEditingEnabled());
+				ValueBox->SetEnabled(TAttribute<bool>(this, &SPCGEditorGraphNodePin::IsEditingEnabled));
 			}
 		}
 
@@ -231,18 +224,18 @@ void SPCGEditorGraphNodePin::Construct(const FArguments& InArgs, UEdGraphPin* In
 	// Set up a hover for pins that is tinted the color of the pin.
 
 	SBorder::Construct(SBorder::FArguments()
-		.BorderImage(SPCGEditorGraphNodePin::GetPinBorder())
-		.BorderBackgroundColor(SPCGEditorGraphNodePin::GetHighlightColor())
+		.BorderImage(this, &SPCGEditorGraphNodePin::GetPinBorder)
+		.BorderBackgroundColor(this, &SPCGEditorGraphNodePin::GetHighlightColor)
 		.OnMouseButtonDown(this, &SPCGEditorGraphNodePin::OnPinNameMouseDown)
 		.Padding(0) // NOTE: This is different from base class implementation
 		[
 			SNew(SBorder)
 			.BorderImage(CachedImg_Pin_DiffOutline)
-			.BorderBackgroundColor(GetPinDiffColor())
+			.BorderBackgroundColor(this, &SPCGEditorGraphNodePin::GetPinDiffColor)
 			.Padding(0) // NOTE: This is different from base class implementation
 			[
 				SNew(SLevelOfDetailBranchNode)
-				.UseLowDetailSlot(UseLowDetailPinNames())
+				.UseLowDetailSlot(this, &SPCGEditorGraphNodePin::UseLowDetailPinNames)
 				.LowDetail()
 				[
 					//@TODO: Try creating a pin-colored line replacement that doesn't measure text / call delegates but still renders
@@ -255,9 +248,8 @@ void SPCGEditorGraphNodePin::Construct(const FArguments& InArgs, UEdGraphPin* In
 			]
 		]);
 
-	TSharedPtr<IToolTip> TooltipWidget =
-		SNew(SToolTip)
-		.Text(GetTooltipText());
+	TSharedPtr<IToolTip> TooltipWidget = SNew(SToolTip)
+		.Text(this, &SPCGEditorGraphNodePin::GetTooltipText);
 
 	SetToolTip(TooltipWidget);
 }
