[PCG] IPCGEditorModule : add a permission mode to limit functionality in specific contexts

#rnx

diff --git a/Engine/Plugins/PCG/Source/PCG/Private/PCGComponent.cpp b/Engine/Plugins/PCG/Source/PCG/Private/PCGComponent.cpp
--- a/Engine/Plugins/PCG/Source/PCG/Private/PCGComponent.cpp
+++ b/Engine/Plugins/PCG/Source/PCG/Private/PCGComponent.cpp
@@ -48,6 +48,7 @@
 #include "EditorActorFolders.h"
 #include "ScopedTransaction.h"
 #include "Editor/EditorEngine.h"
+#include "Editor/IPCGEditorModule.h"
 #include "Editor/Transactor.h"
 #include "UObject/ObjectSaveContext.h"
 #include "WorldPartition/WorldPartition.h"
@@ -2302,6 +2303,14 @@ void UPCGComponent::BeginDestroy()
 
 bool UPCGComponent::IsEditorOnly() const
 {
+#if WITH_EDITOR
+	// Only allow non-editor components when running with all permissions
+	if (IPCGEditorModule* EditorModule = IPCGEditorModule::Get(); EditorModule && EditorModule->GetPermissionMode() != EPCGEditorPermissionMode::All)
+	{
+		return true;
+	}
+#endif
+
 	return Super::IsEditorOnly() || (GraphInstance && GraphInstance->IsEditorOnly());
 }
 
diff --git a/Engine/Plugins/PCG/Source/PCG/Public/Editor/IPCGEditorModule.h b/Engine/Plugins/PCG/Source/PCG/Public/Editor/IPCGEditorModule.h
--- a/Engine/Plugins/PCG/Source/PCG/Public/Editor/IPCGEditorModule.h
+++ b/Engine/Plugins/PCG/Source/PCG/Public/Editor/IPCGEditorModule.h
@@ -13,6 +13,13 @@ class IPCGBaseSubsystem;
 struct FPCGStackContext;
 class UPCGComponent;
 
+enum class EPCGEditorPermissionMode : uint8
+{
+	All = 0,
+	Limited,
+	None
+};
+
 class IPCGEditorModule : public IModuleInterface
 {
 public:
@@ -26,6 +33,8 @@ public:
 	virtual const FPCGNodeVisualLogs& GetNodeVisualLogs() const = 0;
 	virtual FPCGNodeVisualLogs& GetNodeVisualLogsMutable() = 0;
 	virtual bool CanSelectPartitionActors() const = 0;
+	virtual void SetPermissionMode(EPCGEditorPermissionMode InPermissionMode) = 0;
+	virtual EPCGEditorPermissionMode GetPermissionMode() const = 0;
 #endif
 
 protected:
diff --git a/Engine/Plugins/PCG/Source/PCGEditor/Public/PCGEditorModule.h b/Engine/Plugins/PCG/Source/PCGEditor/Public/PCGEditorModule.h
--- a/Engine/Plugins/PCG/Source/PCGEditor/Public/PCGEditorModule.h
+++ b/Engine/Plugins/PCG/Source/PCGEditor/Public/PCGEditorModule.h
@@ -37,6 +37,8 @@ public:
 	virtual const FPCGNodeVisualLogs& GetNodeVisualLogs() const override { return NodeVisualLogs; }
 	virtual FPCGNodeVisualLogs& GetNodeVisualLogsMutable() override { return NodeVisualLogs; }
 	virtual bool CanSelectPartitionActors() const override;
+	virtual void SetPermissionMode(EPCGEditorPermissionMode InPermissionMode) override { PermissionMode = InPermissionMode; }
+	virtual EPCGEditorPermissionMode GetPermissionMode() const { return PermissionMode; }
 
 protected:
 	virtual void OnScheduleGraph(const FPCGStackContext& StackContext) override;
@@ -92,4 +94,6 @@ protected:
 	TSet<FPCGStackSharedPtr> ExecutedStacks;
 	TMap<FPCGSoftGraphExecutionSource, TSet<FPCGStackSharedPtr>> ExecutedStacksPerSource;
 	mutable FRWLock ExecutedStacksLock;
+
+	EPCGEditorPermissionMode PermissionMode = EPCGEditorPermissionMode::All;
 };
