[PCG] Optimization: disable CRC calculation when graph cache is disabled.

#rnx

diff --git a/Engine/Plugins/PCG/Source/PCG/Private/Graph/PCGGraphCache.h b/Engine/Plugins/PCG/Source/PCG/Private/Graph/PCGGraphCache.h
--- a/Engine/Plugins/PCG/Source/PCG/Private/Graph/PCGGraphCache.h
+++ b/Engine/Plugins/PCG/Source/PCG/Private/Graph/PCGGraphCache.h
@@ -70,7 +70,7 @@ public:
 
 	static bool IsGatheringCacheSize();
 
-	bool IsEnabled() const;
+	virtual bool IsEnabled() const override;
 
 	int32 GetMemoryBudgetMB() const;
 
diff --git a/Engine/Plugins/PCG/Source/PCG/Private/PCGContext.cpp b/Engine/Plugins/PCG/Source/PCG/Private/PCGContext.cpp
--- a/Engine/Plugins/PCG/Source/PCG/Private/PCGContext.cpp
+++ b/Engine/Plugins/PCG/Source/PCG/Private/PCGContext.cpp
@@ -645,6 +645,12 @@ void FPCGContext::InitializeGraphExecutor(FPCGContext* InContext)
 	InContext->GraphExecutor = GraphExecutor;
 }
 
+bool FPCGContext::IsCacheEnabled()
+{
+	TSharedPtr<FPCGGraphExecutor> Executor = GraphExecutor.Pin();
+	return Executor && Executor->GetCache().IsEnabled();
+}
+
 void FPCGContext::StoreInCache(const FPCGStoreInCacheParams& Params, const FPCGDataCollection& InOutput)
 {
 	if (TSharedPtr<FPCGGraphExecutor> Executor = GraphExecutor.Pin())
diff --git a/Engine/Plugins/PCG/Source/PCG/Private/PCGData.cpp b/Engine/Plugins/PCG/Source/PCG/Private/PCGData.cpp
--- a/Engine/Plugins/PCG/Source/PCG/Private/PCGData.cpp
+++ b/Engine/Plugins/PCG/Source/PCG/Private/PCGData.cpp
@@ -590,6 +590,16 @@ void FPCGDataCollection::ComputeCrcs(bool bFullDataCrc)
 	}
 }
 
+void FPCGDataCollection::SetZeroCrcs()
+{
+	DataCrcs.Empty(TaggedData.Num());
+
+	for (int Index = 0; Index < TaggedData.Num(); ++Index)
+	{
+		DataCrcs.Emplace(0);
+	}
+}
+
 void FPCGDataCollection::AddData(const FPCGTaggedData& InData, const FPCGCrc& InDataCrc)
 {
 	TaggedData.Add(InData);
diff --git a/Engine/Plugins/PCG/Source/PCG/Private/PCGElement.cpp b/Engine/Plugins/PCG/Source/PCG/Private/PCGElement.cpp
--- a/Engine/Plugins/PCG/Source/PCG/Private/PCGElement.cpp
+++ b/Engine/Plugins/PCG/Source/PCG/Private/PCGElement.cpp
@@ -383,6 +383,12 @@ void IPCGElement::PreExecutePrimaryLoopElement(FPCGContext* Context, const UPCGS
 		Context->InputData.TaggedData[DataIndex].OriginalIndex = DataIndex;
 	}
 
+	// No need to split data if we're not using the cache.
+	if (!Context->IsCacheEnabled())
+	{
+		return;
+	}
+
 	TArray<FPCGDataCollection> PrimaryDataCollections;
 	FPCGDataCollection OtherData;
 	if (!PCGElementHelpers::SplitDataPerPrimaryPin(Settings, Context->InputData, ExecutionLoopMode(Settings), PrimaryDataCollections, OtherData))
@@ -413,7 +419,7 @@ void IPCGElement::PreExecutePrimaryLoopElement(FPCGContext* Context, const UPCGS
 		GetDependenciesCrc(FPCGGetDependenciesCrcParams(&SubCollection, Settings, Context->ExecutionSource.Get()), CacheParams.Crc);
 
 		FPCGDataCollection SubCollectionOutput;
-		if(Context->GetFromCache(CacheParams, SubCollectionOutput))
+		if (Context->GetFromCache(CacheParams, SubCollectionOutput))
 		{
 			// Found a match in the cache, add it to the output, and remove the matching inputs.
 			// Note that in the input part we'll take only the data present in the PrimaryDataCollection here, as we will reintroduce only those then.
@@ -563,6 +569,7 @@ void IPCGElement::PostExecute(FPCGContext* Context) const
 	}
 
 	// Output data Crc
+	if (Context->IsCacheEnabled())
 	{
 		TRACE_CPUPROFILER_EVENT_SCOPE(IPCGElement::PostExecute::CRC);
 
@@ -574,6 +581,12 @@ void IPCGElement::PostExecute(FPCGContext* Context) const
 		// Compute Crc from output data which will include output pin labels.
 		Context->OutputData.ComputeCrcs(bShouldComputeFullOutputDataCrc);
 	}
+	else
+	{
+		// We check if Crcs are valid in some parts of the code which can be useful for catching bugs, so add valid 0 Crcs here.
+		// Assumption is that these Crcs will be unused if caching is disabled.
+		Context->OutputData.SetZeroCrcs();
+	}
 
 #if WITH_EDITOR
 	const bool bHasErrorsOrWarnings = Context->Node && Context->HasVisualLogs();
@@ -664,7 +677,7 @@ void IPCGElement::PostExecutePrimaryLoopElement(FPCGContext* Context, const UPCG
 {
 	check(Context);
 
-	if (!Settings)
+	if (!Settings || !Context->IsCacheEnabled())
 	{
 		return;
 	}
@@ -1196,7 +1209,7 @@ PRAGMA_ENABLE_DEPRECATION_WARNINGS
 
 EPCGCachingStatus IPCGElement::RetrieveResultsFromCache(IPCGGraphCache* Cache, const UPCGNode* Node, const FPCGDataCollection& Input, IPCGGraphExecutionSource* ExecutionSource, FPCGDataCollection& Output, FPCGCrc* OutCrc) const
 {
-	if (!Cache)
+	if (!Cache || !Cache->IsEnabled())
 	{
 		return EPCGCachingStatus::NotInCache;
 	}
@@ -1217,7 +1230,7 @@ EPCGCachingStatus IPCGElement::RetrieveResultsFromCache(IPCGGraphCache* Cache, c
 		}
 	}
 
-	if(Params.Crc.IsValid() && Cache->GetFromCache(Params, Output))
+	if (Params.Crc.IsValid() && Cache->GetFromCache(Params, Output))
 	{
 		return EPCGCachingStatus::Cached;
 	}
diff --git a/Engine/Plugins/PCG/Source/PCG/Public/Graph/IPCGGraphCache.h b/Engine/Plugins/PCG/Source/PCG/Public/Graph/IPCGGraphCache.h
--- a/Engine/Plugins/PCG/Source/PCG/Public/Graph/IPCGGraphCache.h
+++ b/Engine/Plugins/PCG/Source/PCG/Public/Graph/IPCGGraphCache.h
@@ -32,6 +32,7 @@ struct FPCGStoreInCacheParams
 class IPCGGraphCache
 {
 public:
+	virtual bool IsEnabled() const = 0;
 	virtual bool GetFromCache(const FPCGGetFromCacheParams& Params, FPCGDataCollection& OutCollection) const = 0;
 	virtual void StoreInCache(const FPCGStoreInCacheParams& Params, const FPCGDataCollection& InCollection) = 0;
-};
\ No newline at end of file
+};
diff --git a/Engine/Plugins/PCG/Source/PCG/Public/PCGContext.h b/Engine/Plugins/PCG/Source/PCG/Public/PCGContext.h
--- a/Engine/Plugins/PCG/Source/PCG/Public/PCGContext.h
+++ b/Engine/Plugins/PCG/Source/PCG/Public/PCGContext.h
@@ -274,6 +274,7 @@ UE_API PRAGMA_ENABLE_DEPRECATION_WARNINGS
 	UE_API bool GetOutputData(FPCGTaskId InTaskId, FPCGDataCollection& OutData);
 	UE_API void ClearOutputData(FPCGTaskId InTaskId);
 
+	UE_API bool IsCacheEnabled();
 	UE_API void StoreInCache(const FPCGStoreInCacheParams& Params, const FPCGDataCollection& InOutput);
 	UE_API bool GetFromCache(const FPCGGetFromCacheParams& Params, FPCGDataCollection& OutOutput) const;
 
diff --git a/Engine/Plugins/PCG/Source/PCG/Public/PCGData.h b/Engine/Plugins/PCG/Source/PCG/Public/PCGData.h
--- a/Engine/Plugins/PCG/Source/PCG/Public/PCGData.h
+++ b/Engine/Plugins/PCG/Source/PCG/Public/PCGData.h
@@ -317,6 +317,7 @@ struct FPCGDataCollection
 
 	/** Computes CRCs for all data items. */
 	UE_API void ComputeCrcs(bool bFullDataCrc);
+	UE_API void SetZeroCrcs();
 
 	/** Add data and CRC to collection. */
 	UE_API void AddData(const FPCGTaggedData& InData, const FPCGCrc& InDataCrc);
