[PCG] Broadcast cosmetic change on preview platform changes so that node visually refreshes.

#rnx

diff --git a/Engine/Plugins/PCG/Source/PCG/Private/Elements/ControlFlow/PCGPlatformSwitch.cpp b/Engine/Plugins/PCG/Source/PCG/Private/Elements/ControlFlow/PCGPlatformSwitch.cpp
--- a/Engine/Plugins/PCG/Source/PCG/Private/Elements/ControlFlow/PCGPlatformSwitch.cpp
+++ b/Engine/Plugins/PCG/Source/PCG/Private/Elements/ControlFlow/PCGPlatformSwitch.cpp
@@ -124,6 +124,34 @@ FString UPCGPlatformSwitchSettings::GetAdditionalTitleInformation() const
 }
 
 #if WITH_EDITOR
+void UPCGPlatformSwitchSettings::PostInitProperties()
+{
+	Super::PostInitProperties();
+
+	SetupPreviewPlatformChangeEvents();
+}
+
+void UPCGPlatformSwitchSettings::PostLoad()
+{
+	Super::PostLoad();
+
+	SetupPreviewPlatformChangeEvents();
+}
+
+void UPCGPlatformSwitchSettings::PostDuplicate(bool bDuplicateForPIE)
+{
+	Super::PostDuplicate(bDuplicateForPIE);
+
+	SetupPreviewPlatformChangeEvents();
+}
+
+void UPCGPlatformSwitchSettings::PostEditImport()
+{
+	Super::PostEditImport();
+
+	SetupPreviewPlatformChangeEvents();
+}
+
 void UPCGPlatformSwitchSettings::PreSave(FObjectPreSaveContext ObjectSaveContext)
 {
 	Super::PreSave(ObjectSaveContext);
@@ -134,6 +162,13 @@ void UPCGPlatformSwitchSettings::PreSave(FObjectPreSaveContext ObjectSaveContext
 	}
 }
 
+void UPCGPlatformSwitchSettings::BeginDestroy()
+{
+	TearDownPreviewPlatformChangeEvents();
+
+	Super::BeginDestroy();
+}
+
 FText UPCGPlatformSwitchSettings::GetNodeTooltipText() const
 {
 	return LOCTEXT("NodeTooltip", "Statically activates/deactivates output pins based on the current platform or platform group (desktop/console/mobile), useful for culling downstream nodes.");
@@ -152,6 +187,50 @@ EPCGChangeType UPCGPlatformSwitchSettings::GetChangeTypeForProperty(const FName&
 	return ChangeType;
 }
 
+void UPCGPlatformSwitchSettings::SetupPreviewPlatformChangeEvents()
+{
+	if (GEditor)
+	{
+		if (!PreviewPlatformChangedHandle.IsValid())
+		{
+			PreviewPlatformChangedHandle = GEditor->OnPreviewPlatformChanged().AddUObject(this, &UPCGPlatformSwitchSettings::OnPreviewPlatformChanged);
+		}
+
+		if (!PreviewFeatureLevelChangedHandle.IsValid())
+		{
+			PreviewFeatureLevelChangedHandle = GEditor->OnPreviewFeatureLevelChanged().AddUObject(this, &UPCGPlatformSwitchSettings::OnPreviewFeatureLevelChanged);
+		}
+	}
+}
+
+void UPCGPlatformSwitchSettings::TearDownPreviewPlatformChangeEvents()
+{
+	if (GEditor)
+	{
+		if (PreviewPlatformChangedHandle.IsValid())
+		{
+			GEditor->OnPreviewPlatformChanged().Remove(PreviewPlatformChangedHandle);
+		}
+
+		if (PreviewFeatureLevelChangedHandle.IsValid())
+		{
+			GEditor->OnPreviewFeatureLevelChanged().Remove(PreviewFeatureLevelChangedHandle);
+		}
+	}
+}
+
+void UPCGPlatformSwitchSettings::OnPreviewPlatformChanged()
+{
+	// We don't broadcast structural from here. The PCGSubsystem will refresh all runtime gen components when preview platform changes.
+	OnSettingsChangedDelegate.Broadcast(this, EPCGChangeType::Cosmetic);
+}
+
+void UPCGPlatformSwitchSettings::OnPreviewFeatureLevelChanged(ERHIFeatureLevel::Type)
+{
+	// We don't broadcast structural from here. The PCGSubsystem will refresh all runtime gen components when preview platform changes.
+	OnSettingsChangedDelegate.Broadcast(this, EPCGChangeType::Cosmetic);
+}
+
 TArray<FName> UPCGPlatformSwitchSettings::GetPlatformOptions() const
 {
 	const TArray<const FDataDrivenPlatformInfo*>& SortedPlatforms = FDataDrivenPlatformInfoRegistry::GetSortedPlatformInfos(EPlatformInfoType::TruePlatformsOnly);
diff --git a/Engine/Plugins/PCG/Source/PCG/Public/Elements/ControlFlow/PCGPlatformSwitch.h b/Engine/Plugins/PCG/Source/PCG/Public/Elements/ControlFlow/PCGPlatformSwitch.h
--- a/Engine/Plugins/PCG/Source/PCG/Public/Elements/ControlFlow/PCGPlatformSwitch.h
+++ b/Engine/Plugins/PCG/Source/PCG/Public/Elements/ControlFlow/PCGPlatformSwitch.h
@@ -6,6 +6,8 @@
 
 #include "PCGPlatformSwitch.generated.h"
 
+namespace ERHIFeatureLevel { enum Type : int; }
+
 /** Statically activates/deactivates output pins based on the current platform/platform group. */
 UCLASS(BlueprintType, ClassGroup = (Procedural), meta=(Keywords = "if branch platform"))
 class UPCGPlatformSwitchSettings : public UPCGControlFlowSettings
@@ -15,7 +17,12 @@ class UPCGPlatformSwitchSettings : public UPCGControlFlowSettings
 public:
 	//~ Begin UObject interface
 #if WITH_EDITOR
+	virtual void PostInitProperties();
+	virtual void PostLoad() override;
+	virtual void PostDuplicate(bool bDuplicateForPIE) override;
+	virtual void PostEditImport() override;
 	virtual void PreSave(FObjectPreSaveContext ObjectSaveContext) override;
+	virtual void BeginDestroy() override;
 #endif
 	//~ End UObject interface
 
@@ -42,6 +49,13 @@ protected:
 	virtual FPCGElementPtr CreateElement() const override;
 	//~ End UPCGSettings interface
 
+#if WITH_EDITOR
+	void SetupPreviewPlatformChangeEvents();
+	void TearDownPreviewPlatformChangeEvents();
+	void OnPreviewPlatformChanged();
+	void OnPreviewFeatureLevelChanged(ERHIFeatureLevel::Type);
+#endif
+
 public:
 	/** Get the users selection, sanitized to remove None and duplicates. */
 	TArray<FName> GetSanitizedPlatforms() const;
@@ -71,6 +85,12 @@ protected:
 	/** The current platform group used in cooked builds. */
 	UPROPERTY()
 	FName CookedPlatformGroup;
+
+private:
+#if WITH_EDITOR
+	FDelegateHandle PreviewPlatformChangedHandle;
+	FDelegateHandle PreviewFeatureLevelChangedHandle;
+#endif
 };
 
 class FPCGPlatformSwitchElement : public IPCGElement
