[PCG] Add new 'Primitive Data' spawner type to Static Mesh Spawner.

[RN] Added new static mesh spawner mode that takes a primitive table and allows overriding primitive properties when executing on GPU.

diff --git a/Engine/Plugins/PCG/Shaders/Private/BuiltInKernels/PCGCountUniqueAttributeValues.usf b/Engine/Plugins/PCG/Shaders/Private/BuiltInKernels/PCGCountUniqueAttributeValues.usf
--- a/Engine/Plugins/PCG/Shaders/Private/BuiltInKernels/PCGCountUniqueAttributeValues.usf
+++ b/Engine/Plugins/PCG/Shaders/Private/BuiltInKernels/PCGCountUniqueAttributeValues.usf
@@ -69,7 +69,8 @@ void PCGCountUniqueAttributeValuesCS(uint3 GroupId : SV_GroupID, uint GroupIndex
 
 		const int AttributeToCountId = CountUniqueValues_GetAttributeToCountId();
 
-		const int AttributeValue = In_GetStringKey(In_DataIndex, In_ElementIndex, AttributeToCountId);
+		// Get the value of the attribute. Assumes int type currently (StringKeys are stored as ints).
+		const int AttributeValue = In_GetInt(In_DataIndex, In_ElementIndex, AttributeToCountId);
 
 		if (AttributeValue < 0)
 		{
diff --git a/Engine/Plugins/PCG/Shaders/Private/Elements/PCGStaticMeshSpawner.usf b/Engine/Plugins/PCG/Shaders/Private/Elements/PCGStaticMeshSpawner.usf
--- a/Engine/Plugins/PCG/Shaders/Private/Elements/PCGStaticMeshSpawner.usf
+++ b/Engine/Plugins/PCG/Shaders/Private/Elements/PCGStaticMeshSpawner.usf
@@ -48,9 +48,9 @@ void PCGStaticMeshSpawnerCS(uint3 GroupId : SV_GroupID, uint GroupIndex : SV_Gro
 		
 		if (SelectorAttributeId != (uint)-1)
 		{
-			// Select using string key
-			const int MeshStringKey = In_GetStringKey(In_DataIndex, In_ElementIndex, SelectorAttributeId);
-			PrimitiveIndex = SMSpawner_GetPrimitiveIndexFromStringKey(MeshStringKey);
+			// Primitive ID is either the stringkey of the mesh path, or an index into a primitive data table.
+			const int PrimitiveId = In_GetInt(In_DataIndex, In_ElementIndex, SelectorAttributeId);
+			PrimitiveIndex = SMSpawner_GetPrimitiveIndexFromPrimitiveId(PrimitiveId);
 		}
 		else
 		{
diff --git a/Engine/Plugins/PCG/Source/PCG/Private/Compute/BuiltInKernels/PCGCountUniqueAttributeValuesKernel.cpp b/Engine/Plugins/PCG/Source/PCG/Private/Compute/BuiltInKernels/PCGCountUniqueAttributeValuesKernel.cpp
--- a/Engine/Plugins/PCG/Source/PCG/Private/Compute/BuiltInKernels/PCGCountUniqueAttributeValuesKernel.cpp
+++ b/Engine/Plugins/PCG/Source/PCG/Private/Compute/BuiltInKernels/PCGCountUniqueAttributeValuesKernel.cpp
@@ -73,11 +73,12 @@ bool UPCGCountUniqueAttributeValuesKernel::IsKernelDataValid(FPCGContext* InCont
 				return false;
 			}
 
-			if (AttributeDesc.GetAttributeKey().GetType() != EPCGKernelAttributeType::StringKey)
+			if (AttributeDesc.GetAttributeKey().GetType() != EPCGKernelAttributeType::Int
+				&& AttributeDesc.GetAttributeKey().GetType() != EPCGKernelAttributeType::StringKey)
 			{
-				// Attribute value counting only currently supported for attributes of type StringKey.
+				// Attribute value counting only currently supported for attributes of type Int or StringKey.
 				PCG_KERNEL_VALIDATION_ERR(Context, GetSettings(), FText::Format(
-					LOCTEXT("AttributeTypeInvalid", "Cannot count values for attribute '{0}', only attributes of type String Key are currently supported."),
+					LOCTEXT("AttributeTypeInvalid", "Cannot count values for attribute '{0}', only attributes of type Int or String Key are currently supported."),
 					FText::FromName(AttributeName)));
 
 				return false;
@@ -109,17 +110,35 @@ TSharedPtr<const FPCGDataCollectionDesc> UPCGCountUniqueAttributeValuesKernel::C
 	TSharedPtr<FPCGDataCollectionDesc> OutDataDesc = FPCGDataCollectionDesc::MakeShared();
 	TArray<FPCGDataDesc>& OutDataDescs = OutDataDesc->GetDataDescriptionsMutable();
 
-	const int32 AttributeId = InBinding->GetAttributeId(AttributeName, EPCGKernelAttributeType::StringKey);
+	const int32 StringKeyAttributeId = InBinding->GetAttributeId(AttributeName, EPCGKernelAttributeType::StringKey);
+	const bool bAttributePresent = StringKeyAttributeId != INDEX_NONE || InBinding->GetAttributeId(AttributeName, EPCGKernelAttributeType::Int) != INDEX_NONE;
 
-	if (AttributeId != INDEX_NONE)
+	if (bAttributePresent)
 	{
-		// Compute a set of (value, count) pairs across all data.
-		TArray<int32> UniqueStringKeyValues;
-		InputDesc->GetUniqueStringKeyValues(AttributeId, UniqueStringKeyValues);
+		int32 ElementCount = INDEX_NONE;
 
-		// If the highest string key value is 3, we'll allocate 4 counters. String key value 0 is reserved for empty/null string.
-		const int32* MaxStringKeyValue = Algo::MaxElement(UniqueStringKeyValues);
-		const int32 ElementCount = (MaxStringKeyValue ? *MaxStringKeyValue : 0) + 1;
+		// If counted attribute is a stringkey, use max string key value as num elements.
+		if (StringKeyAttributeId != INDEX_NONE)
+		{
+			TArray<int32> UniqueStringKeyValues;
+			InputDesc->GetUniqueStringKeyValues(StringKeyAttributeId, UniqueStringKeyValues);
+
+			// If the highest string key value is 3, we'll allocate 4 counters. String key value 0 is reserved for empty/null string.
+			const int32* MaxStringKeyValue = Algo::MaxElement(UniqueStringKeyValues);
+			ElementCount = (MaxStringKeyValue ? *MaxStringKeyValue : 0) + 1;
+		}
+		else
+		{
+			if (MaxIntValue < 0)
+			{
+				UE_LOG(LogPCG, Warning, TEXT("PCGCountUniqueAttributeValuesKernel: number of unique values (%d) not set."), MaxIntValue);
+				return OutDataDesc;
+			}
+
+			ElementCount = MaxIntValue + 1;
+		}
+
+		check(ElementCount > 0);
 
 		// Base data description object, copied N times below to populate the final description.
 		FPCGDataDesc CountDataDesc;
diff --git a/Engine/Plugins/PCG/Source/PCG/Private/Compute/BuiltInKernels/PCGCountUniqueAttributeValuesKernel.h b/Engine/Plugins/PCG/Source/PCG/Private/Compute/BuiltInKernels/PCGCountUniqueAttributeValuesKernel.h
--- a/Engine/Plugins/PCG/Source/PCG/Private/Compute/BuiltInKernels/PCGCountUniqueAttributeValuesKernel.h
+++ b/Engine/Plugins/PCG/Source/PCG/Private/Compute/BuiltInKernels/PCGCountUniqueAttributeValuesKernel.h
@@ -13,7 +13,7 @@ namespace PCGCountUniqueAttributeValuesConstants
 }
 
 /**
-* Counts how many unique values of a string key attribute are present in an input data collection. Output attribute set is a table of unique string key
+* Counts how many unique values of a string key or int attribute are present in an input data collection. Output attribute set is a table of unique
 * values and corresponding counts. Can output an attribute set per input data, or a single attribute set that counts across all data.
 */
 UCLASS()
@@ -27,6 +27,9 @@ public:
 	/** Whether to produce one set of element counts that count across all input data, rather than producing counters per input data. */
 	void SetEmitPerDataCounts(bool bInEmitPerDataCounts) { bEmitPerDataCounts = bInEmitPerDataCounts; }
 
+	/** Set largest value that attribute can have. Required when counting values of int attributes (but not string key attributes). */
+	void SetMaxIntValue(int32 InMaxIntValue) { MaxIntValue = InMaxIntValue; }
+
 	/** Whether to output a raw array of (string key value, instance count) values from analysis instead of constructing an attribute set. */
 	void SetOutputRawBuffer(bool bInOutputRawBuffer) { bOutputRawBuffer = bInOutputRawBuffer; }
 
@@ -51,6 +54,10 @@ protected:
 	UPROPERTY()
 	FName AttributeName;
 
+	/** Set largest value that attribute can have. Required when counting values of int attributes (but not string key attributes). */
+	UPROPERTY()
+	int32 MaxIntValue = INDEX_NONE;
+
 	/** Produce one set of element counts that count across all input data, rather than producing counters per input data. */
 	UPROPERTY()
 	bool bEmitPerDataCounts = true;
diff --git a/Engine/Plugins/PCG/Source/PCG/Private/Compute/DataInterfaces/BuiltInKernels/PCGCountUniqueAttributeValuesDataInterface.cpp b/Engine/Plugins/PCG/Source/PCG/Private/Compute/DataInterfaces/BuiltInKernels/PCGCountUniqueAttributeValuesDataInterface.cpp
--- a/Engine/Plugins/PCG/Source/PCG/Private/Compute/DataInterfaces/BuiltInKernels/PCGCountUniqueAttributeValuesDataInterface.cpp
+++ b/Engine/Plugins/PCG/Source/PCG/Private/Compute/DataInterfaces/BuiltInKernels/PCGCountUniqueAttributeValuesDataInterface.cpp
@@ -122,43 +122,10 @@ bool UPCGCountUniqueAttributeValuesDataProvider::PrepareForExecute_GameThread(UP
 		check(OutputCountAttributeId != -1);
 	}
 
-	const TSharedPtr<const FPCGDataCollectionDesc> InputDataDesc = InBinding->GetCachedKernelPinDataDesc(GetProducerKernel(), PCGPinConstants::DefaultInputLabel, /*bIsInput*/true);
-
-	if (!ensure(InputDataDesc))
-	{
-		return true;
-	}
-
 	AttributeToCountId = InBinding->GetAttributeId(AttributeToCountName, EPCGKernelAttributeType::StringKey);
-
-	TArray<TArray<int32>> UniqueStringKeyValuesPerOutputData;
-
-	if (AttributeToCountId != INDEX_NONE)
+	if (AttributeToCountId == INDEX_NONE)
 	{
-		if (bEmitPerDataCounts)
-		{
-			// Gather unique string key values for each data.
-			for (const FPCGDataDesc& DataDesc : InputDataDesc->GetDataDescriptions())
-			{
-				TArray<int32> UniqueStringKeyValues;
-
-				for (const FPCGKernelAttributeDesc& AttributeDesc : DataDesc.GetAttributeDescriptions())
-				{
-					if (AttributeDesc.GetAttributeId() == AttributeToCountId)
-					{
-						UniqueStringKeyValues = AttributeDesc.GetUniqueStringKeys();
-						break;
-					}
-				}
-
-				UniqueStringKeyValuesPerOutputData.Emplace_GetRef() = MoveTemp(UniqueStringKeyValues);
-			}
-		}
-		else
-		{
-			// Gather unique string key values across all data.
-			InputDataDesc->GetUniqueStringKeyValues(AttributeToCountId, UniqueStringKeyValuesPerOutputData.Emplace_GetRef());
-		}
+		AttributeToCountId = InBinding->GetAttributeId(AttributeToCountName, EPCGKernelAttributeType::Int);
 	}
 
 	return true;
diff --git a/Engine/Plugins/PCG/Source/PCG/Private/Compute/DataInterfaces/Elements/PCGSkinnedMeshSpawnerDataInterface.cpp b/Engine/Plugins/PCG/Source/PCG/Private/Compute/DataInterfaces/Elements/PCGSkinnedMeshSpawnerDataInterface.cpp
--- a/Engine/Plugins/PCG/Source/PCG/Private/Compute/DataInterfaces/Elements/PCGSkinnedMeshSpawnerDataInterface.cpp
+++ b/Engine/Plugins/PCG/Source/PCG/Private/Compute/DataInterfaces/Elements/PCGSkinnedMeshSpawnerDataInterface.cpp
@@ -81,7 +81,7 @@ void UPCGSkinnedMeshSpawnerDataInterface::GetSupportedInputs(TArray<FShaderFunct
 		.AddParam(EShaderFundamentalType::Uint); // InPrimitiveIndex
 
 	OutFunctions.AddDefaulted_GetRef()
-		.SetName(TEXT("SMSpawner_GetPrimitiveIndexFromStringKey"))
+		.SetName(TEXT("SMSpawner_GetPrimitiveIndexFromPrimitiveId"))
 		.AddReturnType(EShaderFundamentalType::Uint) // Primitive index
 		.AddParam(EShaderFundamentalType::Int); // InMeshPathStringKey
 
@@ -172,7 +172,7 @@ void UPCGSkinnedMeshSpawnerDataInterface::GetHLSL(FString& OutHLSL, FString cons
 		"	return {DataInterfaceName}_PrimitiveStringKeys[InPrimitiveIndex];\n"
 		"}\n"
 		"\n"
-		"uint SMSpawner_GetPrimitiveIndexFromStringKey_{DataInterfaceName}(int InMeshPathStringKey)\n"
+		"uint SMSpawner_GetPrimitiveIndexFromPrimitiveId_{DataInterfaceName}(int InMeshPathStringKey)\n"
 		"{\n"
 		"	for (uint Index = 0; Index < {DataInterfaceName}_NumPrimitives; ++Index)\n"
 		"	{\n"
diff --git a/Engine/Plugins/PCG/Source/PCG/Private/Compute/DataInterfaces/Elements/PCGStaticMeshSpawnerDataInterface.cpp b/Engine/Plugins/PCG/Source/PCG/Private/Compute/DataInterfaces/Elements/PCGStaticMeshSpawnerDataInterface.cpp
--- a/Engine/Plugins/PCG/Source/PCG/Private/Compute/DataInterfaces/Elements/PCGStaticMeshSpawnerDataInterface.cpp
+++ b/Engine/Plugins/PCG/Source/PCG/Private/Compute/DataInterfaces/Elements/PCGStaticMeshSpawnerDataInterface.cpp
@@ -17,7 +17,9 @@
 #include "Elements/PCGStaticMeshSpawnerKernel.h"
 #include "InstanceDataPackers/PCGInstanceDataPackerBase.h"
 #include "MeshSelectors/PCGMeshSelectorByAttribute.h"
+#include "MeshSelectors/PCGMeshSelectorPrimitiveData.h"
 #include "MeshSelectors/PCGMeshSelectorWeighted.h"
+#include "Metadata/PCGObjectPropertyOverride.h"
 
 #include "GlobalRenderResources.h"
 #include "RenderGraphBuilder.h"
@@ -28,11 +30,12 @@
 #include "ShaderCore.h"
 #include "ShaderParameterMetadataBuilder.h"
 #include "SystemTextures.h"
+#include "ComputeFramework/ShaderParamTypeDefinition.h"
 #include "Engine/AssetManager.h"
 #include "Engine/StaticMesh.h"
 #include "Engine/StreamableManager.h"
 #include "Engine/World.h"
-#include "ComputeFramework/ShaderParamTypeDefinition.h"
+#include "Materials/MaterialInterface.h"
 
 #include UE_INLINE_GENERATED_CPP_BY_NAME(PCGStaticMeshSpawnerDataInterface)
 
@@ -89,9 +92,9 @@ void UPCGStaticMeshSpawnerDataInterface::GetSupportedInputs(TArray<FShaderFuncti
 		.AddParam(EShaderFundamentalType::Uint); // InPrimitiveIndex
 
 	OutFunctions.AddDefaulted_GetRef()
-		.SetName(TEXT("SMSpawner_GetPrimitiveIndexFromStringKey"))
+		.SetName(TEXT("SMSpawner_GetPrimitiveIndexFromPrimitiveId"))
 		.AddReturnType(EShaderFundamentalType::Uint) // Primitive index
-		.AddParam(EShaderFundamentalType::Int); // InMeshPathStringKey
+		.AddParam(EShaderFundamentalType::Int); // InPrimitiveId
 
 	OutFunctions.AddDefaulted_GetRef()
 		.SetName(TEXT("SMSpawner_GetPrimitiveSelectionCDF"))
@@ -105,6 +108,7 @@ void UPCGStaticMeshSpawnerDataInterface::GetSupportedInputs(TArray<FShaderFuncti
 
 BEGIN_SHADER_PARAMETER_STRUCT(FPCGStaticMeshSpawnerDataInterfaceParameters,)
 	SHADER_PARAMETER_ARRAY(FUintVector4, AttributeIdOffsetStrides, [UPCGStaticMeshSpawnerDataInterface::MAX_ATTRIBUTES])
+	SHADER_PARAMETER_RDG_BUFFER_SRV(StructuredBuffer<int32>, PrimitiveIds)
 	SHADER_PARAMETER_RDG_BUFFER_SRV(StructuredBuffer<int32>, PrimitiveStringKeys)
 	SHADER_PARAMETER_RDG_BUFFER_SRV(StructuredBuffer<FVector4f>, PrimitiveMeshBoundsMin)
 	SHADER_PARAMETER_RDG_BUFFER_SRV(StructuredBuffer<FVector4f>, PrimitiveMeshBoundsMax)
@@ -139,6 +143,7 @@ void UPCGStaticMeshSpawnerDataInterface::GetHLSL(FString& OutHLSL, FString const
 		"uint4 {DataInterfaceName}_AttributeIdOffsetStrides[{MaxAttributes}];\n"
 		"StructuredBuffer<float4> {DataInterfaceName}_PrimitiveMeshBoundsMin;\n"
 		"StructuredBuffer<float4> {DataInterfaceName}_PrimitiveMeshBoundsMax;\n"
+		"StructuredBuffer<int> {DataInterfaceName}_PrimitiveIds;\n"
 		"StructuredBuffer<int> {DataInterfaceName}_PrimitiveStringKeys;\n"
 		"DECLARE_SCALAR_ARRAY(float, {DataInterfaceName}_SelectionCDF, {MaxPrimitives});\n"
 		"\n"
@@ -182,11 +187,11 @@ void UPCGStaticMeshSpawnerDataInterface::GetHLSL(FString& OutHLSL, FString const
 		"	return {DataInterfaceName}_PrimitiveStringKeys[InPrimitiveIndex];\n"
 		"}\n"
 		"\n"
-		"uint SMSpawner_GetPrimitiveIndexFromStringKey_{DataInterfaceName}(int InMeshPathStringKey)\n"
+		"uint SMSpawner_GetPrimitiveIndexFromPrimitiveId_{DataInterfaceName}(int InPrimitiveId)\n"
 		"{\n"
 		"	for (uint Index = 0; Index < {DataInterfaceName}_NumPrimitives; ++Index)\n"
 		"	{\n"
-		"		if ({DataInterfaceName}_PrimitiveStringKeys[Index] == InMeshPathStringKey)\n"
+		"		if ({DataInterfaceName}_PrimitiveIds[Index] == InPrimitiveId)\n"
 		"		{\n"
 		"			return Index;\n"
 		"		}\n"
@@ -212,6 +217,15 @@ UComputeDataProvider* UPCGStaticMeshSpawnerDataInterface::CreateDataProvider() c
 	return NewObject<UPCGStaticMeshSpawnerDataProvider>();
 }
 
+void UPCGStaticMeshSpawnerDataProvider::Initialize(const UComputeDataInterface* InDataInterface, UObject* InBinding, uint64 InInputMask, uint64 InOutputMask)
+{
+	Super::Initialize(InDataInterface, InBinding, InInputMask, InOutputMask);
+
+	const UPCGStaticMeshSpawnerSettings* Settings = CastChecked<UPCGStaticMeshSpawnerSettings>(GetProducerKernel()->GetSettings());
+	bSpawningByAttribute = !!Cast<UPCGMeshSelectorByAttribute>(Settings->MeshSelectorParameters);
+	bSpawningByPrimitiveData = !!Cast<UPCGMeshSelectorPrimitiveData>(Settings->MeshSelectorParameters);
+}
+
 bool UPCGStaticMeshSpawnerDataProvider::PerformPreExecuteReadbacks_GameThread(UPCGDataBinding* InBinding)
 {
 	TRACE_CPUPROFILER_EVENT_SCOPE(UPCGStaticMeshSpawnerDataProvider::PerformPreExecuteReadbacks_GameThread);
@@ -227,7 +241,7 @@ bool UPCGStaticMeshSpawnerDataProvider::PerformPreExecuteReadbacks_GameThread(UP
 	// Obtain index of analysis data that we need to read back.
 	if (AnalysisDataIndex == INDEX_NONE)
 	{
-		if (!Cast<UPCGMeshSelectorByAttribute>(Settings->MeshSelectorParameters))
+		if (!bSpawningByAttribute && !bSpawningByPrimitiveData)
 		{
 			// Non by-attribute selection does not need to do any readbacks.
 			return true;
@@ -242,8 +256,22 @@ bool UPCGStaticMeshSpawnerDataProvider::PerformPreExecuteReadbacks_GameThread(UP
 		}
 	}
 
-	// Readback analysis data - poll until readback complete (return true).
-	if (!InBinding->ReadbackInputDataToCPU(AnalysisDataIndex))
+	// Obtain index of primitive data that we need to read back.
+	if (PrimitiveDataIndex == INDEX_NONE && bSpawningByPrimitiveData)
+	{
+		PrimitiveDataIndex = GetProducerKernel() ? InBinding->GetFirstInputDataIndex(GetProducerKernel(), PCGStaticMeshSpawnerKernelConstants::PrimitiveTablePinLabel) : INDEX_NONE;
+
+		if (PrimitiveDataIndex == INDEX_NONE)
+		{
+			// No data to read back.
+			return true;
+		}
+	}
+
+	// Readback analysis data and primitive data if required. Poll until readback complete (returns true).
+	const bool bAnalysisDataReadBack = InBinding->ReadbackInputDataToCPU(AnalysisDataIndex);
+	const bool bPrimitiveDataReadBack = PrimitiveDataIndex == INDEX_NONE || InBinding->ReadbackInputDataToCPU(PrimitiveDataIndex);
+	if (!bAnalysisDataReadBack || !bPrimitiveDataReadBack)
 	{
 		return false;
 	}
@@ -264,12 +292,12 @@ bool UPCGStaticMeshSpawnerDataProvider::PerformPreExecuteReadbacks_GameThread(UP
 
 			const int32 NumElements = AnalysisMetadata->GetItemCountForChild();
 
-			StringKeyToInstanceCount.Reserve(NumElements);
+			PrimitiveIdToInstanceCount.Reserve(NumElements);
 
 			// TODO: Range based get would scale better.
 			for (int64 MetadataKey = 0; MetadataKey < NumElements; ++MetadataKey)
 			{
-				StringKeyToInstanceCount.Add(ValueAttribute->GetValue(MetadataKey), CountAttribute->GetValue(MetadataKey));
+				PrimitiveIdToInstanceCount.Add(ValueAttribute->GetValue(MetadataKey), CountAttribute->GetValue(MetadataKey));
 			}
 		}
 		else
@@ -293,7 +321,7 @@ bool UPCGStaticMeshSpawnerDataProvider::PerformPreExecuteReadbacks_GameThread(UP
 		const uint32* RawData = RawAnalysisData->GetConstData().GetData();
 		for (const uint32* It = RawData; It < RawData + NumUints; It += 2)
 		{
-			StringKeyToInstanceCount.Add(static_cast<int32>(It[0]), It[1]);
+			PrimitiveIdToInstanceCount.Add(static_cast<int32>(It[0]), It[1]);
 		}
 	}
 
@@ -332,9 +360,12 @@ bool UPCGStaticMeshSpawnerDataProvider::PrepareForExecute_GameThread(UPCGDataBin
 	}
 
 	const UPCGMeshSelectorByAttribute* SelectorByAttribute = Cast<UPCGMeshSelectorByAttribute>(Settings->MeshSelectorParameters);
-	if (SelectorAttributeId == INDEX_NONE && SelectorByAttribute)
+	const UPCGMeshSelectorPrimitiveData* SelectorPrimitiveData = Cast<UPCGMeshSelectorPrimitiveData>(Settings->MeshSelectorParameters);
+
+	if (SelectorAttributeId == INDEX_NONE && (SelectorByAttribute || SelectorPrimitiveData))
 	{
-		const FName SelectorName = SelectorByAttribute->AttributeName;
+		const FName SelectorName = SelectorByAttribute ? SelectorByAttribute->AttributeName : SelectorPrimitiveData->PrimitiveIndexAttribute;
+
 		const TSharedPtr<const FPCGDataCollectionDesc> InputDataDesc = InBinding->GetCachedKernelPinDataDesc(GetProducerKernel(), PCGPinConstants::DefaultInputLabel, /*bIsInput=*/true);
 
 		if (!ensure(InputDataDesc))
@@ -356,7 +387,8 @@ bool UPCGStaticMeshSpawnerDataProvider::PrepareForExecute_GameThread(UPCGDataBin
 
 			for (const FPCGKernelAttributeDesc& AttributeDesc : Desc.GetAttributeDescriptions())
 			{
-				if (AttributeDesc.GetAttributeKey().GetIdentifier().Name == SelectorName && AttributeDesc.GetAttributeKey().GetType() == EPCGKernelAttributeType::StringKey)
+				if (AttributeDesc.GetAttributeKey().GetIdentifier().Name == SelectorName
+					&& (AttributeDesc.GetAttributeKey().GetType() == EPCGKernelAttributeType::StringKey || AttributeDesc.GetAttributeKey().GetType() == EPCGKernelAttributeType::Int))
 				{
 					SelectorAttributeId = AttributeDesc.GetAttributeId();
 					break;
@@ -474,22 +506,27 @@ bool UPCGStaticMeshSpawnerDataProvider::PrepareForExecute_GameThread(UPCGDataBin
 
 FComputeDataProviderRenderProxy* UPCGStaticMeshSpawnerDataProvider::GetRenderProxy()
 {
-	return new FPCGStaticMeshSpawnerDataProviderProxy(AttributeIdOffsetStrides, SelectorAttributeId, PrimitiveStringKeys, PrimitiveSelectionCDF, SelectedMeshAttributeId, PrimitiveMeshBounds);
+	return new FPCGStaticMeshSpawnerDataProviderProxy(
+		AttributeIdOffsetStrides, SelectorAttributeId, PrimitiveIds, PrimitiveStringKeys, PrimitiveSelectionCDF, SelectedMeshAttributeId, PrimitiveMeshBounds);
 }
 
 void UPCGStaticMeshSpawnerDataProvider::Reset()
 {
 	Super::Reset();
 
+	bSpawningByAttribute = false;
+	bSpawningByPrimitiveData = false;
 	AttributeIdOffsetStrides.Empty();
+	PrimitiveIds.Empty();
 	PrimitiveStringKeys.Empty();
 	PrimitiveMeshBounds.Empty();
 	PrimitiveSelectionCDF.Empty();
 	SelectorAttributeId = INDEX_NONE;
 	NumInputPoints = 0;
 	SelectedMeshAttributeId = INDEX_NONE;
-	StringKeyToInstanceCount.Empty();
+	PrimitiveIdToInstanceCount.Empty();
 	AnalysisDataIndex = INDEX_NONE;
+	PrimitiveDataIndex = INDEX_NONE;
 	bPrimitiveDescriptorsCreated = false;
 	PrimitiveDescriptors.Empty();
 	CustomFloatCount = 0;
@@ -502,6 +539,9 @@ void UPCGStaticMeshSpawnerDataProvider::Reset()
 
 void UPCGStaticMeshSpawnerDataProvider::CreatePrimitiveDescriptors(FPCGContext* InContext, UPCGDataBinding* InBinding)
 {
+	TRACE_CPUPROFILER_EVENT_SCOPE(UPCGStaticMeshSpawnerDataProvider::CreatePrimitiveDescriptors);
+	check(InContext);
+
 	PrimitiveDescriptors.Empty();
 
 	const UPCGStaticMeshSpawnerSettings* Settings = CastChecked<UPCGStaticMeshSpawnerSettings>(GetProducerKernel()->GetSettings());
@@ -575,7 +615,7 @@ void UPCGStaticMeshSpawnerDataProvider::CreatePrimitiveDescriptors(FPCGContext*
 		const FName SelectorName = SelectorByAttribute->AttributeName;
 
 		// Compute how many instances we expect for each mesh, if we do not already have the answer from analysis.
-		if (StringKeyToInstanceCount.IsEmpty())
+		if (PrimitiveIdToInstanceCount.IsEmpty())
 		{
 			for (const FPCGDataDesc& Desc : InputDataDesc->GetDataDescriptions())
 			{
@@ -594,7 +634,7 @@ void UPCGStaticMeshSpawnerDataProvider::CreatePrimitiveDescriptors(FPCGContext*
 							{
 								if (AttributeStringKeyValue > 0)
 								{
-									StringKeyToInstanceCount.FindOrAdd(AttributeStringKeyValue, 0) += Desc.GetElementCountForAttribute(AttributeDesc).X;
+									PrimitiveIdToInstanceCount.FindOrAdd(AttributeStringKeyValue, 0) += Desc.GetElementCountForAttribute(AttributeDesc).X;
 								}
 							}
 
@@ -613,21 +653,21 @@ void UPCGStaticMeshSpawnerDataProvider::CreatePrimitiveDescriptors(FPCGContext*
 			}
 		}
 
-		PrimitiveDescriptors.Reserve(StringKeyToInstanceCount.Num());
+		PrimitiveDescriptors.Reserve(PrimitiveIdToInstanceCount.Num());
 
-		for (const TPair<int32, uint32>& StringKeyAndInstanceCount : StringKeyToInstanceCount)
+		for (const TPair<int32, uint32>& PrimitiveIdAndInstanceCount : PrimitiveIdToInstanceCount)
 		{
-			if (StringKeyAndInstanceCount.Value == 0)
+			if (PrimitiveIdAndInstanceCount.Value == 0)
 			{
 				continue;
 			}
 
-			if (!ensure(InBinding->GetStringTable().IsValidIndex(StringKeyAndInstanceCount.Key)))
+			if (!ensure(InBinding->GetStringTable().IsValidIndex(PrimitiveIdAndInstanceCount.Key)))
 			{
 				continue;
 			}
 
-			const FString& MeshPathString = InBinding->GetStringTable()[StringKeyAndInstanceCount.Key];
+			const FString& MeshPathString = InBinding->GetStringTable()[PrimitiveIdAndInstanceCount.Key];
 			if (MeshPathString.IsEmpty())
 			{
 				continue;
@@ -643,7 +683,7 @@ void UPCGStaticMeshSpawnerDataProvider::CreatePrimitiveDescriptors(FPCGContext*
 
 			FPCGProceduralISMComponentDescriptor Descriptor;
 			Descriptor = SelectorByAttribute->TemplateDescriptor;
-			Descriptor.NumInstances = StringKeyAndInstanceCount.Value;
+			Descriptor.NumInstances = PrimitiveIdAndInstanceCount.Value;
 			Descriptor.WorldBounds = InContext->ExecutionSource->GetExecutionState().GetBounds();
 			Descriptor.NumCustomFloats = CustomFloatCount;
 			Descriptor.StaticMesh = FSoftObjectPath(MeshPathString);
@@ -654,8 +694,115 @@ void UPCGStaticMeshSpawnerDataProvider::CreatePrimitiveDescriptors(FPCGContext*
 				Descriptor.NumInstances = TotalInputPointCount;
 			}
 
-			PrimitiveStringKeys.Add(StringKeyAndInstanceCount.Key);
+			PrimitiveIds.Add(PrimitiveIdAndInstanceCount.Key);
+			PrimitiveDescriptors.Add(MoveTemp(Descriptor));
+
+			// For by-attribute spawning, the primitive ID is the mesh path string key.
+			PrimitiveStringKeys.Add(PrimitiveIdAndInstanceCount.Key);
+		}
+
+		PrimitiveSelectionCDF.SetNumZeroed(PrimitiveDescriptors.Num());
+	}
+	else if (const UPCGMeshSelectorPrimitiveData* SelectorPrimitiveData = Cast<UPCGMeshSelectorPrimitiveData>(Settings->MeshSelectorParameters))
+	{
+		if (PrimitiveIdToInstanceCount.IsEmpty())
+		{
+			return;
+		}
+
+		const bool bIndexValid = InBinding->GetInputDataCollection().TaggedData.IsValidIndex(PrimitiveDataIndex);
+		const UPCGParamData* PrimitiveData = bIndexValid ? Cast<UPCGParamData>(InBinding->GetInputDataCollection().TaggedData[PrimitiveDataIndex].Data) : nullptr;
+		const UPCGMetadata* PrimitiveMetadata = PrimitiveData ? PrimitiveData->ConstMetadata() : nullptr;
+		
+		if (!PrimitiveMetadata)
+		{
+			PCG_KERNEL_VALIDATION_ERR(InContext, Settings, LOCTEXT("NoPrimitiveData", "No primitive data found, no primitives will spawn."));
+			return;
+		}
+
+		const FPCGMetadataAttributeBase* MeshAttributeBase = PrimitiveMetadata->GetConstAttribute(SelectorPrimitiveData->MeshAttribute);
+		const bool bIsString = MeshAttributeBase && MeshAttributeBase->GetTypeId() == PCG::Private::MetadataTypes<FString>::Id;
+		const bool bIsSoftObjectPath = MeshAttributeBase && MeshAttributeBase->GetTypeId() == PCG::Private::MetadataTypes<FSoftObjectPath>::Id;
+		if (!MeshAttributeBase || (!bIsString && !bIsSoftObjectPath))
+		{
+			PCG_KERNEL_VALIDATION_WARN(InContext, Settings,
+				FText::Format(LOCTEXT("MeshAttributeInvalid", "Mesh attribute {0} not found in primitive data, or is not of required type Soft Object Path or String."), FText::FromName(SelectorPrimitiveData->MeshAttribute)));
+			return;
+		}
+
+		const FPCGMetadataAttribute<FString>* MeshAttributeString = bIsString ? static_cast<const FPCGMetadataAttribute<FString>*>(MeshAttributeBase) : nullptr;
+		const FPCGMetadataAttribute<FSoftObjectPath>* MeshAttributeSoftObjectPath = bIsSoftObjectPath ? static_cast<const FPCGMetadataAttribute<FSoftObjectPath>*>(MeshAttributeBase) : nullptr;
+
+		FPCGProceduralISMComponentDescriptor Descriptor;
+		Descriptor = SelectorPrimitiveData->TemplateDescriptor;
+		Descriptor.WorldBounds = InContext->ExecutionSource->GetExecutionState().GetBounds();
+		Descriptor.NumCustomFloats = CustomFloatCount;
+
+		FPCGObjectOverrides PropertyOverrides(&Descriptor);
+		PropertyOverrides.Initialize(SelectorPrimitiveData->PrimitiveOverrideAttributes, &Descriptor, PrimitiveData, InContext);
+
+		FPCGMeshMaterialOverrideHelper MaterialOverrideHelper;
+		MaterialOverrideHelper.Initialize(*InContext, /*bInByAttributeOverride=*/true, SelectorPrimitiveData->MaterialOverrideAttributes, PrimitiveMetadata);
+
+		const int32 NumPrimitiveDataElements = PrimitiveMetadata->GetItemCountForChild();
+
+		PrimitiveDescriptors.Reserve(PrimitiveIdToInstanceCount.Num());
+		PrimitiveIds.Reserve(PrimitiveIdToInstanceCount.Num());
+		PrimitiveStringKeys.Reserve(PrimitiveIdToInstanceCount.Num());
+
+		for (const TPair<int32, uint32>& PrimitiveIdAndInstanceCount : PrimitiveIdToInstanceCount)
+		{
+			if (PrimitiveIdAndInstanceCount.Value == 0)
+			{
+				continue;
+			}
+
+			TRACE_CPUPROFILER_EVENT_SCOPE(UPCGStaticMeshSpawnerDataProvider::CreatePrimitiveDescriptor);
+
+			const int32 PrimitiveId = PrimitiveIdAndInstanceCount.Key;
+			if (PrimitiveId >= NumPrimitiveDataElements)
+			{
+				PCG_KERNEL_VALIDATION_WARN(InContext, Settings,
+					FText::Format(LOCTEXT("OutOfBoundsPrimitiveIndex", "Primitive index {0} exceeds size of primitive table ({1})."), PrimitiveId, NumPrimitiveDataElements));
+				continue;
+			}
+
+			const PCGMetadataValueKey ValueKey = MeshAttributeSoftObjectPath ? MeshAttributeSoftObjectPath->GetValueKey(PrimitiveId) : MeshAttributeString->GetValueKey(PrimitiveId);
+			const FSoftObjectPath MeshPath = MeshAttributeSoftObjectPath ? MeshAttributeSoftObjectPath->GetValue(ValueKey) : FSoftObjectPath(MeshAttributeString->GetValue(ValueKey));
+
+			if (PrimitiveDescriptors.Num() >= PCGComputeConstants::MAX_PRIMITIVE_COMPONENTS_PER_SPAWNER)
+			{
+				PCG_KERNEL_VALIDATION_WARN(InContext, Settings,
+					FText::Format(PCGStaticMeshSpawnerDataInterface::TooManyPrimitivesFormat, FText::FromString(FString::FromInt(PCGComputeConstants::MAX_PRIMITIVE_COMPONENTS_PER_SPAWNER))));
+
+				break;
+			}
+
+			Descriptor.NumInstances = PrimitiveIdAndInstanceCount.Value;
+			Descriptor.StaticMesh = MeshPath;
+
+			PropertyOverrides.Apply(PrimitiveId);
+
+			if (MaterialOverrideHelper.OverridesMaterials())
+			{
+				for (const TSoftObjectPtr<UMaterialInterface>& Material : MaterialOverrideHelper.GetMaterialOverrides(PrimitiveId))
+				{
+					// todo_pcg: Add async load.
+					Descriptor.OverrideMaterials.Add(Material.LoadSynchronous());
+				}
+			}
+
+			// Sanity check instance count.
+			if (!ensure(Descriptor.NumInstances <= TotalInputPointCount))
+			{
+				Descriptor.NumInstances = TotalInputPointCount;
+			}
+
+			PrimitiveIds.Add(PrimitiveId);
 			PrimitiveDescriptors.Add(MoveTemp(Descriptor));
+
+			// Look up the mesh path string key value.
+			PrimitiveStringKeys.Add(InBinding->GetStringTable().IndexOfByKey(MeshPath));
 		}
 
 		PrimitiveSelectionCDF.SetNumZeroed(PrimitiveDescriptors.Num());
@@ -710,7 +857,11 @@ void UPCGStaticMeshSpawnerDataProvider::CreatePrimitiveDescriptors(FPCGContext*
 
 			CumulativeWeight += Weight;
 			PrimitiveSelectionCDF.Add(CumulativeWeight);
-			PrimitiveStringKeys.Add(InBinding->GetStringTable().IndexOfByKey(Entry.Descriptor.StaticMesh.ToString()));
+
+			// For weighted spawning, primitive ID is the string key value.
+			const int32 StringKeyValue = InBinding->GetStringTable().IndexOfByKey(Entry.Descriptor.StaticMesh.ToString());
+			PrimitiveIds.Add(StringKeyValue);
+			PrimitiveStringKeys.Add(StringKeyValue);
 
 			FPCGProceduralISMComponentDescriptor Descriptor;
 			Descriptor = Entry.Descriptor;
@@ -834,6 +985,7 @@ bool UPCGStaticMeshSpawnerDataProvider::SetupPrimitives(FPCGContext* InContext,
 FPCGStaticMeshSpawnerDataProviderProxy::FPCGStaticMeshSpawnerDataProviderProxy(
 	const TArray<FUintVector4>& InAttributeIdOffsetStrides,
 	int32 InSelectorAttributeId,
+	const TArray<int32>& InPrimitiveIds,
 	const TArray<int32>& InPrimitiveStringKeys,
 	TArray<float> InSelectionCDF,
 	int32 InSelectedMeshAttributeId,
@@ -841,6 +993,7 @@ FPCGStaticMeshSpawnerDataProviderProxy::FPCGStaticMeshSpawnerDataProviderProxy(
 	: AttributeIdOffsetStrides(InAttributeIdOffsetStrides)
 	, SelectionCDF(InSelectionCDF)
 	, SelectorAttributeId(InSelectorAttributeId)
+	, PrimitiveIds(InPrimitiveIds)
 	, PrimitiveStringKeys(InPrimitiveStringKeys)
 	, SelectedMeshAttributeId(InSelectedMeshAttributeId)
 {
@@ -861,6 +1014,20 @@ bool FPCGStaticMeshSpawnerDataProviderProxy::IsValid(FValidationData const& InVa
 
 void FPCGStaticMeshSpawnerDataProviderProxy::AllocateResources(FRDGBuilder& GraphBuilder, FAllocationData const& InAllocationData)
 {
+	if (!PrimitiveIds.IsEmpty())
+	{
+		const FRDGBufferDesc BufferDesc = FRDGBufferDesc::CreateStructuredDesc(PrimitiveIds.GetTypeSize(), PrimitiveIds.Num());
+
+		FRDGBufferRef Buffer = GraphBuilder.CreateBuffer(BufferDesc, TEXT("PCGStaticMeshSpawner_PrimitiveIds"));
+		PrimitiveIdsBufferSRV = GraphBuilder.CreateSRV(Buffer);
+
+		GraphBuilder.QueueBufferUpload(Buffer, MakeArrayView(PrimitiveIds));
+	}
+	else
+	{
+		PrimitiveIdsBufferSRV = GraphBuilder.CreateSRV(FRDGBufferSRVDesc(GSystemTextures.GetDefaultStructuredBuffer(GraphBuilder, PrimitiveIds.GetTypeSize())));
+	}
+
 	if (!PrimitiveStringKeys.IsEmpty())
 	{
 		const FRDGBufferDesc BufferDesc = FRDGBufferDesc::CreateStructuredDesc(PrimitiveStringKeys.GetTypeSize(), PrimitiveStringKeys.Num());
@@ -921,6 +1088,7 @@ void FPCGStaticMeshSpawnerDataProviderProxy::GatherDispatchData(FDispatchData co
 			Parameters.AttributeIdOffsetStrides[Index] = AttributeIdOffsetStrides[Index];
 		}
 
+		Parameters.PrimitiveIds = PrimitiveIdsBufferSRV;
 		Parameters.PrimitiveStringKeys = PrimitiveStringKeysBufferSRV;
 
 		for (int32 Index = 0; Index < SelectionCDF.Num(); ++Index)
diff --git a/Engine/Plugins/PCG/Source/PCG/Private/Compute/DataInterfaces/Elements/PCGStaticMeshSpawnerDataInterface.h b/Engine/Plugins/PCG/Source/PCG/Private/Compute/DataInterfaces/Elements/PCGStaticMeshSpawnerDataInterface.h
--- a/Engine/Plugins/PCG/Source/PCG/Private/Compute/DataInterfaces/Elements/PCGStaticMeshSpawnerDataInterface.h
+++ b/Engine/Plugins/PCG/Source/PCG/Private/Compute/DataInterfaces/Elements/PCGStaticMeshSpawnerDataInterface.h
@@ -38,6 +38,7 @@ class UPCGStaticMeshSpawnerDataProvider : public UPCGComputeDataProvider
 
 public:
 	//~ Begin UComputeDataProvider Interface
+	virtual void Initialize(const UComputeDataInterface* InDataInterface, UObject* InBinding, uint64 InInputMask, uint64 InOutputMask) override;
 	virtual bool PerformPreExecuteReadbacks_GameThread(UPCGDataBinding* InBinding) override;
 	virtual bool PrepareForExecute_GameThread(UPCGDataBinding* InBinding) override;
 	FComputeDataProviderRenderProxy* GetRenderProxy() override;
@@ -50,35 +51,40 @@ public:
 	virtual bool SetupPrimitives(FPCGContext* InContext, UPCGDataBinding* InBinding);
 
 public:
+	bool bSpawningByAttribute = false;
+	bool bSpawningByPrimitiveData = false;
+
 	/** Attributes to use for writing per-instance custom floats. */
-	UPROPERTY()
 	TArray<FUintVector4> AttributeIdOffsetStrides;
 
-	UPROPERTY()
+	/** Array of primitive IDs used by the spawner kernel to look up a spawned primitive index from a selector attribute value.
+	* When using direct by-attribute selection, each primitive ID is the string key value of the mesh path, otherwise it is an
+	* index into a primitive data table.
+	*/
+	TArray<int32> PrimitiveIds;
+
+	/** Mesh path string keys of each primitive. */
 	TArray<int32> PrimitiveStringKeys;
 
-	UPROPERTY()
 	TArray<FBox> PrimitiveMeshBounds;
 
-	UPROPERTY()
 	TArray<float> PrimitiveSelectionCDF;
 
 	/** Attribute Id for mesh selector. */
-	UPROPERTY()
 	int32 SelectorAttributeId = INDEX_NONE;
 
-	UPROPERTY()
 	int32 NumInputPoints = 0;
 
 	/** Output attribute Id for the output selected mesh. */
-	UPROPERTY()
 	int32 SelectedMeshAttributeId = INDEX_NONE;
 
-	/** The number of instances per string key value, used for by-attribute spawning. */
-	TMap<int32, uint32> StringKeyToInstanceCount;
+	/** The number of instances per primitive ID, used for by-attribute spawning. */
+	TMap<int32, uint32> PrimitiveIdToInstanceCount;
 
 	int32 AnalysisDataIndex = INDEX_NONE;
 
+	int32 PrimitiveDataIndex = INDEX_NONE;
+
 	bool bPrimitiveDescriptorsCreated = false;
 
 	TArray<FPCGProceduralISMComponentDescriptor> PrimitiveDescriptors;
@@ -104,6 +110,7 @@ public:
 	FPCGStaticMeshSpawnerDataProviderProxy(
 		const TArray<FUintVector4>& InAttributeIdOffsetStrides,
 		int32 InSelectorAttributeId,
+		const TArray<int32>& InPrimitiveIds,
 		const TArray<int32>& InPrimitiveStringKeys,
 		TArray<float> InSelectionCDF,
 		int32 InSelectedMeshAttributeId,
@@ -123,6 +130,9 @@ protected:
 
 	int32 SelectorAttributeId = INDEX_NONE;
 
+	TArray<int32> PrimitiveIds;
+	FRDGBufferSRVRef PrimitiveIdsBufferSRV = nullptr;
+
 	TArray<int32> PrimitiveStringKeys;
 	FRDGBufferSRVRef PrimitiveStringKeysBufferSRV = nullptr;
 
diff --git a/Engine/Plugins/PCG/Source/PCG/Private/Elements/PCGStaticMeshSpawner.cpp b/Engine/Plugins/PCG/Source/PCG/Private/Elements/PCGStaticMeshSpawner.cpp
--- a/Engine/Plugins/PCG/Source/PCG/Private/Elements/PCGStaticMeshSpawner.cpp
+++ b/Engine/Plugins/PCG/Source/PCG/Private/Elements/PCGStaticMeshSpawner.cpp
@@ -18,6 +18,7 @@
 #include "InstanceDataPackers/PCGInstanceDataPackerBase.h"
 #include "MeshSelectors/PCGMeshSelectorBase.h"
 #include "MeshSelectors/PCGMeshSelectorByAttribute.h"
+#include "MeshSelectors/PCGMeshSelectorPrimitiveData.h"
 #include "MeshSelectors/PCGMeshSelectorWeighted.h"
 
 #include "GrassInstancedStaticMeshComponent.h"
@@ -42,6 +43,11 @@ static TAutoConsoleVariable<bool> CVarUseRawDataForAnalysis(
 	ECVF_ReadOnly // Structural change, could flush graph compilation cache, but trickier to flush compute graph instance pool in UPCGGraph currently.
 	);
 
+namespace PCGStaticMeshSpawner
+{
+	static const FName PrimitiveTablePinLabel = TEXT("PrimitiveTable");
+}
+
 UPCGStaticMeshSpawnerSettings::UPCGStaticMeshSpawnerSettings(const FObjectInitializer &ObjectInitializer)
 {
 	MeshSelectorType = UPCGMeshSelectorWeighted::StaticClass();
@@ -57,25 +63,46 @@ UPCGStaticMeshSpawnerSettings::UPCGStaticMeshSpawnerSettings(const FObjectInitia
 #if WITH_EDITOR
 void UPCGStaticMeshSpawnerSettings::CreateKernels(FPCGGPUCompilationContext& InOutContext, UObject* InObjectOuter, TArray<UPCGComputeKernel*>& OutKernels, TArray<FPCGKernelEdge>& OutEdges) const
 {
-	PCGKernelHelpers::FCreateKernelParams CreateParams(InObjectOuter, this);
-
 	const bool bUseRawBuffer = CVarUseRawDataForAnalysis.GetValueOnAnyThread();
 
-	UPCGStaticMeshSpawnerKernel* SpawnerKernel = PCGKernelHelpers::CreateKernel<UPCGStaticMeshSpawnerKernel>(InOutContext, CreateParams, OutKernels, OutEdges);
-	SpawnerKernel->SetUseRawInstanceCountsBuffer(bUseRawBuffer);
+	UPCGStaticMeshSpawnerKernel* SpawnerKernel = nullptr;
+
+	// Create spawner kernel.
+	{
+		PCGKernelHelpers::FCreateKernelParams CreateParams(InObjectOuter, this);
+
+		// If spawning by-primitive-data, wire up the primitive data.
+		if (const UPCGMeshSelectorPrimitiveData* Selector = Cast<UPCGMeshSelectorPrimitiveData>(MeshSelectorParameters))
+		{
+			CreateParams.NodeInputPinsToWire.Add(PCGStaticMeshSpawner::PrimitiveTablePinLabel);
+		}
+
+		SpawnerKernel = PCGKernelHelpers::CreateKernel<UPCGStaticMeshSpawnerKernel>(InOutContext, CreateParams, OutKernels, OutEdges);
+		SpawnerKernel->SetUseRawInstanceCountsBuffer(bUseRawBuffer);
+	}
 
 	// If doing by-attribute selection, add analysis kernel that will count how many instances of each mesh are present.
-	if (const UPCGMeshSelectorByAttribute* Selector = Cast<UPCGMeshSelectorByAttribute>(MeshSelectorParameters))
+	const UPCGMeshSelectorByAttribute* SelectorByAttribute = Cast<UPCGMeshSelectorByAttribute>(MeshSelectorParameters);
+	const UPCGMeshSelectorPrimitiveData* SelectorPrimitiveData = Cast<UPCGMeshSelectorPrimitiveData>(MeshSelectorParameters);
+	if (SelectorByAttribute || SelectorPrimitiveData)
 	{
+		PCGKernelHelpers::FCreateKernelParams CreateParams(InObjectOuter, this);
+
 		// Don't wire count kernel to node output pin, wire manually to the spawner kernel below.
 		CreateParams.NodeOutputPinsToWire.Empty();
 
 		UPCGCountUniqueAttributeValuesKernel* CountKernel = PCGKernelHelpers::CreateKernel<UPCGCountUniqueAttributeValuesKernel>(InOutContext, CreateParams, OutKernels, OutEdges);
-		CountKernel->SetAttributeName(Selector->AttributeName);
+		CountKernel->SetAttributeName(SelectorByAttribute ? SelectorByAttribute->AttributeName : SelectorPrimitiveData->PrimitiveIndexAttribute);
 		// We operate across all input data rather than spawning for each input data separately.
 		CountKernel->SetEmitPerDataCounts(false);
 		CountKernel->SetOutputRawBuffer(bUseRawBuffer);
 
+		// Rather than add a set of max reduction passes to compute this on the fly, we set the maximum number of values to the known static maximum value.
+		if (SelectorPrimitiveData)
+		{
+			CountKernel->SetMaxIntValue(PCGComputeConstants::MAX_PRIMITIVE_COMPONENTS_PER_SPAWNER - 1);
+		}
+
 		OutEdges.Emplace(FPCGPinReference(CountKernel, PCGPinConstants::DefaultOutputLabel), FPCGPinReference(SpawnerKernel, PCGStaticMeshSpawnerKernelConstants::InstanceCountsPinLabel));
 	}
 }
@@ -109,6 +136,15 @@ TArray<FPCGPinProperties> UPCGStaticMeshSpawnerSettings::InputPinProperties() co
 	FPCGPinProperties& InputPinProperty = Properties.Emplace_GetRef(PCGPinConstants::DefaultInputLabel, EPCGDataType::Point, bAllowMultipleConnections);
 	InputPinProperty.SetRequiredPin();
 
+	if (MeshSelectorType && MeshSelectorType->IsChildOf(UPCGMeshSelectorPrimitiveData::StaticClass()))
+	{
+		FPCGPinProperties& PrimitiveDataPin = Properties.Emplace_GetRef(PCGStaticMeshSpawner::PrimitiveTablePinLabel, EPCGDataType::Param, /*bAllowMultipleConnections=*/false);
+		PrimitiveDataPin.SetRequiredPin();
+#if WITH_EDITOR
+		PrimitiveDataPin.Tooltip = LOCTEXT("PrimitiveDataPinTooltip", "Table of attributes that drive primitive properties. The primitive index attribute on the input points index into rows of this table.");
+#endif
+	}
+
 	return Properties;
 }
 
@@ -117,6 +153,19 @@ FPCGElementPtr UPCGStaticMeshSpawnerSettings::CreateElement() const
 	return MakeShared<FPCGStaticMeshSpawnerElement>();
 }
 
+EPCGChangeType UPCGStaticMeshSpawnerSettings::GetChangeTypeForProperty(const FName& InPropertyName) const
+{
+	EPCGChangeType ChangeType = Super::GetChangeTypeForProperty(InPropertyName);
+
+	if (InPropertyName == GET_MEMBER_NAME_CHECKED(UPCGStaticMeshSpawnerSettings, MeshSelectorType))
+	{
+		// Mesh selectors can drive node pins.
+		ChangeType |= EPCGChangeType::Structural | EPCGChangeType::Cosmetic;
+	}
+
+	return ChangeType;
+}
+
 FPCGContext* FPCGStaticMeshSpawnerElement::CreateContext()
 {
 	return new FPCGStaticMeshSpawnerContext();
diff --git a/Engine/Plugins/PCG/Source/PCG/Private/Elements/PCGStaticMeshSpawnerKernel.cpp b/Engine/Plugins/PCG/Source/PCG/Private/Elements/PCGStaticMeshSpawnerKernel.cpp
--- a/Engine/Plugins/PCG/Source/PCG/Private/Elements/PCGStaticMeshSpawnerKernel.cpp
+++ b/Engine/Plugins/PCG/Source/PCG/Private/Elements/PCGStaticMeshSpawnerKernel.cpp
@@ -19,6 +19,7 @@
 #include "Graph/PCGGPUGraphCompilationContext.h"
 #include "InstanceDataPackers/PCGInstanceDataPackerBase.h"
 #include "MeshSelectors/PCGMeshSelectorByAttribute.h"
+#include "MeshSelectors/PCGMeshSelectorPrimitiveData.h"
 #include "MeshSelectors/PCGMeshSelectorWeighted.h"
 
 #include "ShaderCompilerCore.h"
@@ -100,6 +101,51 @@ TSharedPtr<const FPCGDataCollectionDesc> UPCGStaticMeshSpawnerKernel::ComputeOut
 				}
 			}
 		}
+		else if (UPCGMeshSelectorPrimitiveData* SelectorPrimitiveData = Cast<UPCGMeshSelectorPrimitiveData>(SMSettings->MeshSelectorParameters))
+		{
+			TSharedPtr<FPCGContextHandle> ContextHandle = InBinding->GetContextHandle().Pin();
+			FPCGContext* Context = ContextHandle ? ContextHandle->GetContext() : nullptr;
+
+			// Collect output string keys from primitive table.
+			const int32 PrimitiveTableIndex = InBinding->GetFirstInputDataIndex(this, PCGStaticMeshSpawnerKernelConstants::PrimitiveTablePinLabel);
+			const bool bIndexValid = InBinding->GetInputDataCollection().TaggedData.IsValidIndex(PrimitiveTableIndex);
+			const UPCGParamData* PrimitiveData = bIndexValid ? Cast<UPCGParamData>(InBinding->GetInputDataCollection().TaggedData[PrimitiveTableIndex].Data) : nullptr;
+			const UPCGMetadata* PrimitiveMetadata = PrimitiveData ? PrimitiveData->ConstMetadata() : nullptr;
+
+			if (!PrimitiveMetadata)
+			{
+				PCG_KERNEL_VALIDATION_ERR(Context, SMSettings, LOCTEXT("NoPrimitiveData", "No primitive table data found, no primitives will spawn."));
+				return OutDataDesc;
+			}
+
+			const FPCGMetadataAttributeBase* MeshAttributeBase = PrimitiveMetadata->GetConstAttribute(SelectorPrimitiveData->MeshAttribute);
+			const bool bIsString = MeshAttributeBase && MeshAttributeBase->GetTypeId() == PCG::Private::MetadataTypes<FString>::Id;
+			const bool bIsSoftObjectPath = MeshAttributeBase && MeshAttributeBase->GetTypeId() == PCG::Private::MetadataTypes<FSoftObjectPath>::Id;
+			if (!MeshAttributeBase || (!bIsString && !bIsSoftObjectPath))
+			{
+				PCG_KERNEL_VALIDATION_WARN(Context, SMSettings,
+					FText::Format(LOCTEXT("MeshAttributeInvalid", "Mesh attribute {0} not found in primitive table data, or is not of required type Soft Object Path or String."), FText::FromName(SelectorPrimitiveData->MeshAttribute)));
+				return OutDataDesc;
+			}
+
+			const FPCGMetadataAttribute<FString>* MeshAttributeString = bIsString ? static_cast<const FPCGMetadataAttribute<FString>*>(MeshAttributeBase) : nullptr;
+			const FPCGMetadataAttribute<FSoftObjectPath>* MeshAttributeSoftObjectPath = bIsSoftObjectPath ? static_cast<const FPCGMetadataAttribute<FSoftObjectPath>*>(MeshAttributeBase) : nullptr;
+
+			const int32 NumPrimitiveDataElements = PrimitiveMetadata->GetItemCountForChild();
+
+			for (int32 TableIndex = 0; TableIndex < NumPrimitiveDataElements; ++TableIndex)
+			{
+				const PCGMetadataValueKey ValueKey = MeshAttributeString ? MeshAttributeString->GetValueKey(TableIndex) : MeshAttributeSoftObjectPath->GetValueKey(TableIndex);
+				const FString MeshPath = MeshAttributeString ? MeshAttributeString->GetValue(TableIndex) : MeshAttributeSoftObjectPath->GetValue(TableIndex).ToString();
+				const int32 StringKey = InBinding->GetStringTable().IndexOfByKey(MeshPath);
+				UniqueStringKeys.AddUnique(StringKey);
+			}
+
+			for (FPCGDataDesc& DataDesc : OutDataDesc->GetDataDescriptionsMutable())
+			{
+				DataDesc.AddAttribute(FPCGKernelAttributeKey(SMSettings->OutAttributeName, EPCGKernelAttributeType::StringKey), InBinding, &UniqueStringKeys);
+			}
+		}
 		else if (SMSettings->MeshSelectorParameters)
 		{
 			UE_LOG(LogPCG, Error, TEXT("Mesh selector not supported by GPU Static Mesh Spawner: %s"), *SMSettings->MeshSelectorParameters->GetName());
@@ -164,7 +210,13 @@ void UPCGStaticMeshSpawnerKernel::GetKernelAttributeKeys(TArray<FPCGKernelAttrib
 void UPCGStaticMeshSpawnerKernel::GetInputPins(TArray<FPCGPinProperties>& OutPins) const
 {
 	OutPins.Emplace(PCGPinConstants::DefaultInputLabel, EPCGDataType::Point);
+
+	// Used to receive analysis data when using by-attribute spawning.
 	OutPins.Emplace(PCGStaticMeshSpawnerKernelConstants::InstanceCountsPinLabel, bUseRawInstanceCountsBuffer ? FPCGDataTypeInfoRawBuffer::AsId() : FPCGDataTypeInfoParam::AsId());
+
+	// Used when spawning from primitive table. This data is not actually used by the kernel/on the GPU.
+	// todo_pcg: In the future it might be nice to signal the pin doesn't require data upload to GPU.
+	OutPins.Emplace(PCGStaticMeshSpawnerKernelConstants::PrimitiveTablePinLabel, FPCGDataTypeInfoParam::AsId());
 }
 
 void UPCGStaticMeshSpawnerKernel::GetOutputPins(TArray<FPCGPinPropertiesGPU>& OutPins) const
@@ -205,10 +257,10 @@ bool UPCGStaticMeshSpawnerKernel::PerformStaticValidation()
 			}
 		}
 	}
-	else if (!SMSettings->MeshSelectorParameters || !SMSettings->MeshSelectorParameters->IsA<UPCGMeshSelectorByAttribute>())
+	else if (!SMSettings->MeshSelectorParameters || (!SMSettings->MeshSelectorParameters->IsA<UPCGMeshSelectorByAttribute>() && !SMSettings->MeshSelectorParameters->IsA<UPCGMeshSelectorPrimitiveData>()))
 	{
 #if PCG_KERNEL_LOGGING_ENABLED
-		AddStaticLogEntry(LOCTEXT("InvalidMeshSelector", "Currently GPU Static Mesh Spawner nodes must use PCGMeshSelectorWeighted or UPCGMeshSelectorByAttribute as the mesh selector type."), EPCGKernelLogVerbosity::Error);
+		AddStaticLogEntry(LOCTEXT("InvalidMeshSelector", "Currently GPU Static Mesh Spawner nodes must use PCGMeshSelectorWeighted or UPCGMeshSelectorByAttribute or UPCGMeshSelectorPrimitiveData as the mesh selector type."), EPCGKernelLogVerbosity::Error);
 #endif
 		return false;
 	}
diff --git a/Engine/Plugins/PCG/Source/PCG/Private/MeshSelectors/PCGMeshSelectorBase.cpp b/Engine/Plugins/PCG/Source/PCG/Private/MeshSelectors/PCGMeshSelectorBase.cpp
--- a/Engine/Plugins/PCG/Source/PCG/Private/MeshSelectors/PCGMeshSelectorBase.cpp
+++ b/Engine/Plugins/PCG/Source/PCG/Private/MeshSelectors/PCGMeshSelectorBase.cpp
@@ -32,7 +32,7 @@ PRAGMA_ENABLE_DEPRECATION_WARNINGS
 	else
 	{
 		PCGE_LOG_C(Error, GraphAndLog, &Context, Message);
-		return false;
+		return true;
 	}
 }
 
diff --git a/Engine/Plugins/PCG/Source/PCG/Private/MeshSelectors/PCGMeshSelectorPrimitiveData.cpp b/Engine/Plugins/PCG/Source/PCG/Private/MeshSelectors/PCGMeshSelectorPrimitiveData.cpp
new file mode 100644
--- /dev/null
+++ b/Engine/Plugins/PCG/Source/PCG/Private/MeshSelectors/PCGMeshSelectorPrimitiveData.cpp
@@ -0,0 +1,23 @@
+// Copyright Epic Games, Inc. All Rights Reserved.
+
+#include "MeshSelectors/PCGMeshSelectorPrimitiveData.h"
+
+#include "Elements/PCGStaticMeshSpawnerContext.h"
+
+#include UE_INLINE_GENERATED_CPP_BY_NAME(PCGMeshSelectorPrimitiveData)
+
+#define LOCTEXT_NAMESPACE "PCGMeshSelectorPrimitiveData"
+
+bool UPCGMeshSelectorPrimitiveData::SelectMeshInstances(
+	FPCGStaticMeshSpawnerContext& Context,
+	const UPCGStaticMeshSpawnerSettings* Settings,
+	const UPCGBasePointData* InPointData,
+	TArray<FPCGMeshInstanceList>& OutMeshInstances,
+	UPCGBasePointData* OutPointData) const
+{
+	// todo_pcg: Currently a stub until CPU support is implemented (possibly implemented via a match and set).
+	PCGE_LOG_C(Error, GraphAndLog, &Context, LOCTEXT("SelectorGPUOnly", "Primitive selection via primitive data is currently only supported for GPU execution."));
+	return true;
+}
+
+#undef LOCTEXT_NAMESPACE
diff --git a/Engine/Plugins/PCG/Source/PCG/Public/Elements/PCGStaticMeshSpawner.h b/Engine/Plugins/PCG/Source/PCG/Public/Elements/PCGStaticMeshSpawner.h
--- a/Engine/Plugins/PCG/Source/PCG/Public/Elements/PCGStaticMeshSpawner.h
+++ b/Engine/Plugins/PCG/Source/PCG/Public/Elements/PCGStaticMeshSpawner.h
@@ -41,6 +41,7 @@ protected:
 	UE_API virtual TArray<FPCGPinProperties> InputPinProperties() const override;
 	virtual TArray<FPCGPinProperties> OutputPinProperties() const override { return Super::DefaultPointOutputPinProperties(); }
 	UE_API virtual FPCGElementPtr CreateElement() const override;
+	UE_API virtual EPCGChangeType GetChangeTypeForProperty(const FName& InPropertyName) const override;
 	// ~End UPCGSettings interface
 
 public:
diff --git a/Engine/Plugins/PCG/Source/PCG/Public/Elements/PCGStaticMeshSpawnerKernel.h b/Engine/Plugins/PCG/Source/PCG/Public/Elements/PCGStaticMeshSpawnerKernel.h
--- a/Engine/Plugins/PCG/Source/PCG/Public/Elements/PCGStaticMeshSpawnerKernel.h
+++ b/Engine/Plugins/PCG/Source/PCG/Public/Elements/PCGStaticMeshSpawnerKernel.h
@@ -11,6 +11,7 @@ class UPCGStaticMeshSpawnerSettings;
 namespace PCGStaticMeshSpawnerKernelConstants
 {
 	const FName InstanceCountsPinLabel = TEXT("InstanceCounts");
+	const FName PrimitiveTablePinLabel = TEXT("PrimitiveTable");
 }
 
 UCLASS()
diff --git a/Engine/Plugins/PCG/Source/PCG/Public/MeshSelectors/PCGMeshSelectorPrimitiveData.h b/Engine/Plugins/PCG/Source/PCG/Public/MeshSelectors/PCGMeshSelectorPrimitiveData.h
new file mode 100644
--- /dev/null
+++ b/Engine/Plugins/PCG/Source/PCG/Public/MeshSelectors/PCGMeshSelectorPrimitiveData.h
@@ -0,0 +1,52 @@
+// Copyright Epic Games, Inc. All Rights Reserved.
+
+#pragma once
+
+#include "MeshSelectors/PCGMeshSelectorBase.h"
+
+#include "Metadata/PCGObjectPropertyOverride.h"
+
+#include "MeshSelectors/PCGISMDescriptor.h"
+
+#include "PCGMeshSelectorPrimitiveData.generated.h"
+
+class UPCGBasePointData;
+class UPCGStaticMeshSpawnerSettings;
+struct FPCGMeshInstanceList;
+struct FPCGStaticMeshSpawnerContext;
+
+/** Selector type that uses a table of primitive data (one row per primitive) and expects a table row index on the input points. */
+UCLASS(MinimalAPI, BlueprintType, ClassGroup = (Procedural))
+class UPCGMeshSelectorPrimitiveData : public UPCGMeshSelectorBase
+{
+	GENERATED_BODY()
+
+public:
+	PCG_API virtual bool SelectMeshInstances(
+		FPCGStaticMeshSpawnerContext& Context,
+		const UPCGStaticMeshSpawnerSettings* Settings,
+		const UPCGBasePointData* InPointData,
+		TArray<FPCGMeshInstanceList>& OutMeshInstances,
+		UPCGBasePointData* OutPointData) const override;
+
+public:
+	/** Base/default settings for the spawned primitive before overrides are applied. */
+	UPROPERTY(EditAnywhere, Category = "Primitive Setup")
+	FPCGSoftISMComponentDescriptor TemplateDescriptor;
+
+	/** Attribute on input source data that serves as a row index into the primitive table. */
+	UPROPERTY(EditAnywhere, Category = "Primitive Setup")
+	FName PrimitiveIndexAttribute;
+
+	/** Attribute in primitive data that provides the mesh for the spawned primitives. */
+	UPROPERTY(EditAnywhere, Category = "Primitive Setup")
+	FName MeshAttribute = TEXT("Mesh");
+
+	/** Per-material-slot list of attributes in primitive data that specify materials to override on the spawned primitives. */
+	UPROPERTY(EditAnywhere, Category = "Primitive Setup")
+	TArray<FName> MaterialOverrideAttributes;
+
+	/** Array of attributes in primitive data that drive properties on the spawned primitives. */
+	UPROPERTY(EditAnywhere, Category = "Primitive Setup")
+	TArray<FPCGObjectPropertyOverrideDescription> PrimitiveOverrideAttributes;
+};
