[PCG] Horde compile errors

#rnx

diff --git a/Engine/Plugins/PCG/Source/PCG/Public/Elements/PCGStaticMeshSpawner.h b/Engine/Plugins/PCG/Source/PCG/Public/Elements/PCGStaticMeshSpawner.h
--- a/Engine/Plugins/PCG/Source/PCG/Public/Elements/PCGStaticMeshSpawner.h
+++ b/Engine/Plugins/PCG/Source/PCG/Public/Elements/PCGStaticMeshSpawner.h
@@ -34,14 +34,16 @@ public:
 	UE_API virtual FText GetDefaultNodeTitle() const override;
 	virtual EPCGSettingsType GetType() const override { return EPCGSettingsType::Spawner; }
 	UE_API virtual void ApplyDeprecation(UPCGNode* InOutNode) override;
-#endif
+#endif // WITH_EDITOR
 	virtual bool UseSeed() const override { return true; }
 
 protected:
 	UE_API virtual TArray<FPCGPinProperties> InputPinProperties() const override;
 	virtual TArray<FPCGPinProperties> OutputPinProperties() const override { return Super::DefaultPointOutputPinProperties(); }
 	UE_API virtual FPCGElementPtr CreateElement() const override;
+#if WITH_EDITOR
 	UE_API virtual EPCGChangeType GetChangeTypeForProperty(const FName& InPropertyName) const override;
+#endif
 	// ~End UPCGSettings interface
 
 public:
