[PCG] Define PCG::FLock/FRecursiveLock as UE::FMutex/FRecursiveMutex
- Replace usage of FSpinLock by FMutex (avoid spinlock potential issues)
- Replace non recursive usages of FCriticalSection by PCG::FLock (reduce native platform lock)
- Replace recursive usages of FCriticalSection by PCG::FRecursiveLock (reduce native platform lock)

diff --git a/Engine/Plugins/PCG/Source/PCG/Private/Data/PCGBasePointData.cpp b/Engine/Plugins/PCG/Source/PCG/Private/Data/PCGBasePointData.cpp
--- a/Engine/Plugins/PCG/Source/PCG/Private/Data/PCGBasePointData.cpp
+++ b/Engine/Plugins/PCG/Source/PCG/Private/Data/PCGBasePointData.cpp
@@ -804,7 +804,7 @@ void UPCGBasePointData::SetPoints(const UPCGBasePointData* From, UPCGBasePointDa
 
 void UPCGBasePointData::RecomputeBounds() const
 {
-	FScopeLock Lock(&CachedDataLock);
+	PCG::TScopeLock Lock(CachedDataLock);
 	if (!bBoundsAreDirty)
 	{
 		return;
@@ -835,7 +835,7 @@ const PCGPointOctree::FPointOctree& UPCGBasePointData::GetPointOctree() const
 
 void UPCGBasePointData::RebuildOctree() const
 {
-	FScopeLock Lock(&CachedDataLock);
+	PCG::TScopeLock Lock(CachedDataLock);
 	if (!bOctreeIsDirty)
 	{
 		return;
diff --git a/Engine/Plugins/PCG/Source/PCG/Private/Data/PCGDynamicMeshData.cpp b/Engine/Plugins/PCG/Source/PCG/Private/Data/PCGDynamicMeshData.cpp
--- a/Engine/Plugins/PCG/Source/PCG/Private/Data/PCGDynamicMeshData.cpp
+++ b/Engine/Plugins/PCG/Source/PCG/Private/Data/PCGDynamicMeshData.cpp
@@ -67,12 +67,11 @@ FBox UPCGDynamicMeshData::GetBounds() const
 {
 	if (bDynamicMeshBoundsAreDirty)
 	{
-		DynamicMeshBoundsLock.Lock();
+		PCG::TScopeLock Lock(DynamicMeshBoundsLock);
 		if (bDynamicMeshBoundsAreDirty)
 		{
 			ResetBounds();
 		}
-		DynamicMeshBoundsLock.Unlock();
 	}
 
 	return CachedBounds;
@@ -145,7 +144,7 @@ const UE::Geometry::FDynamicMeshOctree3& UPCGDynamicMeshData::GetDynamicMeshOctr
 {
 	if (DynamicMesh && bDynamicMeshOctreeIsDirty)
 	{
-		FScopeLock Lock(&DynamicMeshOctreeLock);
+		PCG::TScopeLock Lock(DynamicMeshOctreeLock);
 
 		if (bDynamicMeshOctreeIsDirty)
 		{
diff --git a/Engine/Plugins/PCG/Source/PCG/Private/Data/PCGPointData.cpp b/Engine/Plugins/PCG/Source/PCG/Private/Data/PCGPointData.cpp
--- a/Engine/Plugins/PCG/Source/PCG/Private/Data/PCGPointData.cpp
+++ b/Engine/Plugins/PCG/Source/PCG/Private/Data/PCGPointData.cpp
@@ -363,7 +363,7 @@ FPCGPoint UPCGPointData::GetPoint(int32 Index) const
 void UPCGPointData::RebuildOctreeOld() const
 {
 PRAGMA_DISABLE_DEPRECATION_WARNINGS
-	FScopeLock Lock(&CachedDataLock);
+	PCG::TScopeLock Lock(CachedDataLock);
 
 	if (!bOctreeOldIsDirty)
 	{
diff --git a/Engine/Plugins/PCG/Source/PCG/Private/Data/PCGSpatialData.cpp b/Engine/Plugins/PCG/Source/PCG/Private/Data/PCGSpatialData.cpp
--- a/Engine/Plugins/PCG/Source/PCG/Private/Data/PCGSpatialData.cpp
+++ b/Engine/Plugins/PCG/Source/PCG/Private/Data/PCGSpatialData.cpp
@@ -158,29 +158,30 @@ void UPCGSpatialDataWithPointCache::GetResourceSizeEx(FResourceSizeEx& Cumulativ
 	PointArrayDataCache.GetResourceSizeEx(CumulativeResourceSize);
 }
 
-const UPCGBasePointData* FPCGPointDataCache::ToBasePointDataInternal(FPCGContext* Context, const FBox& InBounds, bool bSupportsBoundedPointData, FCriticalSection& InCacheLock, TFunctionRef<const UPCGBasePointData*(FPCGContext*, const FBox&)> CreatePointDataFunc)
+const UPCGBasePointData* FPCGPointDataCache::ToBasePointDataInternal(FPCGContext* Context, const FBox& InBounds, bool bSupportsBoundedPointData, PCG::FLock& InCacheLock, TFunctionRef<const UPCGBasePointData*(FPCGContext*, const FBox&)> CreatePointDataFunc)
 {
 	if (InBounds.IsValid && bSupportsBoundedPointData)
 	{
 		const UPCGBasePointData* BoundedPointData = nullptr;
-		InCacheLock.Lock();
-		check(CachedBoundedPointDataBoxes.Num() == CachedBoundedPointData.Num());
-		for (int CachedDataIndex = 0; CachedDataIndex < CachedBoundedPointDataBoxes.Num(); ++CachedDataIndex)
 		{
-			if (InBounds.Equals(CachedBoundedPointDataBoxes[CachedDataIndex]))
+			PCG::TScopeLock ScopeLock(InCacheLock);
+			check(CachedBoundedPointDataBoxes.Num() == CachedBoundedPointData.Num());
+			for (int CachedDataIndex = 0; CachedDataIndex < CachedBoundedPointDataBoxes.Num(); ++CachedDataIndex)
 			{
-				BoundedPointData = CachedBoundedPointData[CachedDataIndex];
-				break;
+				if (InBounds.Equals(CachedBoundedPointDataBoxes[CachedDataIndex]))
+				{
+					BoundedPointData = CachedBoundedPointData[CachedDataIndex];
+					break;
+				}
 			}
-		}
 
-		if (!BoundedPointData)
-		{
-			BoundedPointData = CreatePointDataFunc(Context, InBounds);
-			CachedBoundedPointDataBoxes.Add(InBounds);
-			CachedBoundedPointData.Add(BoundedPointData);
+			if (!BoundedPointData)
+			{
+				BoundedPointData = CreatePointDataFunc(Context, InBounds);
+				CachedBoundedPointDataBoxes.Add(InBounds);
+				CachedBoundedPointData.Add(BoundedPointData);
+			}
 		}
-		InCacheLock.Unlock();
 
 		return BoundedPointData;
 	}
@@ -188,14 +189,11 @@ const UPCGBasePointData* FPCGPointDataCache::ToBasePointDataInternal(FPCGContext
 	{
 		if (!CachedPointData)
 		{
-			InCacheLock.Lock();
-
+			PCG::TScopeLock ScopeLock(InCacheLock);
 			if (!CachedPointData)
 			{
 				CachedPointData = CreatePointDataFunc(Context, InBounds);
 			}
-
-			InCacheLock.Unlock();
 		}
 
 		return CachedPointData;
diff --git a/Engine/Plugins/PCG/Source/PCG/Private/Data/PCGWorldData.cpp b/Engine/Plugins/PCG/Source/PCG/Private/Data/PCGWorldData.cpp
--- a/Engine/Plugins/PCG/Source/PCG/Private/Data/PCGWorldData.cpp
+++ b/Engine/Plugins/PCG/Source/PCG/Private/Data/PCGWorldData.cpp
@@ -148,7 +148,7 @@ const TSet<TObjectKey<AActor>>& FPCGWorldQueryActorFilterCache::GetCachedFilterA
 {
 	if (!FilterActors.IsEmpty() && CachedFilterActorsDirty)
 	{
-		FScopeLock Lock(&CachedFilterActorsLock);
+		PCG::TScopeLock Lock(CachedFilterActorsLock);
 		if (CachedFilterActorsDirty)
 		{
 			CachedFilterActors.Empty(FilterActors.Num());
diff --git a/Engine/Plugins/PCG/Source/PCG/Private/Graph/PCGGraphCompiler.cpp b/Engine/Plugins/PCG/Source/PCG/Private/Graph/PCGGraphCompiler.cpp
--- a/Engine/Plugins/PCG/Source/PCG/Private/Graph/PCGGraphCompiler.cpp
+++ b/Engine/Plugins/PCG/Source/PCG/Private/Graph/PCGGraphCompiler.cpp
@@ -133,7 +133,7 @@ void FPCGGraphCompilerCache::RemoveFromCacheRecursive(UPCGGraph* InGraph)
 	TArray<UPCGGraph*> ParentGraphs;
 
 	{
-		FScopeLock Lock(&GraphDependenciesLock);
+		PCG::TScopeLock Lock(GraphDependenciesLock);
 		GraphDependencies.MultiFind(InGraph, ParentGraphs);
 		GraphDependencies.Remove(InGraph);
 	}
diff --git a/Engine/Plugins/PCG/Source/PCG/Private/Graph/PCGGraphCompiler.h b/Engine/Plugins/PCG/Source/PCG/Private/Graph/PCGGraphCompiler.h
--- a/Engine/Plugins/PCG/Source/PCG/Private/Graph/PCGGraphCompiler.h
+++ b/Engine/Plugins/PCG/Source/PCG/Private/Graph/PCGGraphCompiler.h
@@ -43,7 +43,7 @@ private:
 	TMap<UPCGGraph*, TMap<uint32, TArray<TObjectPtr<UPCGComputeGraph>>>> TopGraphToComputeGraphMap;
 
 #if WITH_EDITOR
-	FCriticalSection GraphDependenciesLock;
+	PCG::FLock GraphDependenciesLock;
 	TMultiMap<UPCGGraph*, UPCGGraph*> GraphDependencies;
 #endif // WITH_EDITOR
 };
diff --git a/Engine/Plugins/PCG/Source/PCG/Private/Graph/PCGGraphExecutor.cpp b/Engine/Plugins/PCG/Source/PCG/Private/Graph/PCGGraphExecutor.cpp
--- a/Engine/Plugins/PCG/Source/PCG/Private/Graph/PCGGraphExecutor.cpp
+++ b/Engine/Plugins/PCG/Source/PCG/Private/Graph/PCGGraphExecutor.cpp
@@ -164,49 +164,6 @@ namespace PCGGraphExecutor
 			}
 		}
 	}
-		
-	// Similar to existing UE::TScopeLock but allows optional TryLock in which case caller is responsible to 
-	// check if Lock succeeded or not with IsLocked(). This allows execution code to be skipped if non main thread
-	// tasks fail to get the lock instead of spinning for nothing
-	template<typename MutexType>
-	class TScopeLock
-	{
-	public:
-		UE_NONCOPYABLE(TScopeLock);
-
-		UE_NODISCARD_CTOR TScopeLock(MutexType& InMutex, bool bTryLock=false)
-			: Mutex(&InMutex), bLocked(false)
-		{
-			check(Mutex);
-			if (bTryLock)
-			{
-				bLocked = Mutex->TryLock();
-			}
-			else
-			{
-				Mutex->Lock();
-				bLocked = true;
-			}
-		}
-
-		~TScopeLock()
-		{
-			if (bLocked && Mutex)
-			{
-				Mutex->Unlock();
-				Mutex = nullptr;
-			}
-		}
-			
-		bool IsLocked() const
-		{
-			return bLocked;
-		}
-
-	private:
-		MutexType* Mutex;
-		bool bLocked;
-	};
 }
 
 FPCGGraphExecutor::FPCGGraphExecutor()
@@ -328,7 +285,7 @@ FPCGTaskId FPCGGraphExecutor::ScheduleGraph(const FPCGScheduleGraphParams& InPar
 		// Setup fetch task on input node
 		CompiledTasks[0].Element = InParams.InputElement;
 
-		PCGGraphExecutor::TScopeLock ScopeLock(ScheduleLock);
+		PCG::TScopeLock ScopeLock(ScheduleLock);
 
 		FPCGGraphScheduleTask& ScheduledTask = ScheduledTasks.Emplace_GetRef();
 		ScheduledTask.Tasks = MoveTemp(CompiledTasks);
@@ -409,7 +366,7 @@ TSet<IPCGGraphExecutionSource*> FPCGGraphExecutor::Cancel(TFunctionRef<bool(TWea
 
 	// Visit scheduled tasks
 	{
-		PCGGraphExecutor::TScopeLock ScopeLock(ScheduleLock);
+		PCG::TScopeLock ScopeLock(ScheduleLock);
 		for (const FPCGGraphScheduleTask& ScheduledTask : ScheduledTasks)
 		{
 			if (CancelFilter(ScheduledTask.ExecutionSource))
@@ -421,7 +378,7 @@ TSet<IPCGGraphExecutionSource*> FPCGGraphExecutor::Cancel(TFunctionRef<bool(TWea
 		
 	{
 		// Visit ready tasks
-		PCGGraphExecutor::TScopeLock ScopeLock(LiveTasksLock);
+		PCG::TScopeLock ScopeLock(LiveTasksLock);
 		for (const FPCGGraphTask& Task : ReadyTasks)
 		{
 			if (CancelFilter(Task.ExecutionSource))
@@ -483,7 +440,7 @@ TSet<IPCGGraphExecutionSource*> FPCGGraphExecutor::Cancel(TFunctionRef<bool(TWea
 
 		// Remove from scheduled tasks
 		{
-			PCGGraphExecutor::TScopeLock ScopeLock(ScheduleLock);
+			PCG::TScopeLock ScopeLock(ScheduleLock);
 			for (int32 ScheduledTaskIndex = ScheduledTasks.Num() - 1; ScheduledTaskIndex >= 0; --ScheduledTaskIndex)
 			{
 				FPCGGraphScheduleTask& ScheduledTask = ScheduledTasks[ScheduledTaskIndex];
@@ -550,7 +507,7 @@ TSet<IPCGGraphExecutionSource*> FPCGGraphExecutor::Cancel(TFunctionRef<bool(TWea
 					}
 				};
 				
-				PCGGraphExecutor::TScopeLock ScopeLock(LiveTasksLock);
+				PCG::TScopeLock ScopeLock(LiveTasksLock);
 				CancelActiveTasks(ActiveTasks);
 				CancelActiveTasks(ActiveTasksGameThreadOnly);
 			}
@@ -573,7 +530,7 @@ TSet<IPCGGraphExecutionSource*> FPCGGraphExecutor::Cancel(TFunctionRef<bool(TWea
 		{
 			TArray<FPCGGraphTask> CancelledReadyTasks;
 			{
-				PCGGraphExecutor::TScopeLock ScopeLock(LiveTasksLock);
+				PCG::TScopeLock ScopeLock(LiveTasksLock);
 				for (int32 ReadyTaskIndex = ReadyTasks.Num() - 1; ReadyTaskIndex >= 0; --ReadyTaskIndex)
 				{
 					FPCGGraphTask& Task = ReadyTasks[ReadyTaskIndex];
@@ -605,8 +562,8 @@ TSet<IPCGGraphExecutionSource*> FPCGGraphExecutor::Cancel(TFunctionRef<bool(TWea
 		{
 			TArray<TSharedPtr<FPCGGraphActiveTask>> CancelledPausedTasks;
 			{
-				PCGGraphExecutor::TScopeLock ScopeLock(LiveTasksLock);
-				PCGGraphExecutor::TScopeLock PauseScopeLock(PausedTaskSuccessorsLock);
+				PCG::TScopeLock ScopeLock(LiveTasksLock);
+				PCG::TScopeLock PauseScopeLock(PausedTaskSuccessorsLock);
 				for (int32 PausedTaskIndex = PausedTasks.Num() - 1; PausedTaskIndex >= 0; --PausedTaskIndex)
 				{
 					TSharedPtr<FPCGGraphActiveTask> Task = PausedTasks[PausedTaskIndex];
@@ -783,7 +740,7 @@ FPCGTaskId FPCGGraphExecutor::ScheduleSingleElement(TSharedRef<IPCGElement> Elem
 	Task.Element = Element;
 
 	{
-		PCGGraphExecutor::TScopeLock ScopeLock(ScheduleLock);
+		PCG::TScopeLock ScopeLock(ScheduleLock);
 
 		// Assign task id
 		Task.NodeId = NextTaskId++;
@@ -808,7 +765,7 @@ void FPCGGraphExecutor::MarkInputResults(TArrayView<const FPCGTaskId> InInputRes
 
 	if (InInputResults.Num() > 0)
 	{
-		PCGGraphExecutor::TScopeLock ScopeLock(TaskOutputsLock);
+		PCG::TScopeLock ScopeLock(TaskOutputsLock);
 		for (FPCGTaskId TaskId : InInputResults)
 		{
 			FOutputDataInfo& OutputDataInfo = TaskOutputs[TaskId];
@@ -825,7 +782,7 @@ void FPCGGraphExecutor::MarkInputResults(TArrayView<const FPCGTaskId> InInputRes
 bool FPCGGraphExecutor::GetOutputData(FPCGTaskId TaskId, FPCGDataCollection& OutData)
 {
 	TRACE_CPUPROFILER_EVENT_SCOPE(FPCGGraphExecutor::GetOutputData);
-	PCGGraphExecutor::TScopeLock ScopeLock(TaskOutputsLock);
+	PCG::TScopeLock ScopeLock(TaskOutputsLock);
 	if (FOutputDataInfo* OutputDataInfo = TaskOutputs.Find(TaskId))
 	{
 		ensure(!OutputDataInfo->bCulled);
@@ -841,7 +798,7 @@ bool FPCGGraphExecutor::GetOutputData(FPCGTaskId TaskId, FPCGDataCollection& Out
 void FPCGGraphExecutor::ClearOutputData(FPCGTaskId TaskId)
 {
 	TRACE_CPUPROFILER_EVENT_SCOPE(FPCGGraphExecutor::ClearOutputData);
-	PCGGraphExecutor::TScopeLock ScopeLock(TaskOutputsLock);
+	PCG::TScopeLock ScopeLock(TaskOutputsLock);
 	if (FOutputDataInfo* OutputDataInfo = TaskOutputs.Find(TaskId); OutputDataInfo && ensure(OutputDataInfo->bNeedsManualClear))
 	{
 		TaskOutputs.Remove(TaskId);
@@ -871,7 +828,7 @@ void FPCGGraphExecutor::CheckState(bool bIsInGameThread)
 	}
 
 	// Since this is for an error case avoid locking if lock is already taken, we can comeback later but if we haven't checked in more than 3 seconds actually lock
-	PCGGraphExecutor::TScopeLock ScheduleScopeLock(ScheduleLock, /*bTryLock=*/ ElapsedTimeSinceLastCheck < 3.0);
+	PCG::TScopeLock ScheduleScopeLock(ScheduleLock, /*bTryLock=*/ ElapsedTimeSinceLastCheck < 3.0);
 	if (!ScheduleScopeLock.IsLocked())
 	{
 		return;
@@ -885,10 +842,10 @@ void FPCGGraphExecutor::CheckState(bool bIsInGameThread)
 		return;
 	}
 		
-	PCGGraphExecutor::TScopeLock ScopeLock(TasksLock);
-	PCGGraphExecutor::TScopeLock ChildScopeLock(LiveTasksLock);
-	PCGGraphExecutor::TScopeLock PausedLock(PausedTaskSuccessorsLock);
-	PCGGraphExecutor::TScopeLock CachingResultsScopeLock(CachingResultsLock);
+	PCG::TScopeLock ScopeLock(TasksLock);
+	PCG::TScopeLock ChildScopeLock(LiveTasksLock);
+	PCG::TScopeLock PausedLock(PausedTaskSuccessorsLock);
+	PCG::TScopeLock CachingResultsScopeLock(CachingResultsLock);
 
 	// This is a safeguard to check if we're in a stuck state
 	if (CachingResultsForGC.Num() == 0 && ReadyTasks.Num() == 0 && ActiveTasks.Num() == 0 && ActiveTasksGameThreadOnly.Num() == 0 && PausedTasks.Num() == 0 && Tasks.Num() > 0)
@@ -907,7 +864,7 @@ bool FPCGGraphExecutor::ProcessScheduledTasks(double EndTime)
 
 	{
 		// Process any newly scheduled graphs to execute
-		PCGGraphExecutor::TScopeLock ScopeLock(ScheduleLock, !bIsInGameThread);
+		PCG::TScopeLock ScopeLock(ScheduleLock, !bIsInGameThread);
 		if (ScopeLock.IsLocked())
 		{
 			if (ScheduledTasks.Num() > 0)
@@ -917,8 +874,8 @@ bool FPCGGraphExecutor::ProcessScheduledTasks(double EndTime)
 				{
 					TArray<FPCGGraphTask, TInlineAllocator<16>> NewReadyTasks;
 					{
-						PCGGraphExecutor::TScopeLock ChildScopeLock(TasksLock);
-						PCGGraphExecutor::TScopeLock OutputLock(TaskOutputsLock);
+						PCG::TScopeLock ChildScopeLock(TasksLock);
+						PCG::TScopeLock OutputLock(TaskOutputsLock);
 
 						// Every X ScheduledTasks, we want to check the frame budget and early out if we consumed it.
 						const int32 MaxScheduledTasksPerBudgetCheck = 32;
@@ -1015,11 +972,11 @@ void FPCGGraphExecutor::ExecuteTasksEnded()
 	if (GetNonScheduledRemainingTaskCount() == 0)
 	{
 		// Make sure no scheduling can happen here (and check condition again)
-		PCGGraphExecutor::TScopeLock ScopeLock(ScheduleLock);
+		PCG::TScopeLock ScopeLock(ScheduleLock);
 		if (GetNonScheduledRemainingTaskCount() == 0)
 		{
 			{
-				PCGGraphExecutor::TScopeLock ScopeChildLock(TasksLock);
+				PCG::TScopeLock ScopeChildLock(TasksLock);
 				if (!ensure(TaskSuccessors.IsEmpty()))
 				{
 					TaskSuccessors.Reset();
@@ -1076,7 +1033,7 @@ void FPCGGraphExecutor::PostTaskExecute(TSharedPtr<FPCGGraphActiveTask> ActiveTa
 	// If input data collection was modified (for example replacing proxies with full data), add any new references.
 	if (ActiveTask.Context->bInputDataModified)
 	{
-		PCGGraphExecutor::TScopeLock ScopeLock(LiveTasksLock);
+		PCG::TScopeLock ScopeLock(LiveTasksLock);
 
 		for (const FPCGTaggedData& Data : ActiveTask.Context->InputData.TaggedData)
 		{
@@ -1092,7 +1049,7 @@ void FPCGGraphExecutor::PostTaskExecute(TSharedPtr<FPCGGraphActiveTask> ActiveTa
 	if (!bIsDone || !ActiveTask.Context->DynamicDependencies.IsEmpty())
 	{
 		TRACE_CPUPROFILER_EVENT_SCOPE(FPCGGraphExecutor::PostTaskExecute::NotDone);
-		PCGGraphExecutor::TScopeLock ScopeLock(LiveTasksLock);
+		PCG::TScopeLock ScopeLock(LiveTasksLock);
 		
 		// Once we get the lock check for cancellation again in case task was cancelled while executing
 		if (!ActiveTaskPtr->bWasCancelled)
@@ -1122,8 +1079,8 @@ void FPCGGraphExecutor::PostTaskExecute(TSharedPtr<FPCGGraphActiveTask> ActiveTa
 				// This could be supported but for now we have no use for this added complexity.
 				if (bHasPausedDependencies)
 				{
-					PCGGraphExecutor::TScopeLock PauseLock(PausedTaskSuccessorsLock);
-					PCGGraphExecutor::TScopeLock OutputLock(TaskOutputsLock);
+					PCG::TScopeLock PauseLock(PausedTaskSuccessorsLock);
+					PCG::TScopeLock OutputLock(TaskOutputsLock);
 					for (TSet<FPCGTaskId>::TIterator It = ActiveTask.Context->DynamicDependencies.CreateIterator(); It; ++It)
 					{
 						FPCGTaskId DependencyId = *It;
@@ -1204,7 +1161,7 @@ void FPCGGraphExecutor::PostTaskExecute(TSharedPtr<FPCGGraphActiveTask> ActiveTa
 			// Debug display can span multiple frames so hold onto reference to task.
 			// TODO: Investigate scheduling an additional task to perform the debug display rather than manually holding onto references here.
 			{
-				PCGGraphExecutor::TScopeLock ScopeLock(CollectGCReferenceTasksLock);
+				PCG::TScopeLock ScopeLock(CollectGCReferenceTasksLock);
 				CollectGCReferenceTasks.Add(ActiveTaskPtr);
 			}
 
@@ -1253,7 +1210,7 @@ void FPCGGraphExecutor::PostTaskExecute(TSharedPtr<FPCGGraphActiveTask> ActiveTa
 
 	// Erase from ActiveTasks
 	{
-		PCGGraphExecutor::TScopeLock ScopeLock(LiveTasksLock);
+		PCG::TScopeLock ScopeLock(LiveTasksLock);
 		TArray<TSharedPtr<FPCGGraphActiveTask>>& ActiveTaskArray = ActiveTask.bIsGameThreadOnly ? ActiveTasksGameThreadOnly : ActiveTasks;
 		const int32 TaskIndex = ActiveTask.TaskIndex >= 0 ? ActiveTask.TaskIndex : ActiveTaskArray.Find(ActiveTaskPtr);
 		if (ActiveTaskArray.IsValidIndex(TaskIndex))
@@ -1288,7 +1245,7 @@ void FPCGGraphExecutor::OnTaskInputsReady(FPCGGraphTask& Task, TArray<FCachedRes
 			PrepareForExecute(Task, CachedResult, /*bLiveTasksLockAlreadyLocked=*/false);
 		}
 
-		PCGGraphExecutor::TScopeLock ChildScopeLock(LiveTasksLock);
+		PCG::TScopeLock ChildScopeLock(LiveTasksLock);
 		if (CachedResult)
 		{
 			check(CachedResult->TaskId != InvalidPCGTaskId);
@@ -1389,7 +1346,7 @@ void FPCGGraphExecutor::PrepareForExecute(FPCGGraphTask& Task, FCachedResult*& O
 
 	// This section requires the lock on the caching results to prevent interaction between different tasks and the GC.
 	{
-		PCGGraphExecutor::TScopeLock ScopeLock(CachingResultsLock);
+		PCG::TScopeLock ScopeLock(CachingResultsLock);
 		FCachedResult LocalCachedResult;
 		CacheStatus = Task.Element->RetrieveResultsFromCache(&GraphCache, Task.Node, Task.TaskInput, Task.ExecutionSource.Get(), LocalCachedResult.Output, &DependenciesCrc);
 
@@ -1478,7 +1435,7 @@ void FPCGGraphExecutor::PrepareForExecute(FPCGGraphTask& Task, FCachedResult*& O
 		Task.bIsBypassed = true;
 
 		// Directly transfer Cache result for debugged tasks to Context output
-		PCGGraphExecutor::TScopeLock ScopeLock(CachingResultsLock);
+		PCG::TScopeLock ScopeLock(CachingResultsLock);
 		Task.Context->OutputData = MoveTemp(OutCachedResult->Output);
 		OutCachedResult = nullptr;
 		CachingResultsForGC.Remove(Task.NodeId);
@@ -1507,7 +1464,7 @@ void FPCGGraphExecutor::DebugDisplayTask(TWeakPtr<FPCGGraphExecutor::FGameThread
 	}
 
 	// If we fall through to here, debug display completed. Cleanup.
-	PCGGraphExecutor::TScopeLock ScopeLock(Handler->GetExecutor()->CollectGCReferenceTasksLock);
+	PCG::TScopeLock ScopeLock(Handler->GetExecutor()->CollectGCReferenceTasksLock);
 	Handler->GetExecutor()->CollectGCReferenceTasks.Remove(ActiveTaskPtr);
 }
 #endif // WITH_EDITOR
@@ -1557,7 +1514,7 @@ bool FPCGGraphExecutor::ExecuteScheduling(double EndTime, TSharedPtr<FPCGGraphAc
 	{
 		TRACE_CPUPROFILER_EVENT_SCOPE(FPCGGraphExecutor::ExecuteScheduling::ReadyTasks);
 		const bool bTryLock = (OutMainThreadTask == nullptr);
-		PCGGraphExecutor::TScopeLock ScopeLock(LiveTasksLock, bTryLock);
+		PCG::TScopeLock ScopeLock(LiveTasksLock, bTryLock);
 		if (ScopeLock.IsLocked())
 		{
 			if (bNeedToCheckPausedTasks || bForceCheckPausedTasks)
@@ -1662,7 +1619,7 @@ bool FPCGGraphExecutor::ExecuteScheduling(double EndTime, TSharedPtr<FPCGGraphAc
 
 	// Dispatch Tasks
 	{
-		PCGGraphExecutor::TScopeLock ScopeLock(LiveTasksLock, !bIsInGameThread);
+		PCG::TScopeLock ScopeLock(LiveTasksLock, !bIsInGameThread);
 		if (ScopeLock.IsLocked())
 		{
 			if (bIsInGameThread && OutMainThreadTask)
@@ -1961,8 +1918,8 @@ TArray<FPCGGraphExecutor::FCachedResult*> FPCGGraphExecutor::QueueNextTasksInter
 	TSet<FPCGTaskId> Successors;
 	TArray<FPCGGraphTask, TInlineAllocator<16>> NewReadyTasks;
 	{
-		PCGGraphExecutor::TScopeLock ScopeLock(TasksLock);
-		PCGGraphExecutor::TScopeLock OutputLock(TaskOutputsLock);
+		PCG::TScopeLock ScopeLock(TasksLock);
+		PCG::TScopeLock OutputLock(TaskOutputsLock);
 
 		if (TaskSuccessors.RemoveAndCopyValue(FinishedTask, Successors))
 		{
@@ -2001,7 +1958,7 @@ TArray<FPCGGraphExecutor::FCachedResult*> FPCGGraphExecutor::QueueNextTasksInter
 	}
 
 	{
-		PCGGraphExecutor::TScopeLock PauseScopeLock(PausedTaskSuccessorsLock);
+		PCG::TScopeLock PauseScopeLock(PausedTaskSuccessorsLock);
 
 		TArray<TSharedPtr<FPCGGraphActiveTask>> ActiveTaskSuccessors;
 
@@ -2065,7 +2022,7 @@ TArray<FPCGTaskId> FPCGGraphExecutor::ProcessCachedResultsInternal(TArray<FCache
 
 	if (CachedResults.Num() > 0)
 	{
-		PCGGraphExecutor::TScopeLock ScopeLock(CachingResultsLock);
+		PCG::TScopeLock ScopeLock(CachingResultsLock);
 		for (const FCachedResult* CachedResult : CachedResults)
 		{
 			check(CachingResultsForGC.Contains(CachedResult->TaskId));
@@ -2087,7 +2044,7 @@ bool FPCGGraphExecutor::CancelNextTasks(FPCGTaskId CancelledTask, TSet<IPCGGraph
 
 	TSet<FPCGTaskId> LocalSuccessors;
 	{
-		PCGGraphExecutor::TScopeLock ScopeLock(TasksLock);
+		PCG::TScopeLock ScopeLock(TasksLock);
 		if (TSet<FPCGTaskId>* Successors = TaskSuccessors.Find(CancelledTask))
 		{
 			LocalSuccessors = MoveTemp(*Successors);
@@ -2121,7 +2078,7 @@ bool FPCGGraphExecutor::CancelNextTasks(FPCGTaskId CancelledTask, TSet<IPCGGraph
 	}
 
 	{
-		PCGGraphExecutor::TScopeLock PauseScopeLock(PausedTaskSuccessorsLock);
+		PCG::TScopeLock PauseScopeLock(PausedTaskSuccessorsLock);
 
 		TArray<TSharedPtr<FPCGGraphActiveTask>> ActiveTaskSuccessors;
 
@@ -2142,7 +2099,7 @@ bool FPCGGraphExecutor::CancelNextTasks(FPCGTaskId CancelledTask, TSet<IPCGGraph
 
 	// Tasks cancelled might have an impact on scheduled-but-not-processed tasks
 	{
-		PCGGraphExecutor::TScopeLock ScopeLock(ScheduleLock);
+		PCG::TScopeLock ScopeLock(ScheduleLock);
 		for (FPCGGraphScheduleTask& ScheduledTask : ScheduledTasks)
 		{
 			if (!OutCancelledExecutionSources.Contains(ScheduledTask.ExecutionSource.Get()) &&
@@ -2158,7 +2115,7 @@ bool FPCGGraphExecutor::CancelNextTasks(FPCGTaskId CancelledTask, TSet<IPCGGraph
 
 void FPCGGraphExecutor::RemoveTaskFromInputSuccessors(FPCGTaskId CancelledTask, const TArray<FPCGGraphTaskInput>& CancelledTaskInputs)
 {
-	PCGGraphExecutor::TScopeLock ScopeLock(TasksLock);
+	PCG::TScopeLock ScopeLock(TasksLock);
 	RemoveTaskFromInputSuccessorsNoLock(CancelledTask, CancelledTaskInputs);
 }
 
@@ -2219,7 +2176,7 @@ void FPCGGraphExecutor::BuildTaskInput(FPCGGraphTask& Task, TArray<FPCGTaskId>&
 	
 	if(Task.Inputs.Num() > 0)
 	{
-		PCGGraphExecutor::TScopeLock OutputLock(TaskOutputsLock);
+		PCG::TScopeLock OutputLock(TaskOutputsLock);
 		for (const FPCGGraphTaskInput& Input : Task.Inputs)
 		{
 			const FOutputDataInfo& TaskOutput = TaskOutputs.FindChecked(Input.TaskId);
@@ -2395,7 +2352,7 @@ void FPCGGraphExecutor::StoreResults(FPCGTaskId InTaskId, const FPCGDataCollecti
 
 	{
 		TRACE_CPUPROFILER_EVENT_SCOPE(FPCGGraphExecutor::StoreResults::Successors);
-		PCGGraphExecutor::TScopeLock ScopeLock(TasksLock);
+		PCG::TScopeLock ScopeLock(TasksLock);
 		if (TSet<FPCGTaskId>* Successors = TaskSuccessors.Find(InTaskId))
 		{
 			OutputDataInfo.RemainingSuccessorCount = Successors->Num();
@@ -2403,7 +2360,7 @@ void FPCGGraphExecutor::StoreResults(FPCGTaskId InTaskId, const FPCGDataCollecti
 	}
 
 	// Store output in map
-	PCGGraphExecutor::TScopeLock OutputLock(TaskOutputsLock);
+	PCG::TScopeLock OutputLock(TaskOutputsLock);
 	TaskOutputs.Add(InTaskId, MoveTemp(OutputDataInfo));
 }
 
@@ -2426,7 +2383,7 @@ void FPCGGraphExecutor::ClearResults()
 		}
 	}
 
-	PCGGraphExecutor::TScopeLock OutputLock(TaskOutputsLock);
+	PCG::TScopeLock OutputLock(TaskOutputsLock);
 
 	for (const TPair<FPCGTaskId, FOutputDataInfo>& TaskOutput : TaskOutputs)
 	{
@@ -2471,7 +2428,7 @@ void FPCGGraphExecutor::CullInactiveDownstreamNodes(FPCGTaskId InCompletedTaskId
 
 	if (PinIdsToDeactivate.Num() > 0)
 	{
-		PCGGraphExecutor::TScopeLock ScopeLock(TasksLock);
+		PCG::TScopeLock ScopeLock(TasksLock);
 		while (!PinIdsToDeactivate.IsEmpty())
 		{
 			const FPCGPinId PinId = PinIdsToDeactivate.Pop(EAllowShrinking::No);
@@ -2568,7 +2525,7 @@ void FPCGGraphExecutor::CullInactiveDownstreamNodes(FPCGTaskId InCompletedTaskId
 	{
 		{
 			// Store culled output
-			PCGGraphExecutor::TScopeLock OutputLock(TaskOutputsLock);
+			PCG::TScopeLock OutputLock(TaskOutputsLock);
 			for (const FPCGTaskId TaskId : AllRemovedTasks)
 			{
 				FOutputDataInfo CullDataInfo;
@@ -2601,13 +2558,13 @@ void FPCGGraphExecutor::SendInactivePinNotification(const UPCGNode* InNode, cons
 
 void FPCGGraphExecutor::AddReleasedContextForGC(FPCGContext* InContext)
 {
-	PCGGraphExecutor::TScopeLock ScopeLock(ReleasedContextLock);
+	PCG::TScopeLock ScopeLock(ReleasedContextLock);
 	ReleasedContextsForGC.Add(InContext);
 }
 
 void FPCGGraphExecutor::RemoveReleasedContextForGC(FPCGContext* InContext)
 {
-	PCGGraphExecutor::TScopeLock ScopeLock(ReleasedContextLock);
+	PCG::TScopeLock ScopeLock(ReleasedContextLock);
 	ensure(ReleasedContextsForGC.Remove(InContext) > 0);
 }
 
@@ -2617,7 +2574,7 @@ void FPCGGraphExecutor::AddReferencedObjects(FReferenceCollector& Collector)
 
 	// Go through all data in the cached output map
 	{
-		PCGGraphExecutor::TScopeLock ScopeLock(TaskOutputsLock);
+		PCG::TScopeLock ScopeLock(TaskOutputsLock);
 		for (auto& OutputDataEntry : TaskOutputs)
 		{
 			OutputDataEntry.Value.DataCollection.AddReferences(Collector);
@@ -2625,7 +2582,7 @@ void FPCGGraphExecutor::AddReferencedObjects(FReferenceCollector& Collector)
 	}
 
 	{
-		PCGGraphExecutor::TScopeLock ScopeLock(LiveTasksLock);
+		PCG::TScopeLock ScopeLock(LiveTasksLock);
 		Algo::ForEach(ReadyTasks, [&Collector](auto& Task)
 		{
 			if (!Task.bHasDonePrepareForExecute)
@@ -2662,7 +2619,7 @@ void FPCGGraphExecutor::AddReferencedObjects(FReferenceCollector& Collector)
 	}
 
 	{
-		PCGGraphExecutor::TScopeLock ScopeLock(CachingResultsLock);
+		PCG::TScopeLock ScopeLock(CachingResultsLock);
 		Algo::ForEach(CachingResultsForGC, [&Collector](auto& CachingResult)
 		{
 			CachingResult.Value->Output.AddReferences(Collector);
@@ -2672,7 +2629,7 @@ void FPCGGraphExecutor::AddReferencedObjects(FReferenceCollector& Collector)
 	// Some task might have been delayed to the game thread and we are no longer referencing them through (Ready/Paused/Active)
 	// but we still want their UObject references to stay valid
 	{
-		PCGGraphExecutor::TScopeLock ScopeLock(CollectGCReferenceTasksLock);
+		PCG::TScopeLock ScopeLock(CollectGCReferenceTasksLock);
 		Algo::ForEach(CollectGCReferenceTasks, [&Collector](auto& Task)
 		{
 			if (Task->Context)
@@ -2685,7 +2642,7 @@ void FPCGGraphExecutor::AddReferencedObjects(FReferenceCollector& Collector)
 	PerExecutionCache.AddStructReferencedObjects(Collector);
 
 	{
-		PCGGraphExecutor::TScopeLock ScopeLock(ReleasedContextLock);
+		PCG::TScopeLock ScopeLock(ReleasedContextLock);
 		Algo::ForEach(ReleasedContextsForGC, [&Collector](FPCGContext* InContext)
 		{
 			if(ensure(InContext))
diff --git a/Engine/Plugins/PCG/Source/PCG/Private/Graph/PCGGraphExecutor.h b/Engine/Plugins/PCG/Source/PCG/Private/Graph/PCGGraphExecutor.h
--- a/Engine/Plugins/PCG/Source/PCG/Private/Graph/PCGGraphExecutor.h
+++ b/Engine/Plugins/PCG/Source/PCG/Private/Graph/PCGGraphExecutor.h
@@ -13,7 +13,6 @@
 #include "Graph/PCGGraphTask.h"
 #include "Graph/PCGStackContext.h"
 
-#include "Misc/SpinLock.h"
 #include "Tasks/Task.h"
 #include "Templates/UniquePtr.h"
 #include "UObject/GCObject.h"
@@ -294,20 +293,19 @@ private:
 	 * Define a Lock level for future reference. Rule is when we have a lock, we can't lock a lower or equal level lock to prevent deadlocks.
 	 * Example: When locking LiveTasksLock we can't lock ScheduleLock or CollectGCReferenceTasksLock 
 	 */
-	using FPCGLock = UE::FSpinLock;
 
 	/** Lock level - 1 (top most lock) */
-	FPCGLock ScheduleLock;
+	PCG::FLock ScheduleLock;
 	TArray<FPCGGraphScheduleTask> ScheduledTasks;
 	std::atomic<bool> SchedulingDisabled = false;
 
 	/** Lock level 2 */
-	mutable FPCGLock TasksLock;
+	mutable PCG::FLock TasksLock;
 	TMap<FPCGTaskId, FPCGGraphTask> Tasks;
 	TMap<FPCGTaskId, TSet<FPCGTaskId>> TaskSuccessors;
 
 	/** Lock level 3 */
-	FPCGLock LiveTasksLock;
+	PCG::FLock LiveTasksLock;
 	TArray<FPCGGraphTask> ReadyTasks;
 	TArray<TSharedPtr<FPCGGraphActiveTask>> ActiveTasks;
 	TArray<TSharedPtr<FPCGGraphActiveTask>> ActiveTasksGameThreadOnly;
@@ -315,7 +313,7 @@ private:
 	bool bNeedToCheckPausedTasks = false;
 
 	/** Lock level leaf */
-	FPCGLock CollectGCReferenceTasksLock;
+	PCG::FLock CollectGCReferenceTasksLock;
 	TSet<TSharedPtr<FPCGGraphActiveTask>> CollectGCReferenceTasks;
 			
 	/** Map of node instances to their output, could be cleared once execution is done */
@@ -332,20 +330,20 @@ private:
 	};
 
 	/** Lock level 4 */
-	UE::FSpinLock PausedTaskSuccessorsLock;
+	PCG::FLock PausedTaskSuccessorsLock;
 	TMap<FPCGTaskId, TArray<TSharedPtr<FPCGGraphActiveTask>>> PausedTaskSuccessors;
 	
 	/** Lock level leaf */
-	FPCGLock CachingResultsLock;
+	PCG::FLock CachingResultsLock;
 	// Used to keep GC references to in flight caching results (not yet stored to output and might not be in cache anymore)
 	TMap<FPCGTaskId, TUniquePtr<FCachedResult>> CachingResultsForGC;
 
 	/** Lock level leaf */
-	FPCGLock ReleasedContextLock;
+	PCG::FLock ReleasedContextLock;
 	TSet<FPCGContext*> ReleasedContextsForGC;
 
 	/** Lock level leaf */
-	FPCGLock TaskOutputsLock;
+	PCG::FLock TaskOutputsLock;
 	TMap<FPCGTaskId, FOutputDataInfo> TaskOutputs;
 	
 	/** Monotonically increasing id. Should be reset once all tasks are executed, should be protected by the ScheduleLock */
diff --git a/Engine/Plugins/PCG/Source/PCG/Private/Grid/PCGLandscapeCache.cpp b/Engine/Plugins/PCG/Source/PCG/Private/Grid/PCGLandscapeCache.cpp
--- a/Engine/Plugins/PCG/Source/PCG/Private/Grid/PCGLandscapeCache.cpp
+++ b/Engine/Plugins/PCG/Source/PCG/Private/Grid/PCGLandscapeCache.cpp
@@ -328,7 +328,7 @@ bool FPCGLandscapeCacheEntry::TouchAndLoad(int32 InTouch) const
 
 	if (!bDataLoaded)
 	{
-		FScopeLock ScopeDataLock(&DataLock);
+		PCG::TScopeLock ScopeDataLock(DataLock);
 		if (!bDataLoaded)
 		{
 			return SerializeFromBulkData();
diff --git a/Engine/Plugins/PCG/Source/PCG/Private/PCGActorAndComponentMapping.cpp b/Engine/Plugins/PCG/Source/PCG/Private/PCGActorAndComponentMapping.cpp
--- a/Engine/Plugins/PCG/Source/PCG/Private/PCGActorAndComponentMapping.cpp
+++ b/Engine/Plugins/PCG/Source/PCG/Private/PCGActorAndComponentMapping.cpp
@@ -182,7 +182,7 @@ void FPCGActorAndComponentMapping::Tick()
 	TSet<UPCGComponent*> ComponentsToUnregister;
 	TSet<UPCGComponent*> ComponentsToRegister;
 	{
-		FScopeLock Lock(&DelayedComponentsRegistrationLock);
+		PCG::TScopeLock Lock(DelayedComponentsRegistrationLock);
 		ComponentsToUnregister = MoveTemp(DelayedComponentsToUnregister);
 		ComponentsToRegister = MoveTemp(DelayedComponentsToRegister);
 	}
@@ -334,7 +334,7 @@ bool FPCGActorAndComponentMapping::RegisterOrUpdatePCGComponent(UPCGComponent* I
 
 #if WITH_EDITOR
 	{
-		FScopeLock Lock(&DelayedComponentsRegistrationLock);
+		PCG::TScopeLock Lock(DelayedComponentsRegistrationLock);
 
 		DelayedComponentsToUnregister.Remove(InComponent);
 		DelayedComponentsToRegister.Remove(InComponent);
@@ -466,7 +466,7 @@ bool FPCGActorAndComponentMapping::RemapPCGComponent(const UPCGComponent* OldCom
 #if WITH_EDITOR
 	// Remove it from the delayed
 	{
-		FScopeLock Lock(&DelayedComponentsRegistrationLock);
+		PCG::TScopeLock Lock(DelayedComponentsRegistrationLock);
 		DelayedComponentsToUnregister.Remove(OldComponent);
 		DelayedComponentsToRegister.Remove(NewComponent);
 	}
@@ -515,7 +515,7 @@ void FPCGActorAndComponentMapping::UnregisterPCGComponent(UPCGComponent* InCompo
 
 #if WITH_EDITOR
 	{
-		FScopeLock Lock(&DelayedComponentsRegistrationLock);
+		PCG::TScopeLock Lock(DelayedComponentsRegistrationLock);
 		DelayedComponentsToRegister.Remove(InComponent);
 		DelayedComponentsToUnregister.Remove(InComponent);
 	}
@@ -529,7 +529,7 @@ void FPCGActorAndComponentMapping::UnregisterPCGComponent(UPCGComponent* InCompo
 		// Ignore this if we force (aka when we actually unregister the delayed one)		
 		if (!bForce && InComponent->IsCreatedByConstructionScript())
 		{
-			FScopeLock Lock(&DelayedComponentsRegistrationLock);
+			PCG::TScopeLock Lock(DelayedComponentsRegistrationLock);
 			DelayedComponentsToUnregister.Add(InComponent);
 			return;
 		}
diff --git a/Engine/Plugins/PCG/Source/PCG/Private/PCGActorAndComponentMapping.h b/Engine/Plugins/PCG/Source/PCG/Private/PCGActorAndComponentMapping.h
--- a/Engine/Plugins/PCG/Source/PCG/Private/PCGActorAndComponentMapping.h
+++ b/Engine/Plugins/PCG/Source/PCG/Private/PCGActorAndComponentMapping.h
@@ -258,7 +258,7 @@ private:
 
 #if WITH_EDITOR
 	/** Components to be registered/unregistered on the next frame. see RegisterOrUpdatePCGComponent/UnregisterPCGComponent for a better understanding on why it is needed. */
-	mutable FCriticalSection DelayedComponentsRegistrationLock;
+	mutable PCG::FLock DelayedComponentsRegistrationLock;
 	TSet<UPCGComponent*> DelayedComponentsToUnregister;
 	TSet<UPCGComponent*> DelayedComponentsToRegister;
 	
diff --git a/Engine/Plugins/PCG/Source/PCG/Private/PCGGraph.cpp b/Engine/Plugins/PCG/Source/PCG/Private/PCGGraph.cpp
--- a/Engine/Plugins/PCG/Source/PCG/Private/PCGGraph.cpp
+++ b/Engine/Plugins/PCG/Source/PCG/Private/PCGGraph.cpp
@@ -1413,7 +1413,7 @@ void UPCGGraph::PostNodeUndo(UPCGNode* InPCGNode)
 
 void UPCGGraph::GetGridSizes(PCGHiGenGrid::FSizeArray& OutGridSizes, bool& bOutHasUnbounded) const
 {
-	FScopeLock Lock(&CachedGridInfoLock);
+	PCG::TScopeLock Lock(CachedGridInfoLock);
 
 	if (!CachedGridInfo)
 	{
@@ -1483,7 +1483,7 @@ void UPCGGraph::GetParentGridSizes(const uint32 InChildGridSize, PCGHiGenGrid::F
 		return;
 	}
 
-	FScopeLock Lock(&CachedGridInfoLock);
+	PCG::TScopeLock Lock(CachedGridInfoLock);
 	
 	if (const PCGHiGenGrid::FSizeArray* FoundGridSizeInfo = ChildGridSizeToParentGridSizes.Find(InChildGridSize))
 	{
@@ -1765,7 +1765,7 @@ void UPCGGraph::NotifyGraphChanged(EPCGChangeType ChangeType)
 		}
 
 		{
-			FScopeLock Lock(&CachedGridInfoLock);
+			PCG::TScopeLock Lock(CachedGridInfoLock);
 			CachedGridInfo.Reset();
 			ChildGridSizeToParentGridSizes.Reset();
 		}
diff --git a/Engine/Plugins/PCG/Source/PCG/Private/PCGModule.cpp b/Engine/Plugins/PCG/Source/PCG/Private/PCGModule.cpp
--- a/Engine/Plugins/PCG/Source/PCG/Private/PCGModule.cpp
+++ b/Engine/Plugins/PCG/Source/PCG/Private/PCGModule.cpp
@@ -146,7 +146,7 @@ void FPCGModule::OnPostInitEngine()
 
 void FPCGModule::ExecuteNextTick(TFunction<void()> TickFunction)
 {
-	FScopeLock Lock(&ExecuteNextTickLock);
+	PCG::TScopeLock Lock(ExecuteNextTickLock);
 	ExecuteNextTicks.Add(TickFunction);
 }
 
@@ -157,7 +157,7 @@ bool FPCGModule::Tick(float DeltaTime)
 	TArray<TFunction<void()>> LocalExecuteNextTicks;
 
 	{
-		FScopeLock Lock(&ExecuteNextTickLock);
+		PCG::TScopeLock Lock(ExecuteNextTickLock);
 		LocalExecuteNextTicks = MoveTemp(ExecuteNextTicks);
 	}
 
diff --git a/Engine/Plugins/PCG/Source/PCG/Private/PCGSettings.cpp b/Engine/Plugins/PCG/Source/PCG/Private/PCGSettings.cpp
--- a/Engine/Plugins/PCG/Source/PCG/Private/PCGSettings.cpp
+++ b/Engine/Plugins/PCG/Source/PCG/Private/PCGSettings.cpp
@@ -618,14 +618,12 @@ FPCGElementPtr UPCGSettings::GetElement() const
 {
 	if (!CachedElement)
 	{
-		CacheLock.Lock();
+		PCG::TScopeLock ScopeLock(CacheLock);
 
 		if (!CachedElement)
 		{
 			CachedElement = CreateElement();
 		}
-
-		CacheLock.Unlock();
 	}
 
 	return CachedElement;
diff --git a/Engine/Plugins/PCG/Source/PCG/Private/Subsystems/PCGSubsystem.cpp b/Engine/Plugins/PCG/Source/PCG/Private/Subsystems/PCGSubsystem.cpp
--- a/Engine/Plugins/PCG/Source/PCG/Private/Subsystems/PCGSubsystem.cpp
+++ b/Engine/Plugins/PCG/Source/PCG/Private/Subsystems/PCGSubsystem.cpp
@@ -368,7 +368,7 @@ APCGWorldActor* UPCGSubsystem::GetPCGWorldActor()
 #if WITH_EDITOR
 	if (!PCGWorldActor && !PCGHelpers::IsRuntimeOrPIE())
 	{
-		FScopeLock Lock(&PCGWorldActorLock);
+		PCG::TScopeLock Lock(PCGWorldActorLock);
 
 		if (!PCGWorldActor)
 		{
@@ -394,7 +394,7 @@ APCGWorldActor* UPCGSubsystem::GetPCGWorldActorForPartitionActor(APCGPartitionAc
 	{
 		UPCGActorHelpers::ForEachActorInLevel<APCGWorldActor>(InActor->GetLevel(), [this](AActor* InActor)
 		{
-			FScopeLock Lock(&PCGWorldActorLock);
+			PCG::TScopeLock Lock(PCGWorldActorLock);
 			if (!PCGWorldActor)
 			{
 				PCGWorldActor = Cast<APCGWorldActor>(InActor);
@@ -442,10 +442,9 @@ void UPCGSubsystem::DestroyCurrentPCGWorldActor()
 {
 	if (PCGWorldActor)
 	{
-		PCGWorldActorLock.Lock();
+		PCG::TScopeLock Lock(PCGWorldActorLock);
 		PCGWorldActor->Destroy();
 		PCGWorldActor = nullptr;
-		PCGWorldActorLock.Unlock();
 	}
 }
 
diff --git a/Engine/Plugins/PCG/Source/PCG/Private/Utils/PCGExtraCapture.cpp b/Engine/Plugins/PCG/Source/PCG/Private/Utils/PCGExtraCapture.cpp
--- a/Engine/Plugins/PCG/Source/PCG/Private/Utils/PCGExtraCapture.cpp
+++ b/Engine/Plugins/PCG/Source/PCG/Private/Utils/PCGExtraCapture.cpp
@@ -100,7 +100,7 @@ void PCGUtils::FExtraCapture::Update(const PCGUtils::FScopedCall& InScopedCall)
 		break;
 	}
 
-	FScopeLock ScopedLock(&Lock);
+	PCG::TScopeLock ScopedLock(Lock);
 	if (!InScopedCall.CapturedMessages.IsEmpty())
 	{
 		TArray<FCapturedMessage>& InstanceMessages = CapturedMessages.FindOrAdd(InScopedCall.Context->Node);
diff --git a/Engine/Plugins/PCG/Source/PCG/Public/Data/PCGBasePointData.h b/Engine/Plugins/PCG/Source/PCG/Public/Data/PCGBasePointData.h
--- a/Engine/Plugins/PCG/Source/PCG/Public/Data/PCGBasePointData.h
+++ b/Engine/Plugins/PCG/Source/PCG/Public/Data/PCGBasePointData.h
@@ -447,7 +447,7 @@ private:
 
 protected:
 
-	mutable FCriticalSection CachedDataLock;
+	mutable PCG::FRecursiveLock CachedDataLock;
 	mutable PCGPointOctree::FPointOctree PCGPointOctree;
 	mutable FBox Bounds;
 
diff --git a/Engine/Plugins/PCG/Source/PCG/Public/Data/PCGDynamicMeshData.h b/Engine/Plugins/PCG/Source/PCG/Public/Data/PCGDynamicMeshData.h
--- a/Engine/Plugins/PCG/Source/PCG/Public/Data/PCGDynamicMeshData.h
+++ b/Engine/Plugins/PCG/Source/PCG/Public/Data/PCGDynamicMeshData.h
@@ -97,11 +97,11 @@ protected:
 
 	mutable UE::Geometry::FDynamicMeshOctree3 DynamicMeshOctree;
 	mutable bool bDynamicMeshOctreeIsDirty = true;
-	mutable FCriticalSection DynamicMeshOctreeLock;
+	mutable PCG::FLock DynamicMeshOctreeLock;
 	
 	mutable FBox CachedBounds = FBox(EForceInit::ForceInit);
 	mutable bool bDynamicMeshBoundsAreDirty = true;
-	mutable UE::FSpinLock DynamicMeshBoundsLock;
+	mutable PCG::FLock DynamicMeshBoundsLock;
 };
 
 #undef UE_API
diff --git a/Engine/Plugins/PCG/Source/PCG/Public/Data/PCGSpatialData.h b/Engine/Plugins/PCG/Source/PCG/Public/Data/PCGSpatialData.h
--- a/Engine/Plugins/PCG/Source/PCG/Public/Data/PCGSpatialData.h
+++ b/Engine/Plugins/PCG/Source/PCG/Public/Data/PCGSpatialData.h
@@ -302,7 +302,7 @@ struct FPCGPointDataCache
 {
 	GENERATED_BODY()
 
-	const UPCGBasePointData* ToBasePointDataInternal(FPCGContext* Context, const FBox& InBounds, bool bSupportsBoundedPointData, FCriticalSection& CacheLock, TFunctionRef<const UPCGBasePointData* (FPCGContext*, const FBox&)> CreatePointDataFunc);
+	const UPCGBasePointData* ToBasePointDataInternal(FPCGContext* Context, const FBox& InBounds, bool bSupportsBoundedPointData, PCG::FLock& CacheLock, TFunctionRef<const UPCGBasePointData* (FPCGContext*, const FBox&)> CreatePointDataFunc);
 	void GetResourceSizeEx(FResourceSizeEx& CumulativeResourceSize);
 
 	UPROPERTY(Transient)
@@ -343,7 +343,7 @@ private:
 	UPROPERTY(Transient)
 	mutable FPCGPointDataCache PointArrayDataCache;
 
-	mutable FCriticalSection CacheLock;
+	mutable PCG::FLock CacheLock;
 };
 
 #undef UE_API
diff --git a/Engine/Plugins/PCG/Source/PCG/Public/Data/PCGWorldData.h b/Engine/Plugins/PCG/Source/PCG/Public/Data/PCGWorldData.h
--- a/Engine/Plugins/PCG/Source/PCG/Public/Data/PCGWorldData.h
+++ b/Engine/Plugins/PCG/Source/PCG/Public/Data/PCGWorldData.h
@@ -311,7 +311,7 @@ private:
 
 	// Transient data to cache filter actors
 	mutable TSet<TObjectKey<AActor>> CachedFilterActors;
-	mutable FCriticalSection CachedFilterActorsLock;
+	mutable PCG::FLock CachedFilterActorsLock;
 	mutable bool CachedFilterActorsDirty = false;
 };
 
diff --git a/Engine/Plugins/PCG/Source/PCG/Public/Grid/PCGLandscapeCache.h b/Engine/Plugins/PCG/Source/PCG/Public/Grid/PCGLandscapeCache.h
--- a/Engine/Plugins/PCG/Source/PCG/Public/Grid/PCGLandscapeCache.h
+++ b/Engine/Plugins/PCG/Source/PCG/Public/Grid/PCGLandscapeCache.h
@@ -2,6 +2,8 @@
 
 #pragma once
 
+#include "PCGCommon.h"
+
 #include "Engine/World.h"
 #include "Engine/Level.h"
 #include "GameFramework/Actor.h"
@@ -110,7 +112,7 @@ private:
 	mutable TArray<TArray<uint8>> LayerData;
 
 	// Transient data
-	mutable FCriticalSection DataLock;
+	mutable PCG::FLock DataLock;
 	mutable int32 Touch = 0;
 	mutable bool bDataLoaded = false;
 	TWeakObjectPtr<UPCGLandscapeCache> OwningCache = nullptr;
diff --git a/Engine/Plugins/PCG/Source/PCG/Public/PCGCommon.h b/Engine/Plugins/PCG/Source/PCG/Public/PCGCommon.h
--- a/Engine/Plugins/PCG/Source/PCG/Public/PCGCommon.h
+++ b/Engine/Plugins/PCG/Source/PCG/Public/PCGCommon.h
@@ -4,6 +4,8 @@
 
 #include "PCGPointPropertiesTraits.h"
 
+#include "Async/Mutex.h"
+#include "Async/RecursiveMutex.h"
 #include "HAL/IConsoleManager.h"
 
 #include "PCGCommon.generated.h"
@@ -23,6 +25,55 @@ struct FPCGContext;
 struct FPCGStack;
 struct FPCGStackContext;
 
+namespace PCG
+{
+	using FLock = UE::FMutex;
+	using FRecursiveLock = UE::FRecursiveMutex;
+
+	// Similar to existing UE::TScopeLock but allows optional TryLock in which case caller is responsible to 
+	// check if Lock succeeded or not with IsLocked(). This allows execution code to be skipped if non main thread
+	// tasks fail to get the lock instead of waiting.
+	template<typename MutexType>
+	class TScopeLock
+	{
+	public:
+		UE_NONCOPYABLE(TScopeLock);
+
+		UE_NODISCARD_CTOR TScopeLock(MutexType& InMutex, bool bTryLock = false)
+			: Mutex(&InMutex), bLocked(false)
+		{
+			check(Mutex);
+			if (bTryLock)
+			{
+				bLocked = Mutex->TryLock();
+			}
+			else
+			{
+				Mutex->Lock();
+				bLocked = true;
+			}
+		}
+
+		~TScopeLock()
+		{
+			if (bLocked && Mutex)
+			{
+				Mutex->Unlock();
+				Mutex = nullptr;
+			}
+		}
+
+		bool IsLocked() const
+		{
+			return bLocked;
+		}
+
+	private:
+		MutexType* Mutex;
+		bool bLocked;
+	};
+}
+
 namespace PCGLog
 {
 	/** Returns the component's owner name. In the case of commandlets, owner full object path will be returned. */
diff --git a/Engine/Plugins/PCG/Source/PCG/Public/PCGGraph.h b/Engine/Plugins/PCG/Source/PCG/Public/PCGGraph.h
--- a/Engine/Plugins/PCG/Source/PCG/Public/PCGGraph.h
+++ b/Engine/Plugins/PCG/Source/PCG/Public/PCGGraph.h
@@ -592,7 +592,7 @@ protected:
 	/** Cached information about higen grids present in graph. */
 	mutable TOptional<FGridInfo> CachedGridInfo;
 	mutable TMap<uint32, PCGHiGenGrid::FSizeArray> ChildGridSizeToParentGridSizes;
-	mutable FCriticalSection CachedGridInfoLock;
+	mutable PCG::FLock CachedGridInfoLock;
 
 	/** Sets whether this graph is marked as editor-only; note that the IsEditorOnly call depends on the local graph value and the value in all subgraphs, recursively. */
 	UPROPERTY(EditAnywhere, BlueprintReadWrite, Category = Cooking, meta=(EditCondition="!bIsStandaloneGraph", EditConditionHides))
diff --git a/Engine/Plugins/PCG/Source/PCG/Public/PCGModule.h b/Engine/Plugins/PCG/Source/PCG/Public/PCGModule.h
--- a/Engine/Plugins/PCG/Source/PCG/Public/PCGModule.h
+++ b/Engine/Plugins/PCG/Source/PCG/Public/PCGModule.h
@@ -85,7 +85,7 @@ private:
 
 	FTSTicker::FDelegateHandle	TickDelegateHandle;
 
-	FCriticalSection ExecuteNextTickLock;
+	PCG::FLock ExecuteNextTickLock;
 	TArray<TFunction<void()>> ExecuteNextTicks;
 
 #if WITH_EDITOR
diff --git a/Engine/Plugins/PCG/Source/PCG/Public/PCGSettings.h b/Engine/Plugins/PCG/Source/PCG/Public/PCGSettings.h
--- a/Engine/Plugins/PCG/Source/PCG/Public/PCGSettings.h
+++ b/Engine/Plugins/PCG/Source/PCG/Public/PCGSettings.h
@@ -510,7 +510,7 @@ public:
 
 private:
 	mutable FPCGElementPtr CachedElement;
-	mutable FCriticalSection CacheLock;
+	mutable PCG::FLock CacheLock;
 
 	// Overridable param section
 public:
diff --git a/Engine/Plugins/PCG/Source/PCG/Public/Subsystems/PCGSubsystem.h b/Engine/Plugins/PCG/Source/PCG/Public/Subsystems/PCGSubsystem.h
--- a/Engine/Plugins/PCG/Source/PCG/Public/Subsystems/PCGSubsystem.h
+++ b/Engine/Plugins/PCG/Source/PCG/Public/Subsystems/PCGSubsystem.h
@@ -410,7 +410,7 @@ private:
 	/** Functions will be executed at the beginning of the tick and then removed from this array. */
 	TArray<FTickAction> BeginTickActions;
 
-	FCriticalSection PCGWorldActorLock;
+	PCG::FLock PCGWorldActorLock;
 
 #if WITH_EDITOR
 	using FConstructionScriptSourceComponents = TMap<FName, TObjectKey<UPCGComponent>>;
diff --git a/Engine/Plugins/PCG/Source/PCG/Public/Utils/PCGExtraCapture.h b/Engine/Plugins/PCG/Source/PCG/Public/Utils/PCGExtraCapture.h
--- a/Engine/Plugins/PCG/Source/PCG/Public/Utils/PCGExtraCapture.h
+++ b/Engine/Plugins/PCG/Source/PCG/Public/Utils/PCGExtraCapture.h
@@ -104,7 +104,7 @@ namespace PCGUtils
 
 	private:
 		TCapturedMessageMap CapturedMessages;
-		mutable FCriticalSection Lock;
+		mutable PCG::FLock Lock;
 	};
 #else
 	struct FScopedCallOutputDevice
