[PCG] Limit GPU Execution based on editor permissions
- Add IPCGEditorModule event to update when permission changes
- UX: Hide HLSL Window
- UX: Hide bExecuteOnGPU flag
- Execution: prevent gpu execution (no matter the internal flags)

#rnx

diff --git a/Engine/Plugins/PCG/Source/PCG/Private/PCGSettings.cpp b/Engine/Plugins/PCG/Source/PCG/Private/PCGSettings.cpp
--- a/Engine/Plugins/PCG/Source/PCG/Private/PCGSettings.cpp
+++ b/Engine/Plugins/PCG/Source/PCG/Private/PCGSettings.cpp
@@ -28,6 +28,7 @@
 
 #if WITH_EDITOR
 #include "Editor.h"
+#include "Editor/IPCGEditorModule.h"
 #include "Hash/PCGSettingsHashContext.h"
 #endif
 
@@ -316,6 +317,11 @@ void UPCGSettings::PostLoad()
 		bDebug = ExecutionMode_DEPRECATED == EPCGSettingsExecutionMode::Debug;
 		ExecutionMode_DEPRECATED = EPCGSettingsExecutionMode::Enabled;
 	}
+
+	if (IPCGEditorModule* EditorModule = IPCGEditorModule::Get(); EditorModule && EditorModule->GetPermissionMode() != EPCGEditorPermissionMode::All)
+	{
+		bExecuteOnGPU = false;
+	}
 #endif
 
 	CacheCrc();
@@ -1011,8 +1017,31 @@ TArray<FPCGSettingsOverridableParam> UPCGSettings::GatherOverridableParams() con
 
 	return PCGSettingsHelpers::GetAllOverridableParams(GetClass(), Config);
 }
+
+bool UPCGSettings::DisplayExecuteOnGPUSettingCondition() const
+{
+	if (IPCGEditorModule* EditorModule = IPCGEditorModule::Get(); EditorModule && EditorModule->GetPermissionMode() != EPCGEditorPermissionMode::All)
+	{
+		return false;
+	}
+
+	return DisplayExecuteOnGPUSetting();
+}
 #endif // WITH_EDITOR
 
+bool UPCGSettings::ShouldExecuteOnGPU() const
+{
+#if WITH_EDITOR
+	// In case property was changed and then permissions were lowered
+	if (IPCGEditorModule* EditorModule = IPCGEditorModule::Get(); EditorModule && EditorModule->GetPermissionMode() != EPCGEditorPermissionMode::All)
+	{
+		return false;
+	}
+#endif
+
+	return bExecuteOnGPU;
+}
+
 void UPCGSettings::InitializeCachedOverridableParams(bool bReset)
 {
 	// Don't do it for default object
diff --git a/Engine/Plugins/PCG/Source/PCG/Public/Editor/IPCGEditorModule.h b/Engine/Plugins/PCG/Source/PCG/Public/Editor/IPCGEditorModule.h
--- a/Engine/Plugins/PCG/Source/PCG/Public/Editor/IPCGEditorModule.h
+++ b/Engine/Plugins/PCG/Source/PCG/Public/Editor/IPCGEditorModule.h
@@ -35,6 +35,9 @@ public:
 	virtual bool CanSelectPartitionActors() const = 0;
 	virtual void SetPermissionMode(EPCGEditorPermissionMode InPermissionMode) = 0;
 	virtual EPCGEditorPermissionMode GetPermissionMode() const = 0;
+
+	DECLARE_EVENT_OneParam(IPCGEditorModule, FPermissionModeChanged, EPCGEditorPermissionMode);
+	virtual FPermissionModeChanged& OnPermissionModeChanged() = 0;
 #endif
 
 protected:
diff --git a/Engine/Plugins/PCG/Source/PCG/Public/PCGSettings.h b/Engine/Plugins/PCG/Source/PCG/Public/PCGSettings.h
--- a/Engine/Plugins/PCG/Source/PCG/Public/PCGSettings.h
+++ b/Engine/Plugins/PCG/Source/PCG/Public/PCGSettings.h
@@ -573,14 +573,13 @@ public:
 public:
 
 	/** Whether this node should be executed on the GPU. */
-	virtual bool ShouldExecuteOnGPU() const { return bExecuteOnGPU; }
+	UE_API virtual bool ShouldExecuteOnGPU() const;
 
 #if WITH_EDITOR
 	/** Create the GPU kernels and kernel edges that should execute for this node. */
 	virtual void CreateKernels(FPCGGPUCompilationContext& InOutContext, UObject* InObjectOuter, TArray<UPCGComputeKernel*>& OutKernels, TArray<FPCGKernelEdge>& OutEdges) const {};
 
 	/** Whether to display GPU execution option in node settings UI. */
-	UFUNCTION()
 	virtual bool DisplayExecuteOnGPUSetting() const { return false; }
 #endif
 
@@ -590,12 +589,17 @@ protected:
 	virtual bool HasExecutionDependencyPin() const { return true; }
 
 private:
+#if WITH_EDITOR
+	UFUNCTION()
+	UE_API bool DisplayExecuteOnGPUSettingCondition() const;
+#endif
+
 	/** Adds the universal Execution Dependency pin to the node's pin properties. */
 	UE_API void AddExecutionDependencyPin(TArray<FPCGPinProperties>& OutPins) const;
 
 protected:
 	/** Whether this node should be executed on the GPU. */
-	UPROPERTY(EditAnywhere, Category = "GPU", meta = (EditCondition = "DisplayExecuteOnGPUSetting()", EditConditionHides, HideEditConditionToggle))
+	UPROPERTY(EditAnywhere, Category = "GPU", meta = (EditCondition = "DisplayExecuteOnGPUSettingCondition()", EditConditionHides, HideEditConditionToggle))
 	bool bExecuteOnGPU = false;
 
 public:
diff --git a/Engine/Plugins/PCG/Source/PCGEditor/Private/PCGEditor.cpp b/Engine/Plugins/PCG/Source/PCGEditor/Private/PCGEditor.cpp
--- a/Engine/Plugins/PCG/Source/PCGEditor/Private/PCGEditor.cpp
+++ b/Engine/Plugins/PCG/Source/PCGEditor/Private/PCGEditor.cpp
@@ -132,6 +132,11 @@ namespace FPCGEditor_private
 		ViewportID[2],
 		ViewportID[3]
 	};
+
+	const FName AllRestrictedPanels[] =
+	{
+		HLSLSourceID,
+	};
 }
 
 UPCGEditorGraph* FPCGEditor::GetPCGEditorGraph(UPCGGraph* InGraph)
@@ -157,6 +162,11 @@ void FPCGEditor::Initialize(const EToolkitMode::Type InMode, const TSharedPtr<cl
 
 	InspectionDataManager.OnInspectedStackChangedDelegate.AddSP(SharedThis(this), &FPCGEditor::UpdateAfterInspectedStackChanged);
 
+	if (IPCGEditorModule* PCGEditorModule = IPCGEditorModule::Get())
+	{
+		PCGEditorModule->OnPermissionModeChanged().AddSP(SharedThis(this), &FPCGEditor::OnPermissionChanged);
+	}
+
 	PCGGraphBeingEdited = InPCGGraph;
 
 	UpdateDefaultExecutionSource();
@@ -294,6 +304,12 @@ bool FPCGEditor::OnValidateNodeTitle(const FText& NewName, UEdGraphNode* GraphNo
 	return false;
 }
 
+void FPCGEditor::OnPermissionChanged(EPCGEditorPermissionMode InPermissionMode)
+{
+	// Menus and tabs can change based on permissions so we close the editor instead of dealing with refreshing the UX
+	CloseWindow(EAssetEditorCloseReason::EditorRefreshRequested);
+}
+
 void FPCGEditor::UpdateAfterInspectedStackChanged(const FPCGStack& FullStack)
 {
 	const bool bValidStack = !FullStack.GetStackFrames().IsEmpty();
@@ -572,7 +588,16 @@ bool FPCGEditor::IsPanelCurrentlyForeground(FName PanelID) const
 
 bool FPCGEditor::IsPanelAvailable(const FName PanelID) const
 {
-	return Algo::AnyOf(FPCGEditor_private::AllDefaultPanels, [PanelID](const FName& It) { return It == PanelID; });
+	bool bIsPanelAvailable = Algo::AnyOf(FPCGEditor_private::AllDefaultPanels, [PanelID](const FName& It) { return It == PanelID; });
+
+	IPCGEditorModule* PCGEditorModule = IPCGEditorModule::Get();
+	EPCGEditorPermissionMode PermissionMode = PCGEditorModule ? PCGEditorModule->GetPermissionMode() : EPCGEditorPermissionMode::None;
+	if (bIsPanelAvailable && PermissionMode != EPCGEditorPermissionMode::All)
+	{
+		bIsPanelAvailable = Algo::NoneOf(FPCGEditor_private::AllRestrictedPanels, [PanelID](const FName& It) { return It == PanelID; });
+	}
+
+	return bIsPanelAvailable;
 }
 
 void FPCGEditor::RegisterTabSpawners(const TSharedRef<FTabManager>& InTabManager)
diff --git a/Engine/Plugins/PCG/Source/PCGEditor/Private/PCGEditorModule.cpp b/Engine/Plugins/PCG/Source/PCGEditor/Private/PCGEditorModule.cpp
--- a/Engine/Plugins/PCG/Source/PCGEditor/Private/PCGEditorModule.cpp
+++ b/Engine/Plugins/PCG/Source/PCGEditor/Private/PCGEditorModule.cpp
@@ -298,6 +298,17 @@ bool FPCGEditorModule::CanSelectPartitionActors() const
 	return GetDefault<UPCGEditorSettings>()->bCanSelectPartitionActors;
 }
 
+void FPCGEditorModule::SetPermissionMode(EPCGEditorPermissionMode InPermissionMode)
+{
+	if (InPermissionMode == PermissionMode)
+	{
+		return;
+	}
+
+	PermissionMode = InPermissionMode;
+	PermissionModeChangedEvent.Broadcast(PermissionMode);
+}
+
 void FPCGEditorModule::RegisterOnEditorModeChange()
 {
 	// Have a callback that catches changes in the Editor modes, to catch when we exit the landscape edit mode.
diff --git a/Engine/Plugins/PCG/Source/PCGEditor/Public/PCGEditor.h b/Engine/Plugins/PCG/Source/PCGEditor/Public/PCGEditor.h
--- a/Engine/Plugins/PCG/Source/PCGEditor/Public/PCGEditor.h
+++ b/Engine/Plugins/PCG/Source/PCGEditor/Public/PCGEditor.h
@@ -15,6 +15,7 @@ class FPCGEditorInspectionDataManager;
 enum class ECheckBoxState : uint8;
 namespace ETextCommit { enum Type : int; }
 
+enum class EPCGEditorPermissionMode : uint8;
 class IPCGBaseSubsystem;
 class FUICommandList;
 class SGraphEditor;
@@ -281,6 +282,9 @@ private:
 	/** Whether selected nodes are inspected or not */
 	ECheckBoxState GetInspectedCheckState() const;
 
+	/** Called when editor module permission mode changes */
+	void OnPermissionChanged(EPCGEditorPermissionMode InPermissionMode);
+
 	void UpdateAfterInspectedStackChanged(const FPCGStack& FullStack);
 	
 	/** Toggle node enabled state for selected nodes */
diff --git a/Engine/Plugins/PCG/Source/PCGEditor/Public/PCGEditorModule.h b/Engine/Plugins/PCG/Source/PCGEditor/Public/PCGEditorModule.h
--- a/Engine/Plugins/PCG/Source/PCGEditor/Public/PCGEditorModule.h
+++ b/Engine/Plugins/PCG/Source/PCGEditor/Public/PCGEditorModule.h
@@ -37,8 +37,11 @@ public:
 	virtual const FPCGNodeVisualLogs& GetNodeVisualLogs() const override { return NodeVisualLogs; }
 	virtual FPCGNodeVisualLogs& GetNodeVisualLogsMutable() override { return NodeVisualLogs; }
 	virtual bool CanSelectPartitionActors() const override;
-	virtual void SetPermissionMode(EPCGEditorPermissionMode InPermissionMode) override { PermissionMode = InPermissionMode; }
+	virtual void SetPermissionMode(EPCGEditorPermissionMode InPermissionMode) override;
 	virtual EPCGEditorPermissionMode GetPermissionMode() const { return PermissionMode; }
+	
+	DECLARE_DERIVED_EVENT(FPCGEditorModule, IPCGEditorModule::FPermissionModeChanged, FPermissionModeChanged);
+	virtual FPermissionModeChanged& OnPermissionModeChanged() override { return PermissionModeChangedEvent; }
 
 protected:
 	virtual void OnScheduleGraph(const FPCGStackContext& StackContext) override;
@@ -96,4 +99,5 @@ protected:
 	mutable FRWLock ExecutedStacksLock;
 
 	EPCGEditorPermissionMode PermissionMode = EPCGEditorPermissionMode::All;
+	FPermissionModeChanged PermissionModeChangedEvent;
 };
