[PCG] Fix stale permission delegate registration

diff --git a/Engine/Plugins/PCG/Source/PCGEditor/Private/PCGEditor.cpp b/Engine/Plugins/PCG/Source/PCGEditor/Private/PCGEditor.cpp
--- a/Engine/Plugins/PCG/Source/PCGEditor/Private/PCGEditor.cpp
+++ b/Engine/Plugins/PCG/Source/PCGEditor/Private/PCGEditor.cpp
@@ -164,7 +164,7 @@ void FPCGEditor::Initialize(const EToolkitMode::Type InMode, const TSharedPtr<cl
 
 	if (IPCGEditorModule* PCGEditorModule = IPCGEditorModule::Get())
 	{
-		PCGEditorModule->OnPermissionModeChanged().AddSP(SharedThis(this), &FPCGEditor::OnPermissionChanged);
+		PCGEditorModule->OnPermissionModeChanged().AddRaw(this, &FPCGEditor::OnPermissionChanged);
 	}
 
 	PCGGraphBeingEdited = InPCGGraph;
@@ -2957,6 +2957,11 @@ void FPCGEditor::OnClose()
 		UnregisterDelegatesForWorld(GEditor->GetEditorWorldContext().World());
 		UnregisterDelegatesForWorld(GEditor->PlayWorld.Get());
 	}
+
+	if (IPCGEditorModule* PCGEditorModule = IPCGEditorModule::Get())
+	{
+		PCGEditorModule->OnPermissionModeChanged().RemoveAll(this);
+	}
 }
 
 void FPCGEditor::InitToolMenuContext(FToolMenuContext& MenuContext)
