[PCG] Replace FRWLock by FSharedMutex to reduce native sync primitive usage
- Declare TUniqueScopeLock/TSharedScopeLock to be used with PCG::FSharedLock

diff --git a/Engine/Plugins/PCG/Source/PCG/Private/Graph/PCGGraphCompiler.cpp b/Engine/Plugins/PCG/Source/PCG/Private/Graph/PCGGraphCompiler.cpp
--- a/Engine/Plugins/PCG/Source/PCG/Private/Graph/PCGGraphCompiler.cpp
+++ b/Engine/Plugins/PCG/Source/PCG/Private/Graph/PCGGraphCompiler.cpp
@@ -79,7 +79,7 @@ UPCGComputeGraph* FPCGGraphCompilerCache::GetCompiledComputeGraph(const UPCGGrap
 {
 #if WITH_EDITOR
 	{
-		FReadScopeLock Lock(GraphToTaskMapLock);
+		PCG::TSharedScopeLock ReadLock(GraphToTaskMapLock);
 
 		if (TMap<uint32, TArray<TObjectPtr<UPCGComputeGraph>>>* GridSizeToComputeGraphs = TopGraphToComputeGraphMap.Find(InGraph))
 		{
@@ -122,7 +122,7 @@ void FPCGGraphCompilerCache::RemoveFromCache(UPCGGraph* InGraph)
 void FPCGGraphCompilerCache::RemoveFromCacheRecursive(UPCGGraph* InGraph)
 {
 	{
-		FWriteScopeLock Lock(GraphToTaskMapLock);
+		PCG::TUniqueScopeLock WriteLock(GraphToTaskMapLock);
 		GraphToTaskMap.Remove(InGraph);
 		GraphToStackContextMap.Remove(InGraph);
 		TopGraphToTaskMap.Remove(InGraph);
@@ -404,7 +404,7 @@ void FPCGGraphCompiler::Compile(UPCGGraph* InGraph)
 	bool bAlreadyCached = false;
 
 	{
-		FReadScopeLock Lock(Cache.GraphToTaskMapLock);
+		PCG::TSharedScopeLock ReadLock(Cache.GraphToTaskMapLock);
 		bAlreadyCached = Cache.GraphToTaskMap.Contains(InGraph);
 	}
 
@@ -434,7 +434,7 @@ void FPCGGraphCompiler::Compile(UPCGGraph* InGraph)
 	// Store back the results in the cache if it's valid
 	if (!CompiledTasks.IsEmpty())
 	{
-		FWriteScopeLock Lock(Cache.GraphToTaskMapLock);
+		PCG::TUniqueScopeLock WriteLock(Cache.GraphToTaskMapLock);
 		if (!Cache.GraphToTaskMap.Contains(InGraph))
 		{
 			Cache.GraphToTaskMap.Add(InGraph, MoveTemp(CompiledTasks));
@@ -446,7 +446,7 @@ void FPCGGraphCompiler::Compile(UPCGGraph* InGraph)
 TArray<FPCGGraphTask> FPCGGraphCompiler::GetPrecompiledTasks(const UPCGGraph* InGraph, uint32 GenerationGridSize, FPCGStackContext& OutStackContext, bool bIsTopGraph) const
 {
 	// Get compiled tasks in a threadsafe way
-	FReadScopeLock ReadLock(Cache.GraphToTaskMapLock);
+	PCG::TSharedScopeLock ReadLock(Cache.GraphToTaskMapLock);
 
 	const TArray<FPCGGraphTask>* ExistingTasks = nullptr;
 	if (bIsTopGraph)
@@ -980,7 +980,7 @@ void FPCGGraphCompiler::AddReferencedObjects(FReferenceCollector& Collector)
 {
 	TRACE_CPUPROFILER_EVENT_SCOPE(FCPGGraphCompiler::AddReferencedObjects);
 
-	FReadScopeLock Lock(Cache.GraphToTaskMapLock);
+	PCG::TSharedScopeLock ReadLock(Cache.GraphToTaskMapLock);
 
 	for (TPair<UPCGGraph*, TMap<uint32, TArray<TObjectPtr<UPCGComputeGraph>>>>& GraphToGridSizes : Cache.TopGraphToComputeGraphMap)
 	{
@@ -1264,7 +1264,7 @@ TArray<FPCGGraphTask> FPCGGraphCompiler::GetCompiledTasks(UPCGGraph* InGraph, ui
 		CompileTopGraph(InGraph, GenerationGridSize);
 
 		// Get compiled tasks in a threadsafe way
-		FReadScopeLock Lock(Cache.GraphToTaskMapLock);
+		PCG::TSharedScopeLock ReadLock(Cache.GraphToTaskMapLock);
 		const TMap<uint32, TArray<FPCGGraphTask>>* GridSizeToCompiledGraph = Cache.TopGraphToTaskMap.Find(InGraph);
 		const TArray<FPCGGraphTask>* Tasks = GridSizeToCompiledGraph ? GridSizeToCompiledGraph->Find(GenerationGridSize) : nullptr;
 
@@ -1286,7 +1286,7 @@ TArray<FPCGGraphTask> FPCGGraphCompiler::GetCompiledTasks(UPCGGraph* InGraph, ui
 		Compile(InGraph);
 
 		// Get compiled tasks in a threadsafe way
-		FReadScopeLock Lock(Cache.GraphToTaskMapLock);
+		PCG::TSharedScopeLock ReadLock(Cache.GraphToTaskMapLock);
 		if (TArray<FPCGGraphTask>* Tasks = Cache.GraphToTaskMap.Find(InGraph))
 		{
 			CompiledTasks = *Tasks;
@@ -1326,7 +1326,7 @@ void FPCGGraphCompiler::CompileTopGraph(UPCGGraph* InGraph, uint32 GenerationGri
 	bool bAlreadyCached = false;
 
 	{
-		FReadScopeLock Lock(Cache.GraphToTaskMapLock);
+		PCG::TSharedScopeLock ReadLock(Cache.GraphToTaskMapLock);
 		const TMap<uint32, TArray<FPCGGraphTask>>* GridSizeToCompiledGraph = Cache.TopGraphToTaskMap.Find(InGraph);
 		bAlreadyCached = GridSizeToCompiledGraph && GridSizeToCompiledGraph->Contains(GenerationGridSize);
 	}
@@ -1510,7 +1510,7 @@ void FPCGGraphCompiler::CompileTopGraph(UPCGGraph* InGraph, uint32 GenerationGri
 	}
 
 	// Store back the results in the cache
-	Cache.GraphToTaskMapLock.WriteLock();
+	PCG::TUniqueScopeLock WriteLock(Cache.GraphToTaskMapLock);
 	TMap<uint32, TArray<FPCGGraphTask>>& TasksPerGenerationGrid = Cache.TopGraphToTaskMap.FindOrAdd(InGraph);
 	if (!TasksPerGenerationGrid.Contains(GenerationGridSize))
 	{
@@ -1522,13 +1522,12 @@ void FPCGGraphCompiler::CompileTopGraph(UPCGGraph* InGraph, uint32 GenerationGri
 	{
 		StackContextPerGenerationGrid.Add(GenerationGridSize, MoveTemp(StackContext));
 	}
-	Cache.GraphToTaskMapLock.WriteUnlock();
 }
 
 FPCGElementPtr FPCGGraphCompiler::GetSharedTrivialElement()
 {
 	{
-		FReadScopeLock Lock(SharedTrivialElementLock);
+		PCG::TSharedScopeLock ReadLock(SharedTrivialElementLock);
 
 		if (SharedTrivialElement)
 		{
@@ -1536,7 +1535,7 @@ FPCGElementPtr FPCGGraphCompiler::GetSharedTrivialElement()
 		}
 	}
 
-	FWriteScopeLock Lock(SharedTrivialElementLock);
+	PCG::TUniqueScopeLock WriteLock(SharedTrivialElementLock);
 
 	if (!SharedTrivialElement)
 	{
@@ -1549,7 +1548,7 @@ FPCGElementPtr FPCGGraphCompiler::GetSharedTrivialElement()
 FPCGElementPtr FPCGGraphCompiler::GetSharedGatherElement()
 {
 	{
-		FReadScopeLock Lock(SharedGatherElementLock);
+		PCG::TSharedScopeLock ReadLock(SharedGatherElementLock);
 
 		if (SharedGatherElement)
 		{
@@ -1557,7 +1556,7 @@ FPCGElementPtr FPCGGraphCompiler::GetSharedGatherElement()
 		}
 	}
 
-	FWriteScopeLock Lock(SharedGatherElementLock);
+	PCG::TUniqueScopeLock WriteLock(SharedGatherElementLock);
 
 	if (!SharedGatherElement)
 	{
@@ -1570,16 +1569,16 @@ FPCGElementPtr FPCGGraphCompiler::GetSharedGatherElement()
 FPCGElementPtr FPCGGraphCompiler::GetSharedTrivialPostGraphElement()
 {
 	{
-		FReadScopeLock Lock(SharedTrivialPostGraphElementLock);
-
+		PCG::TSharedScopeLock ReadLock(SharedTrivialPostGraphElementLock);
+		
 		if (SharedTrivialPostGraphElement)
 		{
 			return SharedTrivialPostGraphElement;
 		}
 	}
 
-	FWriteScopeLock Lock(SharedTrivialPostGraphElementLock);
-
+	PCG::TUniqueScopeLock WriteLock(SharedTrivialPostGraphElementLock);
+	
 	if (!SharedTrivialPostGraphElement)
 	{
 		SharedTrivialPostGraphElement = MakeShared<FPCGTrivialElement>();
@@ -1590,7 +1589,7 @@ FPCGElementPtr FPCGGraphCompiler::GetSharedTrivialPostGraphElement()
 
 void FPCGGraphCompiler::ClearCache()
 {
-	FWriteScopeLock Lock(Cache.GraphToTaskMapLock);
+	PCG::TUniqueScopeLock WriteLock(Cache.GraphToTaskMapLock);
 	Cache.GraphToTaskMap.Reset();
 	Cache.GraphToStackContextMap.Reset();
 	Cache.TopGraphToTaskMap.Reset();
diff --git a/Engine/Plugins/PCG/Source/PCG/Private/Graph/PCGGraphCompiler.h b/Engine/Plugins/PCG/Source/PCG/Private/Graph/PCGGraphCompiler.h
--- a/Engine/Plugins/PCG/Source/PCG/Private/Graph/PCGGraphCompiler.h
+++ b/Engine/Plugins/PCG/Source/PCG/Private/Graph/PCGGraphCompiler.h
@@ -31,7 +31,7 @@ private:
 	void RemoveFromCacheRecursive(UPCGGraph* InGraph);
 #endif // WITH_EDITOR
 
-	mutable FRWLock GraphToTaskMapLock;
+	mutable PCG::FSharedLock GraphToTaskMapLock;
 	TMap<UPCGGraph*, TArray<FPCGGraphTask>> GraphToTaskMap;
 	TMap<UPCGGraph*, FPCGStackContext> GraphToStackContextMap;
 
@@ -148,11 +148,11 @@ private:
 	bool bIsCooking = false;
 
 	FPCGElementPtr SharedTrivialElement;
-	FRWLock SharedTrivialElementLock;
+	PCG::FSharedLock SharedTrivialElementLock;
 
 	FPCGElementPtr SharedTrivialPostGraphElement;
-	FRWLock SharedTrivialPostGraphElementLock;
+	PCG::FSharedLock SharedTrivialPostGraphElementLock;
 
 	FPCGElementPtr SharedGatherElement;
-	FRWLock SharedGatherElementLock;
+	PCG::FSharedLock SharedGatherElementLock;
 };
diff --git a/Engine/Plugins/PCG/Source/PCG/Private/Graph/PCGGraphCompilerGPU.cpp b/Engine/Plugins/PCG/Source/PCG/Private/Graph/PCGGraphCompilerGPU.cpp
--- a/Engine/Plugins/PCG/Source/PCG/Private/Graph/PCGGraphCompilerGPU.cpp
+++ b/Engine/Plugins/PCG/Source/PCG/Private/Graph/PCGGraphCompilerGPU.cpp
@@ -1824,7 +1824,7 @@ void FPCGGraphCompilerGPU::CreateGPUNodes(FPCGGraphCompiler& InOutCompiler, UPCG
 	}
 
 	{
-		FWriteScopeLock Lock(InOutCompiler.GetCache().GraphToTaskMapLock);
+		PCG::TUniqueScopeLock WriteLock(InOutCompiler.GetCache().GraphToTaskMapLock);
 		TMap<uint32, TArray<TObjectPtr<UPCGComputeGraph>>>& GridSizeToComputeGraphs = InOutCompiler.GetCache().TopGraphToComputeGraphMap.FindOrAdd(InGraph);
 		TArray<TObjectPtr<UPCGComputeGraph>>& ComputeGraphs = GridSizeToComputeGraphs.FindOrAdd(InGridSize);
 
diff --git a/Engine/Plugins/PCG/Source/PCG/Private/Grid/PCGLandscapeCache.cpp b/Engine/Plugins/PCG/Source/PCG/Private/Grid/PCGLandscapeCache.cpp
--- a/Engine/Plugins/PCG/Source/PCG/Private/Grid/PCGLandscapeCache.cpp
+++ b/Engine/Plugins/PCG/Source/PCG/Private/Grid/PCGLandscapeCache.cpp
@@ -899,7 +899,7 @@ TSharedPtr<const FPCGLandscapeCacheEntry> UPCGLandscapeCache::GetOrCreateCacheEn
 		return CacheEntry;
 	}
 		
-	FWriteScopeLock ScopeLock(CacheLock);
+	PCG::TUniqueScopeLock WriteLock(CacheLock);
 	CacheMapKey ComponentKey(LandscapeGuid, InComponentKey, LandscapeComponent->GetOwner());
 	if (const TSharedPtr<FPCGLandscapeCacheEntry>* FoundEntry = CachedData.Find(ComponentKey))
 	{
@@ -940,7 +940,7 @@ TSharedPtr<const FPCGLandscapeCacheEntry> UPCGLandscapeCache::GetCacheEntryInter
 
 	{
 #if WITH_EDITOR
-		FReadScopeLock ScopeLock(CacheLock);
+		PCG::TSharedScopeLock ReadLock(CacheLock);
 #else
 		if (!bLoggedNoCacheError && CachedData.IsEmpty())
 		{
@@ -972,7 +972,7 @@ TSharedPtr<const FPCGLandscapeCacheEntry> UPCGLandscapeCache::GetCacheEntryInter
 TArray<FName> UPCGLandscapeCache::GetLayerNames(ALandscapeProxy* Landscape)
 {
 #if WITH_EDITOR
-	FReadScopeLock ScopeLock(CacheLock);
+	PCG::TSharedScopeLock ReadLock(CacheLock);
 #endif
 	return CachedLayerNames.Array();
 }
@@ -1079,7 +1079,7 @@ void UPCGLandscapeCache::OnLandscapeChanged(ALandscapeProxy* InLandscape, const
 		return;
 	}
 
-	CacheLock.WriteLock();
+	PCG::TUniqueScopeLock WriteLock(CacheLock);
 
 	// Just remove these from the cache, they'll be added back on demand
 	InChangeParams.ForEachComponent([this, InLandscape](const ULandscapeComponent* LandscapeComponent)
@@ -1096,8 +1096,6 @@ void UPCGLandscapeCache::OnLandscapeChanged(ALandscapeProxy* InLandscape, const
 	});
 
 	CacheLayerNames(InLandscape);
-
-	CacheLock.WriteUnlock();
 }
 
 void UPCGLandscapeCache::OnLandscapeMoved(AActor* InActor)
@@ -1108,14 +1106,12 @@ void UPCGLandscapeCache::OnLandscapeMoved(AActor* InActor)
 		return;
 	}
 
-	CacheLock.WriteLock();
+	PCG::TUniqueScopeLock WriteLock(CacheLock);
 
 	// Just remove these from the cache, they'll be added back on demand
 	RemoveComponentFromCache_Unsafe(Landscape);
 
 	CacheLayerNames(Landscape);
-
-	CacheLock.WriteUnlock();
 }
 
 void UPCGLandscapeCache::OnLandscapeDeleted(AActor* Actor)
@@ -1126,13 +1122,11 @@ void UPCGLandscapeCache::OnLandscapeDeleted(AActor* Actor)
 		return;
 	}
 
-	CacheLock.WriteLock();
+	PCG::TUniqueScopeLock WriteLock(CacheLock);
 
 	RemoveComponentFromCache_Unsafe(LandscapeProxy);
 
 	Landscapes.Remove(LandscapeProxy);
-
-	CacheLock.WriteUnlock();
 }
 
 void UPCGLandscapeCache::OnLandscapeLoaded(AActor& Actor)
diff --git a/Engine/Plugins/PCG/Source/PCG/Private/Metadata/PCGMetadataAttribute.cpp b/Engine/Plugins/PCG/Source/PCG/Private/Metadata/PCGMetadataAttribute.cpp
--- a/Engine/Plugins/PCG/Source/PCG/Private/Metadata/PCGMetadataAttribute.cpp
+++ b/Engine/Plugins/PCG/Source/PCG/Private/Metadata/PCGMetadataAttribute.cpp
@@ -84,7 +84,7 @@ void FPCGMetadataAttributeBase::Serialize(FPCGMetadataDomain* InMetadata, FArchi
 
 void FPCGMetadataAttributeBase::SetValueFromValueKey(PCGMetadataEntryKey EntryKey, PCGMetadataValueKey ValueKey, bool bResetValueOnDefaultValueKey)
 {
-	FWriteScopeLock ScopeLock(EntryMapLock);
+	PCG::TUniqueScopeLock WriteLock(EntryMapLock);
 	SetValueFromValueKey_Unsafe(EntryKey, ValueKey, bResetValueOnDefaultValueKey);
 }
 
@@ -113,7 +113,7 @@ void FPCGMetadataAttributeBase::SetValuesFromValueKeys(const TArrayView<const TT
 		return;
 	}
 
-	FWriteScopeLock ScopeLock(EntryMapLock);
+	PCG::TUniqueScopeLock WriteLock(EntryMapLock);
 	for (const TTuple<PCGMetadataEntryKey, PCGMetadataValueKey>& EntryValuePair : EntryValuePairs)
 	{
 		SetValueFromValueKey_Unsafe(EntryValuePair.Key, EntryValuePair.Value, bResetValueOnDefaultValueKey, /*bAllowInvalidEntries=*/true);
@@ -127,7 +127,7 @@ void FPCGMetadataAttributeBase::SetValuesFromValueKeys(const TArrayView<const PC
 		return;
 	}
 
-	FWriteScopeLock ScopeLock(EntryMapLock);
+	PCG::TUniqueScopeLock WriteLock(EntryMapLock);
 	for (int32 i = 0; i < EntryKeys.Num(); ++i)
 	{
 		SetValueFromValueKey_Unsafe(EntryKeys[i], ValueKeys[i], bResetValueOnDefaultValueKey, /*bAllowInvalidEntries=*/true);
@@ -141,7 +141,7 @@ void FPCGMetadataAttributeBase::SetValuesFromValueKeys(const TArrayView<const PC
 		return;
 	}
 
-	FWriteScopeLock ScopeLock(EntryMapLock);
+	PCG::TUniqueScopeLock WriteLock(EntryMapLock);
 	for (int32 i = 0; i < EntryKeys.Num(); ++i)
 	{
 		SetValueFromValueKey_Unsafe(*EntryKeys[i], ValueKeys[i], bResetValueOnDefaultValueKey, /*bAllowInvalidEntries=*/true);
@@ -158,13 +158,14 @@ PCGMetadataValueKey FPCGMetadataAttributeBase::GetValueKey(PCGMetadataEntryKey E
 	PCGMetadataValueKey ValueKey = PCGDefaultValueKey;
 	bool bFoundKey = false;
 
-	EntryMapLock.ReadLock();
-	if (const PCGMetadataValueKey* FoundLocalKey = EntryToValueKeyMap.Find(EntryKey))
 	{
-		ValueKey = *FoundLocalKey;
-		bFoundKey = true;
+		PCG::TSharedScopeLock ReadLock(EntryMapLock);
+		if (const PCGMetadataValueKey* FoundLocalKey = EntryToValueKeyMap.Find(EntryKey))
+		{
+			ValueKey = *FoundLocalKey;
+			bFoundKey = true;
+		}
 	}
-	EntryMapLock.ReadUnlock();
 
 	if (!bFoundKey && Parent)
 	{
@@ -225,39 +226,39 @@ void FPCGMetadataAttributeBase::GetValueKeys_Internal(TConstPCGValueRange<PCGMet
 		return;
 	}
 
-	EntryMapLock.ReadLock();
-
-	for (; It; ++It)
 	{
-		const int32 Index = It.GetIndex();
-		const PCGMetadataEntryKey EntryKey = EntryKeys[Index];
+		PCG::TSharedScopeLock ReadLock(EntryMapLock);
 
-		auto SetValueKey = [Index, &OutValueKeys, &UnsetValues](PCGMetadataValueKey ValueKey)
-		{
-			OutValueKeys[Index] = ValueKey;
-			UnsetValues[Index] = false;
-		};
-		
-		if (EntryKey == PCGInvalidEntryKey)
-		{
-			SetValueKey(PCGDefaultValueKey);
-		}
-		else if (const PCGMetadataValueKey* FoundLocalKey = EntryToValueKeyMap.Find(EntryKey))
-		{
-			SetValueKey(*FoundLocalKey);
-		}
-		else if (!Parent)
-		{
-			SetValueKey(PCGDefaultValueKey);
-		}
-		else
+		for (; It; ++It)
 		{
-			bFoundAllKeys = false;
+			const int32 Index = It.GetIndex();
+			const PCGMetadataEntryKey EntryKey = EntryKeys[Index];
+
+			auto SetValueKey = [Index, &OutValueKeys, &UnsetValues](PCGMetadataValueKey ValueKey)
+				{
+					OutValueKeys[Index] = ValueKey;
+					UnsetValues[Index] = false;
+				};
+
+			if (EntryKey == PCGInvalidEntryKey)
+			{
+				SetValueKey(PCGDefaultValueKey);
+			}
+			else if (const PCGMetadataValueKey* FoundLocalKey = EntryToValueKeyMap.Find(EntryKey))
+			{
+				SetValueKey(*FoundLocalKey);
+			}
+			else if (!Parent)
+			{
+				SetValueKey(PCGDefaultValueKey);
+			}
+			else
+			{
+				bFoundAllKeys = false;
+			}
 		}
 	}
 
-	EntryMapLock.ReadUnlock();
-
 	ensure(Parent || bFoundAllKeys);
 
 	if (Parent && !bFoundAllKeys)
diff --git a/Engine/Plugins/PCG/Source/PCG/Private/Metadata/PCGMetadataDomain.cpp b/Engine/Plugins/PCG/Source/PCG/Private/Metadata/PCGMetadataDomain.cpp
--- a/Engine/Plugins/PCG/Source/PCG/Private/Metadata/PCGMetadataDomain.cpp
+++ b/Engine/Plugins/PCG/Source/PCG/Private/Metadata/PCGMetadataDomain.cpp
@@ -208,7 +208,7 @@ void FPCGMetadataDomain::AddToCrc(FArchiveCrc32& Ar, const UPCGData* Data, bool
 	TArray<const FPCGMetadataAttributeBase*> AllAttributes;
 
 	{
-		FReadScopeLock Lock(AttributeLock);
+		PCG::TSharedScopeLock ReadLock(AttributeLock);
 		AllAttributes.Reserve(Attributes.Num());
 
 		for (const TPair<FName, FPCGMetadataAttributeBase*>& Attribute : Attributes)
@@ -288,7 +288,7 @@ FPCGMetadataDomain::FPCGMetadataDomain(UPCGMetadata* InTopMetadata, FPCGMetadata
 
 FPCGMetadataDomain::~FPCGMetadataDomain()
 {
-	FWriteScopeLock ScopeLock(AttributeLock);
+	PCG::TUniqueScopeLock WriteLock(AttributeLock);
 	for (TPair<FName, FPCGMetadataAttributeBase*>& AttributeEntry : Attributes)
 	{
 		delete AttributeEntry.Value;
@@ -550,31 +550,32 @@ void FPCGMetadataDomain::FlattenImpl()
 
 	const int32 NumEntries = GetItemCountForChild();
 
-	AttributeLock.WriteLock();
-	for (auto& AttributePair : Attributes)
 	{
-		FPCGMetadataAttributeBase* Attribute = AttributePair.Value;
-		check(Attribute);
-
-		// For all stored entries (from the root), we need to make sure that entries that should have a concrete value have it
-		// Optimization notes:
-		// - we could skip entries that existed prior to attribute existence, etc.
-		// - we could skip entries that have no parent, but that would require checking against the parent entries in the parent hierarchy
-		for (int64 EntryKey = 0; EntryKey < NumEntries; ++EntryKey)
+		PCG::TUniqueScopeLock WriteLock(AttributeLock);
+		for (auto& AttributePair : Attributes)
 		{
-			// Get value using value inheritance as expected
-			PCGMetadataValueKey ValueKey = Attribute->GetValueKey(EntryKey);
-			if (ValueKey != PCGDefaultValueKey)
+			FPCGMetadataAttributeBase* Attribute = AttributePair.Value;
+			check(Attribute);
+
+			// For all stored entries (from the root), we need to make sure that entries that should have a concrete value have it
+			// Optimization notes:
+			// - we could skip entries that existed prior to attribute existence, etc.
+			// - we could skip entries that have no parent, but that would require checking against the parent entries in the parent hierarchy
+			for (int64 EntryKey = 0; EntryKey < NumEntries; ++EntryKey)
 			{
-				// Set concrete non-default value
-				Attribute->SetValueFromValueKey(EntryKey, ValueKey);
+				// Get value using value inheritance as expected
+				PCGMetadataValueKey ValueKey = Attribute->GetValueKey(EntryKey);
+				if (ValueKey != PCGDefaultValueKey)
+				{
+					// Set concrete non-default value
+					Attribute->SetValueFromValueKey(EntryKey, ValueKey);
+				}
 			}
-		}
 
-		// Finally, flatten values
-		Attribute->Flatten();
+			// Finally, flatten values
+			Attribute->Flatten();
+		}
 	}
-	AttributeLock.WriteUnlock();
 
 	Parent = nullptr;
 	ParentKeys.Reset();
@@ -600,15 +601,16 @@ bool FPCGMetadataDomain::FlattenAndCompress(const TArrayView<const PCGMetadataEn
 		UE_LOG(LogPCG, Warning, TEXT("Tried to flatten and compress a domain '%ls', which doesn't support multiple entries, with %d entries to keep. Will only keep the first one."), *DomainID.DebugName.ToString(), InEntryKeysToKeep.Num());
 	}
 
-	AttributeLock.WriteLock();
-	for (auto& AttributePair : Attributes)
 	{
-		FPCGMetadataAttributeBase* Attribute = AttributePair.Value;
-		check(Attribute);
+		PCG::TUniqueScopeLock WriteLock(AttributeLock);
+		for (auto& AttributePair : Attributes)
+		{
+			FPCGMetadataAttributeBase* Attribute = AttributePair.Value;
+			check(Attribute);
 
-		Attribute->FlattenAndCompress(EntryKeysToKeep);
+			Attribute->FlattenAndCompress(EntryKeysToKeep);
+		}
 	}
-	AttributeLock.WriteUnlock();
 
 	Parent = nullptr;
 	ParentKeys.Reset();
@@ -638,15 +640,16 @@ FPCGMetadataAttributeBase* FPCGMetadataDomain::GetMutableAttribute(FName Attribu
 {
 	FPCGMetadataAttributeBase* Attribute = nullptr;
 
-	AttributeLock.ReadLock();
-	if (FPCGMetadataAttributeBase** FoundAttribute = Attributes.Find(AttributeName))
 	{
-		Attribute = *FoundAttribute;
+		PCG::TSharedScopeLock ReadLock(AttributeLock);
+		if (FPCGMetadataAttributeBase** FoundAttribute = Attributes.Find(AttributeName))
+		{
+			Attribute = *FoundAttribute;
 
-		// Also when accessing an attribute, notify the PCG Data owner that the latest attribute manipulated is this one.
-		SetLastCachedSelectorOnOwner(AttributeName);
+			// Also when accessing an attribute, notify the PCG Data owner that the latest attribute manipulated is this one.
+			SetLastCachedSelectorOnOwner(AttributeName);
+		}
 	}
-	AttributeLock.ReadUnlock();
 
 	return Attribute;
 }
@@ -655,12 +658,13 @@ const FPCGMetadataAttributeBase* FPCGMetadataDomain::GetConstAttribute(FName Att
 {
 	const FPCGMetadataAttributeBase* Attribute = nullptr;
 
-	AttributeLock.ReadLock();
-	if (const FPCGMetadataAttributeBase* const* FoundAttribute = Attributes.Find(AttributeName))
 	{
-		Attribute = *FoundAttribute;
+		PCG::TSharedScopeLock ReadLock(AttributeLock);
+		if (const FPCGMetadataAttributeBase* const* FoundAttribute = Attributes.Find(AttributeName))
+		{
+			Attribute = *FoundAttribute;
+		}
 	}
-	AttributeLock.ReadUnlock();
 
 	return Attribute;
 }
@@ -669,23 +673,24 @@ const FPCGMetadataAttributeBase* FPCGMetadataDomain::GetConstAttributeById(int32
 {
 	const FPCGMetadataAttributeBase* Attribute = nullptr;
 
-	AttributeLock.ReadLock();
-	for (const auto& AttributePair : Attributes)
 	{
-		if (AttributePair.Value && AttributePair.Value->AttributeId == InAttributeId)
+		PCG::TSharedScopeLock ReadLock(AttributeLock);
+		for (const auto& AttributePair : Attributes)
 		{
-			Attribute = AttributePair.Value;
-			break;
+			if (AttributePair.Value && AttributePair.Value->AttributeId == InAttributeId)
+			{
+				Attribute = AttributePair.Value;
+				break;
+			}
 		}
 	}
-	AttributeLock.ReadUnlock();
 
 	return Attribute;
 }
 
 bool FPCGMetadataDomain::HasAttribute(FName AttributeName) const
 {
-	FReadScopeLock ScopeLock(AttributeLock);
+	PCG::TSharedScopeLock ReadLock(AttributeLock);
 	return Attributes.Contains(AttributeName);
 }
 
@@ -698,29 +703,30 @@ bool FPCGMetadataDomain::HasCommonAttributes(const FPCGMetadataDomain* InMetadat
 
 	bool bHasCommonAttribute = false;
 
-	AttributeLock.ReadLock();
-	for (const TPair<FName, FPCGMetadataAttributeBase*>& AttributePair : Attributes)
 	{
-		if (InMetadata->HasAttribute(AttributePair.Key))
+		PCG::TSharedScopeLock ReadLock(AttributeLock);
+		for (const TPair<FName, FPCGMetadataAttributeBase*>& AttributePair : Attributes)
 		{
-			bHasCommonAttribute = true;
-			break;
+			if (InMetadata->HasAttribute(AttributePair.Key))
+			{
+				bHasCommonAttribute = true;
+				break;
+			}
 		}
 	}
-	AttributeLock.ReadUnlock();
 
 	return bHasCommonAttribute;
 }
 
 int32 FPCGMetadataDomain::GetAttributeCount() const
 {
-	FReadScopeLock ScopeLock(AttributeLock);
+	PCG::TSharedScopeLock ReadLock(AttributeLock);
 	return Attributes.Num();
 }
 
 void FPCGMetadataDomain::GetAttributes(TArray<FName>& AttributeNames, TArray<EPCGMetadataTypes>& AttributeTypes) const
 {
-	AttributeLock.ReadLock();
+	PCG::TSharedScopeLock ReadLock(AttributeLock);
 	AttributeNames.Reserve(AttributeNames.Num() + Attributes.Num());
 	AttributeTypes.Reserve(AttributeTypes.Num() + Attributes.Num());
 	for (const TPair<FName, FPCGMetadataAttributeBase*>& Attribute : Attributes)
@@ -737,7 +743,6 @@ void FPCGMetadataDomain::GetAttributes(TArray<FName>& AttributeNames, TArray<EPC
 			AttributeTypes.Add(EPCGMetadataTypes::Unknown);
 		}
 	}
-	AttributeLock.ReadUnlock();
 }
 
 FName FPCGMetadataDomain::GetLatestAttributeNameOrNone() const
@@ -745,16 +750,17 @@ FName FPCGMetadataDomain::GetLatestAttributeNameOrNone() const
 	FName LatestAttributeName = NAME_None;
 	int64 MaxAttributeId = -1;
 	
-	AttributeLock.ReadLock();
-	for (const TPair<FName, FPCGMetadataAttributeBase*>& It : Attributes)
 	{
-		if (It.Value && (It.Value->AttributeId > MaxAttributeId))
+		PCG::TSharedScopeLock ReadLock(AttributeLock);
+		for (const TPair<FName, FPCGMetadataAttributeBase*>& It : Attributes)
 		{
-			MaxAttributeId = It.Value->AttributeId;
-			LatestAttributeName = It.Key;
+			if (It.Value && (It.Value->AttributeId > MaxAttributeId))
+			{
+				MaxAttributeId = It.Value->AttributeId;
+				LatestAttributeName = It.Key;
+			}
 		}
 	}
-	AttributeLock.ReadUnlock();
 
 	return LatestAttributeName;
 }
@@ -793,12 +799,13 @@ FPCGMetadataAttributeBase* FPCGMetadataDomain::CopyAttribute(FName AttributeToCo
 {
 	const FPCGMetadataAttributeBase* OriginalAttribute = nullptr;
 
-	AttributeLock.ReadLock();
-	if (FPCGMetadataAttributeBase** AttributeFound = Attributes.Find(AttributeToCopy))
 	{
-		OriginalAttribute = *AttributeFound;
+		PCG::TSharedScopeLock ReadLock(AttributeLock);
+		if (FPCGMetadataAttributeBase** AttributeFound = Attributes.Find(AttributeToCopy))
+		{
+			OriginalAttribute = *AttributeFound;
+		}
 	}
-	AttributeLock.ReadUnlock();
 
 	if (!OriginalAttribute && Parent)
 	{
@@ -822,14 +829,12 @@ FPCGMetadataAttributeBase* FPCGMetadataDomain::CopyAttribute(const FPCGMetadataA
 
 	if (NewAttribute)
 	{
-		AttributeLock.WriteLock();
+		PCG::TUniqueScopeLock WriteLock(AttributeLock);
 		NewAttribute->AttributeId = NextAttributeId++;
 		AddAttributeInternal(NewAttributeName, NewAttribute);
 
 		// Also when creating an attribute, notify the PCG Data owner that the latest attribute manipulated is this one.
 		SetLastCachedSelectorOnOwner(NewAttributeName);
-
-		AttributeLock.WriteUnlock();
 	}
 
 	return NewAttribute;
@@ -844,20 +849,21 @@ bool FPCGMetadataDomain::RenameAttribute(FName AttributeToRename, FName NewAttri
 	}
 
 	bool bRenamed = false;
-	AttributeLock.WriteLock();
-	if (FPCGMetadataAttributeBase** AttributeFound = Attributes.Find(AttributeToRename))
 	{
-		FPCGMetadataAttributeBase* Attribute = *AttributeFound;
-		RemoveAttributeInternal(AttributeToRename);
-		Attribute->Name = NewAttributeName;
-		AddAttributeInternal(NewAttributeName, Attribute);
+		PCG::TUniqueScopeLock WriteLock(AttributeLock);
+		if (FPCGMetadataAttributeBase** AttributeFound = Attributes.Find(AttributeToRename))
+		{
+			FPCGMetadataAttributeBase* Attribute = *AttributeFound;
+			RemoveAttributeInternal(AttributeToRename);
+			Attribute->Name = NewAttributeName;
+			AddAttributeInternal(NewAttributeName, Attribute);
 
-		// Also when renaming an attribute, notify the PCG Data owner that the latest attribute manipulated is this one.
-		SetLastCachedSelectorOnOwner(NewAttributeName);
-		
-		bRenamed = true;
+			// Also when renaming an attribute, notify the PCG Data owner that the latest attribute manipulated is this one.
+			SetLastCachedSelectorOnOwner(NewAttributeName);
+
+			bRenamed = true;
+		}
 	}
-	AttributeLock.WriteUnlock();
 
 	if (!bRenamed)
 	{
@@ -871,12 +877,13 @@ void FPCGMetadataDomain::ClearAttribute(FName AttributeToClear)
 {
 	FPCGMetadataAttributeBase* Attribute = nullptr;
 
-	AttributeLock.ReadLock();
-	if (FPCGMetadataAttributeBase** AttributeFound = Attributes.Find(AttributeToClear))
 	{
-		Attribute = *AttributeFound;
+		PCG::TSharedScopeLock ReadLock(AttributeLock);
+		if (FPCGMetadataAttributeBase** AttributeFound = Attributes.Find(AttributeToClear))
+		{
+			Attribute = *AttributeFound;
+		}
 	}
-	AttributeLock.ReadUnlock();
 
 	// If the attribute exists, then we can lose all the entries
 	// If it doesn't but it exists in the parent hierarchy, then we must create a new attribute.
@@ -891,13 +898,14 @@ void FPCGMetadataDomain::DeleteAttribute(FName AttributeToDelete)
 	FPCGMetadataAttributeBase* Attribute = nullptr;
 
 	// If it's a local attribute, then just delete it
-	AttributeLock.WriteLock();
-	if (FPCGMetadataAttributeBase** AttributeFound = Attributes.Find(AttributeToDelete))
 	{
-		Attribute = *AttributeFound;
-		RemoveAttributeInternal(AttributeToDelete);
+		PCG::TUniqueScopeLock WriteLock(AttributeLock);
+		if (FPCGMetadataAttributeBase** AttributeFound = Attributes.Find(AttributeToDelete))
+		{
+			Attribute = *AttributeFound;
+			RemoveAttributeInternal(AttributeToDelete);
+		}
 	}
-	AttributeLock.WriteUnlock();
 
 	if (Attribute)
 	{
@@ -929,10 +937,11 @@ bool FPCGMetadataDomain::ChangeAttributeType(FName AttributeName, int16 Attribut
 	{
 		NewAttribute->AttributeId = Attribute->AttributeId;
 
-		AttributeLock.WriteLock();
-		RemoveAttributeInternal(AttributeName);
-		AddAttributeInternal(AttributeName, NewAttribute);
-		AttributeLock.WriteUnlock();
+		{
+			PCG::TUniqueScopeLock WriteLock(AttributeLock);
+			RemoveAttributeInternal(AttributeName);
+			AddAttributeInternal(AttributeName, NewAttribute);
+		}
 
 		delete Attribute;
 		Attribute = nullptr;
@@ -943,19 +952,19 @@ bool FPCGMetadataDomain::ChangeAttributeType(FName AttributeName, int16 Attribut
 
 int64 FPCGMetadataDomain::GetItemCountForChild() const
 {
-	FReadScopeLock ScopeLock(ItemLock);
+	PCG::TSharedScopeLock ReadLock(ItemLock);
 	return ParentKeys.Num() + ItemKeyOffset;
 }
 
 int64 FPCGMetadataDomain::GetLocalItemCount() const
 {
-	FReadScopeLock ScopeLock(ItemLock);
+	PCG::TSharedScopeLock ReadLock(ItemLock);
 	return ParentKeys.Num();
 }
 
 int64 FPCGMetadataDomain::AddEntry(int64 ParentEntry)
 {
-	FWriteScopeLock ScopeLock(ItemLock);
+	PCG::TUniqueScopeLock WriteLock(ItemLock);
 	if (bSupportMultiEntries)
 	{
 		return ParentKeys.Add(ParentEntry) + ItemKeyOffset;
@@ -985,7 +994,7 @@ TArray<int64> FPCGMetadataDomain::AddEntries(TArrayView<const int64> ParentEntry
 	TArray<int64> Result;
 	Result.Reserve(ParentEntryKeys.Num());
 
-	FWriteScopeLock ScopeLock(ItemLock);
+	PCG::TUniqueScopeLock WriteLock(ItemLock);
 	if (bSupportMultiEntries)
 	{
 		ParentKeys.Reserve(ParentKeys.Num() + ParentEntryKeys.Num());
@@ -1020,7 +1029,7 @@ void FPCGMetadataDomain::AddEntriesInPlace(TArrayView<int64*> ParentEntryKeys)
 		return;
 	}
 	
-	FWriteScopeLock ScopeLock(ItemLock);
+	PCG::TUniqueScopeLock WriteLock(ItemLock);
 	if (bSupportMultiEntries)
 	{
 		ParentKeys.Reserve(ParentKeys.Num() + ParentEntryKeys.Num());
@@ -1051,7 +1060,7 @@ void FPCGMetadataDomain::AddEntriesInPlace(TArrayView<int64*> ParentEntryKeys)
 
 int64 FPCGMetadataDomain::AddEntryPlaceholder()
 {
-	FReadScopeLock ScopeLock(ItemLock);
+	PCG::TSharedScopeLock ReadLock(ItemLock);
 	if (bSupportMultiEntries)
 	{
 		return ParentKeys.Num() + DelayedEntriesIndex.IncrementExchange() + ItemKeyOffset;
@@ -1075,7 +1084,7 @@ void FPCGMetadataDomain::AddDelayedEntries(const TArray<TTuple<int64, int64>>& A
 		return;
 	}
 	
-	FWriteScopeLock ScopeLock(ItemLock);
+	PCG::TUniqueScopeLock WriteLock(ItemLock);
 		
 	if (bSupportMultiEntries)
 	{
@@ -1143,7 +1152,7 @@ PCGMetadataEntryKey FPCGMetadataDomain::GetParentKey(PCGMetadataEntryKey LocalIt
 	}
 	else
 	{
-		FReadScopeLock ScopeLock(ItemLock);
+		PCG::TSharedScopeLock ReadLock(ItemLock);
 		if (LocalItemKey - ItemKeyOffset < ParentKeys.Num())
 		{
 			return ParentKeys[LocalItemKey - ItemKeyOffset];
@@ -1181,7 +1190,7 @@ void FPCGMetadataDomain::GetParentKeysWithRange(TPCGValueRange<PCGMetadataEntryK
 		}
 	};
 
-	FReadScopeLock ScopeLock(ItemLock);
+	PCG::TSharedScopeLock ReadLock(ItemLock);
 	if (Mask && ensure(LocalItemKeys.Num() == Mask->Num()))
 	{
 		for (TConstSetBitIterator<> It(*Mask); It; ++It)
@@ -1215,7 +1224,7 @@ void FPCGMetadataDomain::MergeAttributesSubset(PCGMetadataEntryKey InKeyA, const
 	// For each attribute in the current metadata, query the values from point A & B, apply operation on the result and finally store in the out point.
 	InitializeOnSet(OutKey, InKeyA, InMetadataA, InKeyB, InMetadataB);
 
-	AttributeLock.ReadLock();
+	PCG::TSharedScopeLock ReadLock(AttributeLock);
 	for(const TPair<FName, FPCGMetadataAttributeBase*>& AttributePair : Attributes)
 	{
 		const FName& AttributeName = AttributePair.Key;
@@ -1252,14 +1261,13 @@ void FPCGMetadataDomain::MergeAttributesSubset(PCGMetadataEntryKey InKeyA, const
 			Attribute->SetValue(OutKey, AttributeA, InKeyA, AttributeB, InKeyB, Op);
 		}
 	}
-	AttributeLock.ReadUnlock();
 }
 
 void FPCGMetadataDomain::ResetWeightedAttributes(PCGMetadataEntryKey& OutKey)
 {
 	InitializeOnSet(OutKey);
 
-	AttributeLock.ReadLock();
+	PCG::TSharedScopeLock ReadLock(AttributeLock);
 	for(const TPair<FName, FPCGMetadataAttributeBase*>& AttributePair : Attributes)
 	{
 		FPCGMetadataAttributeBase* Attribute = AttributePair.Value;
@@ -1269,7 +1277,6 @@ void FPCGMetadataDomain::ResetWeightedAttributes(PCGMetadataEntryKey& OutKey)
 			Attribute->SetZeroValue(OutKey);
 		}
 	}
-	AttributeLock.ReadUnlock();
 }
 
 void FPCGMetadataDomain::AccumulateWeightedAttributes(PCGMetadataEntryKey InKey, const FPCGMetadataDomain* InMetadata, float Weight, bool bSetNonInterpolableAttributes, PCGMetadataEntryKey& OutKey)
@@ -1283,7 +1290,7 @@ void FPCGMetadataDomain::AccumulateWeightedAttributes(PCGMetadataEntryKey InKey,
 
 	const bool bShouldSetNonInterpolableAttributes = bSetNonInterpolableAttributes && !bHasSetParent;
 
-	AttributeLock.ReadLock();
+	PCG::TSharedScopeLock ReadLock(AttributeLock);
 	for(const TPair<FName, FPCGMetadataAttributeBase*>& AttributePair : Attributes)
 	{
 		const FName& AttributeName = AttributePair.Key;
@@ -1307,7 +1314,6 @@ void FPCGMetadataDomain::AccumulateWeightedAttributes(PCGMetadataEntryKey InKey,
 			}
 		}
 	}
-	AttributeLock.ReadUnlock();
 }
 
 void FPCGMetadataDomain::ComputeWeightedAttribute(PCGMetadataEntryKey& OutKey, const TArrayView<TPair<PCGMetadataEntryKey, float>>& InWeightedKeys, const FPCGMetadataDomain* InMetadata)
@@ -1318,7 +1324,7 @@ void FPCGMetadataDomain::ComputeWeightedAttribute(PCGMetadataEntryKey& OutKey, c
 	}
 
 	// Could ensure that InitializeOnSet returns false...
-	AttributeLock.ReadLock();
+	PCG::TSharedScopeLock ReadLock(AttributeLock);
 	for (const TPair<FName, FPCGMetadataAttributeBase*>& AttributePair : Attributes)
 	{
 		const FName& AttributeName = AttributePair.Key;
@@ -1340,7 +1346,6 @@ void FPCGMetadataDomain::ComputeWeightedAttribute(PCGMetadataEntryKey& OutKey, c
 			Attribute->SetWeightedValue(OutKey, OtherAttribute, InWeightedKeys);
 		}
 	}
-	AttributeLock.ReadUnlock();
 }
 
 int64 FPCGMetadataDomain::GetItemKeyCountForParent() const
@@ -1361,7 +1366,7 @@ void FPCGMetadataDomain::SetAttributes(PCGMetadataEntryKey InKey, const FPCGMeta
 		return;
 	}
 
-	AttributeLock.ReadLock();
+	PCG::TSharedScopeLock ReadLock(AttributeLock);
 	for(const TPair<FName, FPCGMetadataAttributeBase*>& AttributePair : Attributes)
 	{
 		const FName& AttributeName = AttributePair.Key;
@@ -1378,7 +1383,6 @@ void FPCGMetadataDomain::SetAttributes(PCGMetadataEntryKey InKey, const FPCGMeta
 			Attribute->SetValue(OutKey, OtherAttribute, InKey);
 		}
 	}
-	AttributeLock.ReadUnlock();
 }
 
 void FPCGMetadataDomain::SetAttributes(const TArrayView<const PCGMetadataEntryKey>& InOriginalKeys, const FPCGMetadataDomain* InMetadata, const TArrayView<PCGMetadataEntryKey>* InOutOptionalKeys, FPCGContext* OptionalContext)
@@ -1436,77 +1440,78 @@ void FPCGMetadataDomain::SetAttributes(const TArrayView<const PCGMetadataEntryKe
 		}
 	}
 
-	AttributeLock.ReadLock();
-	int32 AttributeOffset = 0;
-	const int32 DefaultAttributesPerDispatch = 64;
-	int32 AttributesPerDispatch = OptionalContext ? DefaultAttributesPerDispatch : 1;
-	if (OptionalContext && OptionalContext->AsyncState.NumAvailableTasks > 0)
 	{
-		AttributesPerDispatch = FMath::Min(OptionalContext->AsyncState.NumAvailableTasks, AttributesPerDispatch);
-	}
-
-	while (AttributeOffset < Attributes.Num())
-	{
-		TArray<FName> AttributeNames;
-		TArray<EPCGMetadataTypes> AttributeTypes;
-		GetAttributes(AttributeNames, AttributeTypes);
-
-		const int32 AttributeCountInCurrentDispatch = FMath::Min(AttributesPerDispatch, Attributes.Num() - AttributeOffset);
-		ParallelFor(AttributeCountInCurrentDispatch, [this, OptionalContext, AttributeOffset, InMetadata, &AttributeNames, &InKeys, &OutKeys](int32 WorkerIndex)
+		PCG::TSharedScopeLock ReadLock(AttributeLock);
+		int32 AttributeOffset = 0;
+		const int32 DefaultAttributesPerDispatch = 64;
+		int32 AttributesPerDispatch = OptionalContext ? DefaultAttributesPerDispatch : 1;
+		if (OptionalContext && OptionalContext->AsyncState.NumAvailableTasks > 0)
 		{
-			LLM_SCOPE_BYTAG(PCG);
+			AttributesPerDispatch = FMath::Min(OptionalContext->AsyncState.NumAvailableTasks, AttributesPerDispatch);
+		}
 
-			const FName AttributeName = AttributeNames[AttributeOffset + WorkerIndex];
-			FPCGMetadataAttributeBase* Attribute = Attributes[AttributeName];
+		while (AttributeOffset < Attributes.Num())
+		{
+			TArray<FName> AttributeNames;
+			TArray<EPCGMetadataTypes> AttributeTypes;
+			GetAttributes(AttributeNames, AttributeTypes);
 
-			if (const FPCGMetadataAttributeBase* OtherAttribute = InMetadata->GetConstAttribute(AttributeName))
+			const int32 AttributeCountInCurrentDispatch = FMath::Min(AttributesPerDispatch, Attributes.Num() - AttributeOffset);
+			ParallelFor(AttributeCountInCurrentDispatch, [this, OptionalContext, AttributeOffset, InMetadata, &AttributeNames, &InKeys, &OutKeys](int32 WorkerIndex)
 			{
-				if (!PCG::Private::IsBroadcastableOrConstructible(OtherAttribute->GetTypeId(), Attribute->GetTypeId()))
-				{
-					PCGE_LOG_C(Error, GraphAndLog, OptionalContext, FText::Format(NSLOCTEXT("PCGMetadata", "TypeMismatch", "Metadata type mismatch with attribute '{0}'"), FText::FromName(AttributeName)));
-					return;
-				}
+				LLM_SCOPE_BYTAG(PCG);
 
-				if (Attribute == OtherAttribute)
-				{
-					TArray<PCGMetadataValueKey> ValueKeys;
-					Attribute->GetValueKeys(TArrayView<const PCGMetadataEntryKey>(InKeys), ValueKeys);
-					Attribute->SetValuesFromValueKeys(OutKeys, ValueKeys);
-				}
-				else
-				{
-					// Create accessor for the other attribute
-					TUniquePtr<const IPCGAttributeAccessor> OtherAttributeAccessor = PCGAttributeAccessorHelpers::CreateConstAccessor(OtherAttribute, InMetadata);
+				const FName AttributeName = AttributeNames[AttributeOffset + WorkerIndex];
+				FPCGMetadataAttributeBase* Attribute = Attributes[AttributeName];
 
-					TArrayView<PCGMetadataEntryKey> InKeysView(InKeys);
-					FPCGAttributeAccessorKeysEntries OtherAttributeKeys(InKeysView);
-
-					if (!OtherAttributeAccessor)
+				if (const FPCGMetadataAttributeBase* OtherAttribute = InMetadata->GetConstAttribute(AttributeName))
+				{
+					if (!PCG::Private::IsBroadcastableOrConstructible(OtherAttribute->GetTypeId(), Attribute->GetTypeId()))
 					{
+						PCGE_LOG_C(Error, GraphAndLog, OptionalContext, FText::Format(NSLOCTEXT("PCGMetadata", "TypeMismatch", "Metadata type mismatch with attribute '{0}'"), FText::FromName(AttributeName)));
 						return;
 					}
 
-					auto GetAndSetValues = [Attribute, &OutKeys, &OtherAttributeAccessor, &OtherAttributeKeys](auto Dummy) -> bool
+					if (Attribute == OtherAttribute)
 					{
-						using Type = decltype(Dummy);
+						TArray<PCGMetadataValueKey> ValueKeys;
+						Attribute->GetValueKeys(TArrayView<const PCGMetadataEntryKey>(InKeys), ValueKeys);
+						Attribute->SetValuesFromValueKeys(OutKeys, ValueKeys);
+					}
+					else
+					{
+						// Create accessor for the other attribute
+						TUniquePtr<const IPCGAttributeAccessor> OtherAttributeAccessor = PCGAttributeAccessorHelpers::CreateConstAccessor(OtherAttribute, InMetadata);
+
+						TArrayView<PCGMetadataEntryKey> InKeysView(InKeys);
+						FPCGAttributeAccessorKeysEntries OtherAttributeKeys(InKeysView);
 
-						auto SetValues = [Attribute, &OutKeys](const TArrayView<Type>& View, const int32 Start, const int32 Range)
+						if (!OtherAttributeAccessor)
 						{
-							TArrayView<PCGMetadataEntryKey> Keys(OutKeys.GetData() + Start, Range);
-							static_cast<FPCGMetadataAttribute<Type>*>(Attribute)->SetValues(Keys, View);
-						};
+							return;
+						}
+
+						auto GetAndSetValues = [Attribute, &OutKeys, &OtherAttributeAccessor, &OtherAttributeKeys](auto Dummy) -> bool
+							{
+								using Type = decltype(Dummy);
+
+								auto SetValues = [Attribute, &OutKeys](const TArrayView<Type>& View, const int32 Start, const int32 Range)
+									{
+										TArrayView<PCGMetadataEntryKey> Keys(OutKeys.GetData() + Start, Range);
+										static_cast<FPCGMetadataAttribute<Type>*>(Attribute)->SetValues(Keys, View);
+									};
 
-						return PCGMetadataElementCommon::ApplyOnAccessorRange<Type>(OtherAttributeKeys, *OtherAttributeAccessor, SetValues, EPCGAttributeAccessorFlags::AllowBroadcast | EPCGAttributeAccessorFlags::AllowConstructible);
-					};
+								return PCGMetadataElementCommon::ApplyOnAccessorRange<Type>(OtherAttributeKeys, *OtherAttributeAccessor, SetValues, EPCGAttributeAccessorFlags::AllowBroadcast | EPCGAttributeAccessorFlags::AllowConstructible);
+							};
 
-					PCGMetadataAttribute::CallbackWithRightType(Attribute->GetTypeId(), GetAndSetValues);
+						PCGMetadataAttribute::CallbackWithRightType(Attribute->GetTypeId(), GetAndSetValues);
+					}
 				}
-			}
-		});
+			});
 
-		AttributeOffset += AttributeCountInCurrentDispatch;
+			AttributeOffset += AttributeCountInCurrentDispatch;
+		}
 	}
-	AttributeLock.ReadUnlock();
 
 	if (InOutOptionalKeys)
 	{
diff --git a/Engine/Plugins/PCG/Source/PCG/Private/PCGActorAndComponentMapping.cpp b/Engine/Plugins/PCG/Source/PCG/Private/PCGActorAndComponentMapping.cpp
--- a/Engine/Plugins/PCG/Source/PCG/Private/PCGActorAndComponentMapping.cpp
+++ b/Engine/Plugins/PCG/Source/PCG/Private/PCGActorAndComponentMapping.cpp
@@ -245,7 +245,7 @@ TArray<FPCGTaskId> FPCGActorAndComponentMapping::DispatchToRegisteredLocalCompon
 	const bool bIsRuntimeGenerated = OriginalComponent->IsManagedByRuntimeGenSystem();
 
 	// TODO: Might be more interesting to copy the set and release the lock.
-	FReadScopeLock ReadLock(ComponentToPartitionActorsMapLock);
+	PCG::TSharedScopeLock ReadLock(ComponentToPartitionActorsMapLock);
 	const TSet<TObjectPtr<APCGPartitionActor>>* PartitionActorsPtr = ComponentToPartitionActorsMap.Find(OriginalComponent);
 
 	if (!PartitionActorsPtr)
@@ -473,9 +473,9 @@ bool FPCGActorAndComponentMapping::RemapPCGComponent(const UPCGComponent* OldCom
 #endif
 
 	// Remap all previous instances
-	auto RemapPreviousInstances = [OldComponent, NewComponent](TMap<const UPCGComponent*, TSet<TObjectPtr<APCGPartitionActor>>>& Map, FRWLock& Lock)
+	auto RemapPreviousInstances = [OldComponent, NewComponent](TMap<const UPCGComponent*, TSet<TObjectPtr<APCGPartitionActor>>>& Map, PCG::FSharedLock& Lock)
 	{
-		FWriteScopeLock WriteLock(Lock);
+		PCG::TUniqueScopeLock WriteLock(Lock);
 
 		if (TSet<TObjectPtr<APCGPartitionActor>>* PartitionActorsPtr = Map.Find(OldComponent))
 		{
@@ -555,7 +555,7 @@ void FPCGActorAndComponentMapping::UnregisterPartitionedPCGComponent(UPCGCompone
 	// inside a lock. So copy the actors to clean up and release the lock before doing RemoveGraphInstance.
 	TSet<TObjectPtr<APCGPartitionActor>> PartitionActorsToCleanUp;
 	{
-		FWriteScopeLock WriteLock(ComponentToPartitionActorsMapLock);
+		PCG::TUniqueScopeLock WriteLock(ComponentToPartitionActorsMapLock);
 		TSet<TObjectPtr<APCGPartitionActor>>* PartitionActorsPtr = ComponentToPartitionActorsMap.Find(InComponent);
 
 		if (PartitionActorsPtr)
@@ -650,7 +650,7 @@ void FPCGActorAndComponentMapping::RegisterPartitionActor(APCGPartitionActor* In
 	check(PartitionActorDescriptor.GetGridSize() > 0);
 
 	{
-		FWriteScopeLock WriteLock(PartitionActorsMapLock);
+		PCG::TUniqueScopeLock WriteLock(PartitionActorsMapLock);
 
 		TMap<FIntVector, TObjectPtr<APCGPartitionActor>>& PartitionActorsMapGrid = PartitionActorsMap.FindOrAdd(PartitionActorDescriptor);
 		if (PartitionActorsMapGrid.Contains(GridCoord))
@@ -679,7 +679,7 @@ PRAGMA_ENABLE_DEPRECATION_WARNINGS
 	// Register to all the components that intersect with the PA. Ignore for runtime generated, it is handled manually
 	if (!PartitionActorDescriptor.IsRuntime())
 	{
-		FWriteScopeLock WriteLock(ComponentToPartitionActorsMapLock);
+		PCG::TUniqueScopeLock WriteLock(ComponentToPartitionActorsMapLock);
 		ForAllIntersectingPartitionedComponents(FBoxCenterAndExtent(InActor->GetFixedBounds()), [this, InActor, WorldActor, PartitionActorDescriptor](UPCGComponent* Component)
 		{
 			// For each component, do the mapping if the component is generated or if the Partition Actor already has a local component
@@ -740,14 +740,14 @@ void FPCGActorAndComponentMapping::UnregisterPartitionActor(APCGPartitionActor*
 
 	if (TMap<FIntVector, TObjectPtr<APCGPartitionActor>>* PartitionActorsMapGrid = PartitionActorsMap.Find(GridDescriptor))
 	{
-		FWriteScopeLock WriteLock(PartitionActorsMapLock);
+		PCG::TUniqueScopeLock WriteLock(PartitionActorsMapLock);
 		PartitionActorsMapGrid->Remove(GridCoord);
 	}
 
 	// Unregister from all intersecting components. Ignore for runtime generated, it is handled manually
 	if (!GridDescriptor.IsRuntime())
 	{
-		FWriteScopeLock WriteLock(ComponentToPartitionActorsMapLock);
+		PCG::TUniqueScopeLock WriteLock(ComponentToPartitionActorsMapLock);
 		ForAllIntersectingPartitionedComponents(FBoxCenterAndExtent(Actor->GetFixedBounds()), [this, Actor](UPCGComponent* Component)
 		{
 			TSet<TObjectPtr<APCGPartitionActor>>* PartitionActorsPtr = ComponentToPartitionActorsMap.Find(Component);
@@ -770,7 +770,7 @@ void FPCGActorAndComponentMapping::ForAllIntersectingPartitionActors(const FBox&
 
 	auto ForAllIntersectingPartitionActorsOfGridSize = [InFunc, &InBounds, this, bInBoundsAreEmpty](const FPCGGridDescriptor& GridDescriptor)
 	{
-		FReadScopeLock ReadLock(PartitionActorsMapLock);
+		PCG::TSharedScopeLock ReadLock(PartitionActorsMapLock);
 
 		const TMap<FIntVector, TObjectPtr<APCGPartitionActor>>* PartitionActorsMapGrid = PartitionActorsMap.Find(GridDescriptor);
 		if (!PartitionActorsMapGrid || PartitionActorsMapGrid->IsEmpty())
@@ -813,7 +813,7 @@ void FPCGActorAndComponentMapping::ForAllIntersectingPartitionActors(const FBox&
 
 	TArray<FPCGGridDescriptor> GridDescriptors;
 	{
-		FReadScopeLock ReadLock(PartitionActorsMapLock);
+		PCG::TSharedScopeLock ReadLock(PartitionActorsMapLock);
 
 		PartitionActorsMap.GenerateKeyArray(GridDescriptors);
 	}
@@ -866,7 +866,7 @@ void FPCGActorAndComponentMapping::UpdateMappingPCGComponentPartitionActor(UPCGC
 			TSet<TObjectPtr<APCGPartitionActor>> RemovedActors;
 
 			{
-				FWriteScopeLock WriteLock(ComponentToPartitionActorsMapLock);
+				PCG::TUniqueScopeLock WriteLock(ComponentToPartitionActorsMapLock);
 				TSet<TObjectPtr<APCGPartitionActor>>& PartitionActors = ComponentToPartitionActorsMap.FindOrAdd(InComponent);
 
 				TSet<TObjectPtr<APCGPartitionActor>> NewMapping;
@@ -938,9 +938,9 @@ void FPCGActorAndComponentMapping::DeleteMappingPCGComponentPartitionActor(UPCGC
 		return;
 	}
 
-	auto DeleteMapping = [this, InComponent](TMap<const UPCGComponent*, TSet<TObjectPtr<APCGPartitionActor>>>& Map, FRWLock& Lock)
+	auto DeleteMapping = [this, InComponent](TMap<const UPCGComponent*, TSet<TObjectPtr<APCGPartitionActor>>>& Map, PCG::FSharedLock& Lock)
 	{
-		FWriteScopeLock WriteLock(Lock);
+		PCG::TUniqueScopeLock WriteLock(Lock);
 
 		if (TSet<TObjectPtr<APCGPartitionActor>>* PartitionActorsPtr = Map.Find(InComponent))
 		{
@@ -1001,7 +1001,7 @@ TSet<UPCGComponent*> FPCGActorAndComponentMapping::GetAllRegisteredComponents()
 
 UPCGComponent* FPCGActorAndComponentMapping::GetLocalComponent(const FPCGGridDescriptor& GridDescriptor, const FIntVector& CellCoords, const UPCGComponent* InOriginalComponent) const
 {	
-	FReadScopeLock ReadLock(PartitionActorsMapLock);
+	PCG::TSharedScopeLock ReadLock(PartitionActorsMapLock);
 
 	if (const TMap<FIntVector, TObjectPtr<APCGPartitionActor>>* PartitionActorsOnGrid = PartitionActorsMap.Find(GridDescriptor))
 	{
@@ -1017,7 +1017,7 @@ UPCGComponent* FPCGActorAndComponentMapping::GetLocalComponent(const FPCGGridDes
 
 APCGPartitionActor* FPCGActorAndComponentMapping::GetPartitionActor(const FPCGGridDescriptor& GridDescriptor, const FIntVector& CellCoords) const
 {
-	FReadScopeLock ReadLock(PartitionActorsMapLock);
+	PCG::TSharedScopeLock ReadLock(PartitionActorsMapLock);
 
 	if (const TMap<FIntVector, TObjectPtr<APCGPartitionActor>>* PartitionActorsOnGrid = PartitionActorsMap.Find(GridDescriptor))
 	{
@@ -1245,9 +1245,8 @@ bool FPCGActorAndComponentMapping::IsActorTracked(const AActor* InActor) const
 
 void FPCGActorAndComponentMapping::ResetPartitionActorsMap()
 {
-	PartitionActorsMapLock.WriteLock();
+	PCG::TUniqueScopeLock WriteLock(PartitionActorsMapLock);
 	PartitionActorsMap.Empty();
-	PartitionActorsMapLock.WriteUnlock();
 }
 
 void FPCGActorAndComponentMapping::RegisterDelegates()
diff --git a/Engine/Plugins/PCG/Source/PCG/Private/PCGActorAndComponentMapping.h b/Engine/Plugins/PCG/Source/PCG/Private/PCGActorAndComponentMapping.h
--- a/Engine/Plugins/PCG/Source/PCG/Private/PCGActorAndComponentMapping.h
+++ b/Engine/Plugins/PCG/Source/PCG/Private/PCGActorAndComponentMapping.h
@@ -247,11 +247,11 @@ private:
 
 	/** Mapping from grid size and grid coords to partition actor. We can only have 1 partition actor per grid cell. */
 	TMap<FPCGGridDescriptor, TMap<FIntVector, TObjectPtr<APCGPartitionActor>>> PartitionActorsMap;
-	mutable FRWLock PartitionActorsMapLock;
+	mutable PCG::FSharedLock PartitionActorsMapLock;
 
 	/** Mapping between original components and its overlapping partition actors. */
 	TMap<const UPCGComponent*, TSet<TObjectPtr<APCGPartitionActor>>> ComponentToPartitionActorsMap;
-	mutable FRWLock ComponentToPartitionActorsMapLock;
+	mutable PCG::FSharedLock ComponentToPartitionActorsMapLock;
 
 	/** Will hold all the components that are not partitioned (and not local) and are tracking something. Will be use to dispatch actor tracking updates. */
 	FPCGComponentOctreeAndMap NonPartitionedOctree;
@@ -269,7 +269,7 @@ private:
 	/** Same mapping but for always tracked keys */
 	TMap<FPCGSelectionKey, TSet<UPCGComponent*>> AlwaysTrackedKeysToComponentsMap;
 
-	mutable FRWLock TrackedComponentsLock;
+	mutable PCG::FSharedLock TrackedComponentsLock;
 
 	// Keep track of actors that aren't yet ready (or if the subsystem is not yet ready), whether we should dirty them in next tick.
 	TMap<TObjectKey<AActor>, TTuple<TObjectKey<UObject>, EChangedObjectReason>> DelayedChangedActors;
diff --git a/Engine/Plugins/PCG/Source/PCG/Private/PCGGraph.cpp b/Engine/Plugins/PCG/Source/PCG/Private/PCGGraph.cpp
--- a/Engine/Plugins/PCG/Source/PCG/Private/PCGGraph.cpp
+++ b/Engine/Plugins/PCG/Source/PCG/Private/PCGGraph.cpp
@@ -1759,7 +1759,7 @@ void UPCGGraph::NotifyGraphChanged(EPCGChangeType ChangeType)
 	{
 		// Graph settings, nodes, graph structure can all change the higen grid sizes.
 		{
-			FWriteScopeLock GridSizeLock(NodeToGridSizeLock);
+			PCG::TUniqueScopeLock WriteLock(NodeToGridSizeLock);
 			NodeToGridSize.Reset();
 			NodeToAllGridSizes.Reset();
 		}
@@ -1840,7 +1840,7 @@ void UPCGGraph::OnNodeChanged(UPCGNode* InNode, EPCGChangeType ChangeType)
 		{
 			ChangeType |= EPCGChangeType::GenerationGrid;
 
-			FWriteScopeLock Lock(NodeToGridSizeLock);
+			PCG::TUniqueScopeLock WriteLock(NodeToGridSizeLock);
 			NodeToGridSize.Reset();
 			NodeToAllGridSizes.Reset();
 			ChildGridSizeToParentGridSizes.Reset();
@@ -2144,7 +2144,7 @@ FInstancedPropertyBag* UPCGGraph::GetMutableUserParametersStruct()
 uint32 UPCGGraph::GetNodeGenerationGridSize(const UPCGNode* InNode, uint32 InDefaultGridSize) const
 {
 	{
-		FReadScopeLock ScopedReadLock(NodeToGridSizeLock);
+		PCG::TSharedScopeLock ReadLock(NodeToGridSizeLock);
 		if (const uint32* CachedGridSize = NodeToGridSize.Find(InNode))
 		{
 			return *CachedGridSize;
@@ -2152,7 +2152,7 @@ uint32 UPCGGraph::GetNodeGenerationGridSize(const UPCGNode* InNode, uint32 InDef
 	}
 
 	{
-		FWriteScopeLock ScopedWriteLock(NodeToGridSizeLock);
+		PCG::TUniqueScopeLock WriteLock(NodeToGridSizeLock);
 		return CalculateNodeGridSizeRecursive_Unsafe(InNode, InDefaultGridSize);
 	}
 }
diff --git a/Engine/Plugins/PCG/Source/PCG/Private/PCGGraphExecutionInspection.cpp b/Engine/Plugins/PCG/Source/PCG/Private/PCGGraphExecutionInspection.cpp
--- a/Engine/Plugins/PCG/Source/PCG/Private/PCGGraphExecutionInspection.cpp
+++ b/Engine/Plugins/PCG/Source/PCG/Private/PCGGraphExecutionInspection.cpp
@@ -51,19 +51,19 @@ void FPCGGraphExecutionInspection::NotifyNodeExecuted(const UPCGNode* InNode, co
 		Timer = *InTimer;
 	}
 
-	FWriteScopeLock Lock(NodeToStacksInWhichNodeExecutedLock);
+	PCG::TUniqueScopeLock WriteLock(NodeToStacksInWhichNodeExecutedLock);
 	NodeToStacksInWhichNodeExecuted.FindOrAdd(InNode).Add(FPCGGraphExecutionInspection::FNodeExecutedNotificationData(*InStack, MoveTemp(Timer)));
 }
 
 TMap<TObjectKey<const UPCGNode>, TSet<FPCGGraphExecutionInspection::FNodeExecutedNotificationData>> FPCGGraphExecutionInspection::GetExecutedNodeStacks() const
 {
-	FReadScopeLock Lock(NodeToStacksInWhichNodeExecutedLock);
+	PCG::TSharedScopeLock ReadLock(NodeToStacksInWhichNodeExecutedLock);
 	return NodeToStacksInWhichNodeExecuted;
 }
 
 uint64 FPCGGraphExecutionInspection::GetNodeInactivePinMask(const UPCGNode* InNode, const FPCGStack& Stack) const
 {
-	FReadScopeLock Lock(NodeToStackToInactivePinMaskLock);
+	PCG::TSharedScopeLock ReadLock(NodeToStackToInactivePinMaskLock);
 
 	if (const TMap<const FPCGStack, uint64>* StackToMask = NodeToStackToInactivePinMask.Find(InNode))
 	{
@@ -83,14 +83,14 @@ void FPCGGraphExecutionInspection::NotifyNodeDynamicInactivePins(const UPCGNode*
 		return;
 	}
 
-	FWriteScopeLock Lock(NodeToStackToInactivePinMaskLock);
+	PCG::TUniqueScopeLock WriteLock(NodeToStackToInactivePinMaskLock);
 	TMap<const FPCGStack, uint64>& StackToInactivePinMask = NodeToStackToInactivePinMask.FindOrAdd(InNode);
 	StackToInactivePinMask.FindOrAdd(*InStack) = InactivePinBitmask;
 }
 
 bool FPCGGraphExecutionInspection::WasNodeExecuted(const UPCGNode* InNode, const FPCGStack& Stack) const
 {
-	FReadScopeLock Lock(NodeToStacksInWhichNodeExecutedLock);
+	PCG::TSharedScopeLock ReadLock(NodeToStacksInWhichNodeExecutedLock);
 	const TSet<FPCGGraphExecutionInspection::FNodeExecutedNotificationData>* FoundNotifications = NodeToStacksInWhichNodeExecuted.Find(InNode);
 
 	// Since the operator== & hash functions don't rely on the timer, we can just build a stub from the stack.
@@ -110,13 +110,13 @@ void FPCGGraphExecutionInspection::StoreInspectionData(const FPCGStack* InStack,
 
 	if (!InOutputData.TaggedData.IsEmpty())
 	{
-		FWriteScopeLock Lock(NodeToStacksThatProducedDataLock);
+		PCG::TUniqueScopeLock WriteLock(NodeToStacksThatProducedDataLock);
 
 		NodeToStacksThatProducedData.FindOrAdd(InNode).Add(*InStack);
 	}
 	else
 	{
-		FWriteScopeLock Lock(NodeToStacksThatProducedDataLock);
+		PCG::TUniqueScopeLock WriteLock(NodeToStacksThatProducedDataLock);
 
 		if (TSet<FPCGStack>* Stacks = NodeToStacksThatProducedData.Find(InNode))
 		{
@@ -169,7 +169,7 @@ void FPCGGraphExecutionInspection::StoreInspectionData(const FPCGStack* InStack,
 			StorePinInspectionDataFromNode(&Stack, InPins, InData, InOutInspectionCache);
 		};
 
-		FWriteScopeLock Lock(InspectionCacheLock);
+		PCG::TUniqueScopeLock WriteLock(InspectionCacheLock);
 
 		// Special case: if we have a static (embedded) subgraph, then the actual data inputs (not params) of the subgraph will be on the input node.
 		// Considering we don't allow inspection on input pins of the input node, then we can move that data up the chain.
@@ -195,14 +195,14 @@ void FPCGGraphExecutionInspection::StoreInspectionData(const FPCGStack* InStack,
 
 const FPCGDataCollection* FPCGGraphExecutionInspection::GetInspectionData(const FPCGStack& InStack) const
 {
-	FReadScopeLock Lock(InspectionCacheLock);
+	PCG::TSharedScopeLock ReadLock(InspectionCacheLock);
 	return InspectionCache.Find(InStack);
 }
 
 void FPCGGraphExecutionInspection::ClearInspectionData(bool bClearPerNodeExecutionData)
 {
 	{
-		FWriteScopeLock Lock(InspectionCacheLock);
+		PCG::TUniqueScopeLock WriteLock(InspectionCacheLock);
 
 		for (TPair<FPCGStack, FPCGDataCollection>& Entry : InspectionCache)
 		{
@@ -215,22 +215,22 @@ void FPCGGraphExecutionInspection::ClearInspectionData(bool bClearPerNodeExecuti
 	if (bClearPerNodeExecutionData)
 	{
 		{
-			FWriteScopeLock Lock(NodeToStacksThatProducedDataLock);
+			PCG::TUniqueScopeLock WriteLock(NodeToStacksThatProducedDataLock);
 			NodeToStacksThatProducedData.Reset();
 		}
 
 		{
-			FWriteScopeLock Lock(NodeToStacksInWhichNodeExecutedLock);
+			PCG::TUniqueScopeLock WriteLock(NodeToStacksInWhichNodeExecutedLock);
 			NodeToStacksInWhichNodeExecuted.Reset();
 		}
 
 		{
-			FWriteScopeLock Lock(NodeToStackToInactivePinMaskLock);
+			PCG::TUniqueScopeLock WriteLock(NodeToStackToInactivePinMaskLock);
 			NodeToStackToInactivePinMask.Reset();
 		}
 
 		{
-			FWriteScopeLock Lock(NodeToStacksTriggeringGPUTransfersLock);
+			PCG::TUniqueScopeLock WriteLock(NodeToStacksTriggeringGPUTransfersLock);
 			NodeToStacksTriggeringGPUUploads.Reset();
 			NodeToStacksTriggeringGPUReadbacks.Reset();
 		}
@@ -239,7 +239,7 @@ void FPCGGraphExecutionInspection::ClearInspectionData(bool bClearPerNodeExecuti
 
 bool FPCGGraphExecutionInspection::HasNodeProducedData(const UPCGNode* InNode, const FPCGStack& Stack) const
 {
-	FReadScopeLock Lock(NodeToStacksThatProducedDataLock);
+	PCG::TSharedScopeLock ReadLock(NodeToStacksThatProducedDataLock);
 
 	const TSet<FPCGStack>* StacksThatProducedData = NodeToStacksThatProducedData.Find(InNode);
 
@@ -253,7 +253,7 @@ void FPCGGraphExecutionInspection::NotifyGPUToCPUReadback(const UPCGNode* InNode
 		return;
 	}
 
-	FWriteScopeLock Lock(NodeToStacksTriggeringGPUTransfersLock);
+	PCG::TUniqueScopeLock WriteLock(NodeToStacksTriggeringGPUTransfersLock);
 	NodeToStacksTriggeringGPUReadbacks.FindOrAdd(InNode).Add(*InStack);
 }
 
@@ -264,20 +264,20 @@ void FPCGGraphExecutionInspection::NotifyCPUToGPUUpload(const UPCGNode* InNode,
 		return;
 	}
 
-	FWriteScopeLock Lock(NodeToStacksTriggeringGPUTransfersLock);
+	PCG::TUniqueScopeLock WriteLock(NodeToStacksTriggeringGPUTransfersLock);
 	NodeToStacksTriggeringGPUUploads.FindOrAdd(InNode).Add(*InStack);
 }
 
 bool FPCGGraphExecutionInspection::DidNodeTriggerGPUToCPUReadback(const UPCGNode* InNode, const FPCGStack& Stack) const
 {
-	FReadScopeLock Lock(NodeToStacksTriggeringGPUTransfersLock);
+	PCG::TSharedScopeLock ReadLock(NodeToStacksTriggeringGPUTransfersLock);
 	const TSet<FPCGStack>* FoundStacks = NodeToStacksTriggeringGPUReadbacks.Find(InNode);
 	return FoundStacks && FoundStacks->Contains(Stack);
 }
 
 bool FPCGGraphExecutionInspection::DidNodeTriggerCPUToGPUUpload(const UPCGNode* InNode, const FPCGStack& Stack) const
 {
-	FReadScopeLock Lock(NodeToStacksTriggeringGPUTransfersLock);
+	PCG::TSharedScopeLock ReadLock(NodeToStacksTriggeringGPUTransfersLock);
 	const TSet<FPCGStack>* FoundStacks = NodeToStacksTriggeringGPUUploads.Find(InNode);
 	return FoundStacks && FoundStacks->Contains(Stack);
 }
diff --git a/Engine/Plugins/PCG/Source/PCG/Private/RuntimeGen/PCGRuntimeGenScheduler.cpp b/Engine/Plugins/PCG/Source/PCG/Private/RuntimeGen/PCGRuntimeGenScheduler.cpp
--- a/Engine/Plugins/PCG/Source/PCG/Private/RuntimeGen/PCGRuntimeGenScheduler.cpp
+++ b/Engine/Plugins/PCG/Source/PCG/Private/RuntimeGen/PCGRuntimeGenScheduler.cpp
@@ -1124,7 +1124,7 @@ void FPCGRuntimeGenScheduler::TickScheduleGeneration(TMap<FGridGenerationKey, do
 
 				// Update component mapping for this PA (add local component).
 				{
-					FWriteScopeLock WriteLock(ActorAndComponentMapping->ComponentToPartitionActorsMapLock);
+					PCG::TUniqueScopeLock WriteLock(ActorAndComponentMapping->ComponentToPartitionActorsMapLock);
 					TSet<TObjectPtr<APCGPartitionActor>>* PartitionActorsPtr = ActorAndComponentMapping->ComponentToPartitionActorsMap.Find(OriginalComponent);
 
 					if (!PartitionActorsPtr)
@@ -1619,7 +1619,7 @@ void FPCGRuntimeGenScheduler::CleanupLocalComponent(APCGPartitionActor* Partitio
 
 		// Remove component mapping.
 		{
-			FWriteScopeLock WriteLock(ActorAndComponentMapping->ComponentToPartitionActorsMapLock);
+			PCG::TUniqueScopeLock WriteLock(ActorAndComponentMapping->ComponentToPartitionActorsMapLock);
 			TSet<TObjectPtr<APCGPartitionActor>>* PartitionActorsPtr = ActorAndComponentMapping->ComponentToPartitionActorsMap.Find(OriginalComponent);
 
 			if (PartitionActorsPtr)
diff --git a/Engine/Plugins/PCG/Source/PCG/Private/Utils/PCGNodeVisualLogs.cpp b/Engine/Plugins/PCG/Source/PCG/Private/Utils/PCGNodeVisualLogs.cpp
--- a/Engine/Plugins/PCG/Source/PCG/Private/Utils/PCGNodeVisualLogs.cpp
+++ b/Engine/Plugins/PCG/Source/PCG/Private/Utils/PCGNodeVisualLogs.cpp
@@ -26,7 +26,7 @@ void FPCGNodeVisualLogs::Log(const FPCGStack& InPCGStack, ELogVerbosity::Type In
 	bool bAdded = false;
 
 	{
-		FWriteScopeLock ScopedWriteLock(LogsLock);
+		PCG::TUniqueScopeLock WriteLock(LogsLock);
 
 		FPCGPerNodeVisualLogs& NodeLogs = StackToLogs.FindOrAdd(InPCGStack);
 
@@ -81,7 +81,7 @@ void FPCGNodeVisualLogs::Log(const FPCGStack& InPCGStack, ELogVerbosity::Type In
 
 bool FPCGNodeVisualLogs::HasLogs(const FPCGStack& InPCGStack) const
 {
-	FReadScopeLock ScopedReadLock(LogsLock);
+	PCG::TSharedScopeLock ReadLock(LogsLock);
 
 	for (const TPair<FPCGStack, FPCGPerNodeVisualLogs>& Entry : StackToLogs)
 	{
@@ -98,7 +98,7 @@ bool FPCGNodeVisualLogs::HasLogs(const FPCGStack& InPCGStack, ELogVerbosity::Typ
 {
 	OutMinVerbosity = ELogVerbosity::All;
 	
-	FReadScopeLock ScopedReadLock(LogsLock);
+	PCG::TSharedScopeLock ReadLock(LogsLock);
 
 	for (const TPair<FPCGStack, FPCGPerNodeVisualLogs>& Entry : StackToLogs)
 	{
@@ -116,7 +116,7 @@ bool FPCGNodeVisualLogs::HasLogs(const FPCGStack& InPCGStack, ELogVerbosity::Typ
 
 bool FPCGNodeVisualLogs::HasLogsOfVerbosity(const FPCGStack& InPCGStack, ELogVerbosity::Type InVerbosity) const
 {
-	FReadScopeLock ScopedReadLock(LogsLock);
+	PCG::TSharedScopeLock ReadLock(LogsLock);
 
 	for (const TPair<FPCGStack, FPCGPerNodeVisualLogs>& Entry : StackToLogs)
 	{
@@ -147,7 +147,7 @@ FPCGPerNodeVisualLogs FPCGNodeVisualLogs::GetLogs(const FPCGStack& InPCGStack) c
 
 void FPCGNodeVisualLogs::ForAllMatchingLogs(const FPCGStack& InPCGStack, TFunctionRef<bool(const FPCGStack&, const FPCGPerNodeVisualLogs&)> InFunc) const
 {
-	FReadScopeLock ScopedReadLock(LogsLock);
+	PCG::TSharedScopeLock ReadLock(LogsLock);
 	for (const TPair<FPCGStack, FPCGPerNodeVisualLogs>& Entry : StackToLogs)
 	{
 		if (Entry.Key.BeginsWith(InPCGStack))
@@ -165,7 +165,7 @@ void FPCGNodeVisualLogs::GetLogsAndSources(const UPCGNode* InNode, FPCGPerNodeVi
 	OutLogs.Reset();
 	OutExecutionSources.Reset();
 
-	FReadScopeLock ScopedReadLock(LogsLock);
+	PCG::TSharedScopeLock ReadLock(LogsLock);
 	for (const TPair<FPCGStack, FPCGPerNodeVisualLogs>& Entry : StackToLogs)
 	{
 		const FPCGStack& Stack = Entry.Key;
@@ -281,7 +281,7 @@ void FPCGNodeVisualLogs::ClearLogs(const FPCGStack& InPCGStack)
 
 	{
 		FGCScopeGuard Guard;
-		FWriteScopeLock ScopedWriteLock(LogsLock);
+		PCG::TUniqueScopeLock WriteLock(LogsLock);
 
 		TArray<FPCGStack> StacksToRemove;
 		for (const TPair<FPCGStack, FPCGPerNodeVisualLogs>& Entry : StackToLogs)
diff --git a/Engine/Plugins/PCG/Source/PCG/Public/Grid/PCGLandscapeCache.h b/Engine/Plugins/PCG/Source/PCG/Public/Grid/PCGLandscapeCache.h
--- a/Engine/Plugins/PCG/Source/PCG/Public/Grid/PCGLandscapeCache.h
+++ b/Engine/Plugins/PCG/Source/PCG/Public/Grid/PCGLandscapeCache.h
@@ -243,6 +243,6 @@ private:
 #endif
 
 #if WITH_EDITOR
-	FRWLock CacheLock;
+	PCG::FSharedLock CacheLock;
 #endif
 };
diff --git a/Engine/Plugins/PCG/Source/PCG/Public/Metadata/PCGMetadataAttribute.h b/Engine/Plugins/PCG/Source/PCG/Public/Metadata/PCGMetadataAttribute.h
--- a/Engine/Plugins/PCG/Source/PCG/Public/Metadata/PCGMetadataAttribute.h
+++ b/Engine/Plugins/PCG/Source/PCG/Public/Metadata/PCGMetadataAttribute.h
@@ -2,6 +2,7 @@
 
 #pragma once
 
+#include "PCGCommon.h"
 #include "PCGMetadataCommon.h"
 
 #include "Utils/PCGValueRange.h"
@@ -122,7 +123,7 @@ private:
 
 protected:
 	TMap<PCGMetadataEntryKey, PCGMetadataValueKey> EntryToValueKeyMap;
-	mutable FRWLock EntryMapLock;
+	mutable PCG::FSharedLock EntryMapLock;
 
 	FPCGMetadataDomain* Metadata = nullptr;
 	const FPCGMetadataAttributeBase* Parent = nullptr;
diff --git a/Engine/Plugins/PCG/Source/PCG/Public/Metadata/PCGMetadataAttributeTpl.h b/Engine/Plugins/PCG/Source/PCG/Public/Metadata/PCGMetadataAttributeTpl.h
--- a/Engine/Plugins/PCG/Source/PCG/Public/Metadata/PCGMetadataAttributeTpl.h
+++ b/Engine/Plugins/PCG/Source/PCG/Public/Metadata/PCGMetadataAttributeTpl.h
@@ -70,7 +70,7 @@ public:
 		// Flatten values, from root to current attribute
 		if (Parent)
 		{
-			FWriteScopeLock ScopeLock(ValueLock);
+			PCG::TUniqueScopeLock WriteLock(ValueLock);
 
 			TArray<const TArray<T>*> OriginalValues;
 			int32 ValueCount = 0;
@@ -149,9 +149,10 @@ public:
 		}
 
 		// Move the new values in place of the old values.
-		ValueLock.WriteLock();
-		Values = std::move(NewValues);
-		ValueLock.WriteUnlock();
+		{
+			PCG::TUniqueScopeLock WriteLock(ValueLock);
+			Values = std::move(NewValues);
+		}
 
 		// And finally, create a new entry to value mapping.
 		// Logic is that each entry to keep will have their "index" as new entry key
@@ -169,9 +170,10 @@ public:
 		}
 
 		// And move the map.
-		EntryMapLock.WriteLock();
-		EntryToValueKeyMap = std::move(NewMap);
-		EntryMapLock.WriteUnlock();
+		{
+			PCG::TUniqueScopeLock WriteLock(EntryMapLock);
+			EntryToValueKeyMap = std::move(NewMap);
+		}
 
 		// At the end, reset value offset, and lose parent.
 		ValueKeyOffset = 0;
@@ -183,13 +185,15 @@ public:
 		ValueKeyOffset = 0;
 		Parent = nullptr;
 
-		EntryMapLock.WriteLock();
-		EntryToValueKeyMap.Empty();
-		EntryMapLock.WriteUnlock();
+		{
+			PCG::TUniqueScopeLock WriteLock(EntryMapLock);
+			EntryToValueKeyMap.Empty();
+		}
 
-		ValueLock.WriteLock();
-		Values.Empty();
-		ValueLock.WriteUnlock();
+		{
+			PCG::TUniqueScopeLock WriteLock(ValueLock);
+			Values.Empty();
+		}
 	}
 
 	const FPCGMetadataAttribute* GetParent() const { return static_cast<const FPCGMetadataAttribute*>(Parent); }
@@ -277,9 +281,8 @@ protected:
 			{
 				const FPCGMetadataAttribute<T>* Current = Parents[i];
 				
-				Current->EntryMapLock.ReadLock();
+				PCG::TSharedScopeLock ReadLock(EntryMapLock);
 				AttributeCopy->EntryToValueKeyMap.Append(Current->EntryToValueKeyMap);
-				Current->EntryMapLock.ReadUnlock();
 			}
 		}
 
@@ -291,7 +294,7 @@ protected:
 			{
 				const FPCGMetadataAttribute<T>* Current = Parents[i];
 				
-				Current->ValueLock.ReadLock();
+				PCG::TSharedScopeLock ReadLock(Current->ValueLock);
 
 				if constexpr (std::is_same_v<T, U>)
 				{
@@ -314,7 +317,6 @@ protected:
 				{
 					AttributeCopy->ValueKeyOffset = Current->ValueKeyOffset;
 				}
-				Current->ValueLock.ReadUnlock();
 			}
 		}
 
@@ -493,7 +495,7 @@ public:
 
 	PCGMetadataValueKey GetValueKeyOffsetForChild() const
 	{
-		FReadScopeLock ScopeLock(ValueLock);
+		PCG::TSharedScopeLock ReadLock(ValueLock);
 		return Values.Num() + ValueKeyOffset;
 	}
 
@@ -504,7 +506,7 @@ public:
 
 		if (FoundValueKey == PCGNotFoundValueKey)
 		{
-			FWriteScopeLock ScopeLock(ValueLock);
+			PCG::TUniqueScopeLock WriteLock(ValueLock);
 			return Values.Add(InValue) + ValueKeyOffset;
 		}
 		else
@@ -539,7 +541,7 @@ public:
 
 		if (bAtLeastOneValueNotFound)
 		{
-			FWriteScopeLock ScopeLock(ValueLock);
+			PCG::TUniqueScopeLock WriteLock(ValueLock);
 			for (int ValueIndex = 0; ValueIndex < UniqueValues.Num(); ++ValueIndex)
 			{
 				if (FoundKeys[ValueIndex] == PCGNotFoundValueKey)
@@ -568,10 +570,11 @@ public:
 
 		int FirstValueIndex = 0;
 
-		ValueLock.WriteLock();
-		FirstValueIndex = Values.Num() + ValueKeyOffset;
-		Values.Append(InValues);
-		ValueLock.WriteUnlock();
+		{
+			PCG::TUniqueScopeLock WriteLock(ValueLock);
+			FirstValueIndex = Values.Num() + ValueKeyOffset;
+			Values.Append(InValues);
+		}
 
 		ValueKeys.SetNum(InValues.Num());
 		for (int ValueIndex = 0; ValueIndex < InValues.Num(); ++ValueIndex)
@@ -619,7 +622,7 @@ public:
 		else if (ValueKey >= ValueKeyOffset)
 		{
 			int32 Index = ValueKey - ValueKeyOffset;
-			FReadScopeLock ScopeLock(ValueLock);
+			PCG::TSharedScopeLock ReadLock(ValueLock);
 			return Index < Values.Num() ? Values[Index] : DefaultValue;
 		}
 		else if (GetParent())
@@ -646,40 +649,40 @@ protected:
 
 		const FPCGMetadataAttribute* ThisParent = GetParent();
 
-		ValueLock.ReadLock();
-
-		for (; It; ++It)
 		{
-			const int32 Index = It.GetIndex();
-			const PCGMetadataValueKey ValueKey = ValueKeys[Index];
+			PCG::TSharedScopeLock ReadLock(ValueLock);
 
-			auto RetrieveValue = [Index, &OutValues, &UnretrievedValues](T Value)
+			for (; It; ++It)
 			{
-				OutValues[Index] = std::move(Value);
-				UnretrievedValues[Index] = false;
-			};
+				const int32 Index = It.GetIndex();
+				const PCGMetadataValueKey ValueKey = ValueKeys[Index];
 
-			if (ValueKey == PCGDefaultValueKey)
-			{
-				RetrieveValue(DefaultValue);
-			}
-			else if (ValueKey >= ValueKeyOffset)
-			{
-				int32 ValueIndex = ValueKey - ValueKeyOffset;
-				RetrieveValue(ValueIndex < Values.Num() ? Values[ValueIndex] : DefaultValue);
-			}
-			else if (!ThisParent)
-			{
-				RetrieveValue(DefaultValue);
-			}
-			else
-			{
-				bFoundAllKeys = false;
+				auto RetrieveValue = [Index, &OutValues, &UnretrievedValues](T Value)
+					{
+						OutValues[Index] = std::move(Value);
+						UnretrievedValues[Index] = false;
+					};
+
+				if (ValueKey == PCGDefaultValueKey)
+				{
+					RetrieveValue(DefaultValue);
+				}
+				else if (ValueKey >= ValueKeyOffset)
+				{
+					int32 ValueIndex = ValueKey - ValueKeyOffset;
+					RetrieveValue(ValueIndex < Values.Num() ? Values[ValueIndex] : DefaultValue);
+				}
+				else if (!ThisParent)
+				{
+					RetrieveValue(DefaultValue);
+				}
+				else
+				{
+					bFoundAllKeys = false;
+				}
 			}
 		}
 
-		ValueLock.ReadUnlock();
-
 		ensure(ThisParent || bFoundAllKeys);
 
 		if (ThisParent && !bFoundAllKeys)
@@ -755,9 +758,11 @@ public:
 		}
 		else
 		{
-			ValueLock.ReadLock();
-			const int32 ValueIndex = Values.FindLast(InValue);
-			ValueLock.ReadUnlock();
+			int32 ValueIndex = INDEX_NONE;
+			{
+				PCG::TSharedScopeLock ReadLock(ValueLock);
+				ValueIndex = Values.FindLast(InValue);
+			}
 
 			if (ValueIndex != INDEX_NONE)
 			{
@@ -805,7 +810,7 @@ public:
 
 		if (ValueKeysSet != InValues.Num())
 		{
-			ValueLock.ReadLock();
+			PCG::TSharedScopeLock ReadLock(ValueLock);
 			for (int ValueIndex = 0; ValueIndex < InValues.Num(); ++ValueIndex)
 			{
 				if (ValueKeys[ValueIndex] != PCGNotFoundValueKey)
@@ -820,7 +825,6 @@ public:
 					++ValueKeysSet;
 				}
 			}
-			ValueLock.ReadUnlock();
 		}
 	}
 
@@ -859,7 +863,7 @@ public:
 		{
 			if (!bLockless)
 			{
-				ValueLock.WriteLock();
+				ValueLock.Lock();
 			}
 
 			StartIndex = Values.Num();
@@ -875,13 +879,13 @@ public:
 
 			if (!bLockless)
 			{
-				ValueLock.WriteUnlock();
+				ValueLock.Unlock();
 			}
 		}
 
 		if (!bLockless)
 		{
-			EntryMapLock.WriteLock();
+			EntryMapLock.Lock();
 		}
 
 		EntryToValueKeyMap.Reserve(EntryToValueKeyMap.Num() + EntryKeys.Num());
@@ -897,7 +901,7 @@ public:
 
 		if (!bLockless)
 		{
-			EntryMapLock.WriteUnlock();
+			EntryMapLock.Unlock();
 		}
 
 		return StartIndex;
@@ -1099,7 +1103,7 @@ protected:
 	}
 
 protected:
-	mutable FRWLock ValueLock;
+	mutable PCG::FSharedLock ValueLock;
 	TArray<T> Values;
 	T DefaultValue = T{};
 	PCGMetadataValueKey ValueKeyOffset = 0;
diff --git a/Engine/Plugins/PCG/Source/PCG/Public/Metadata/PCGMetadataDomain.h b/Engine/Plugins/PCG/Source/PCG/Public/Metadata/PCGMetadataDomain.h
--- a/Engine/Plugins/PCG/Source/PCG/Public/Metadata/PCGMetadataDomain.h
+++ b/Engine/Plugins/PCG/Source/PCG/Public/Metadata/PCGMetadataDomain.h
@@ -269,8 +269,8 @@ protected:
 	TArray<PCGMetadataEntryKey> ParentKeys;
 	int64 ItemKeyOffset = 0;
 
-	mutable FRWLock AttributeLock;
-	mutable FRWLock ItemLock;
+	mutable PCG::FSharedLock AttributeLock;
+	mutable PCG::FSharedLock ItemLock;
 
 	TAtomic<int64> DelayedEntriesIndex = 0;
 };
@@ -301,7 +301,7 @@ FPCGMetadataAttribute<T>* FPCGMetadataDomain::CreateAttribute(FName AttributeNam
 	FPCGMetadataAttribute<T>* NewAttribute = new FPCGMetadataAttribute<T>(this, AttributeName, ParentAttribute, DefaultValue, bAllowsInterpolation);
 
 	{
-		FWriteScopeLock WriteLock(AttributeLock);
+		PCG::TUniqueScopeLock WriteLock(AttributeLock);
 
 		if (FPCGMetadataAttributeBase** ExistingAttribute = Attributes.Find(AttributeName))
 		{
@@ -351,14 +351,14 @@ template<typename T>
 FPCGMetadataAttribute<T>* FPCGMetadataDomain::FindOrCreateAttribute(FName AttributeName, const T& DefaultValue, bool bAllowsInterpolation, bool bOverrideParent, bool bOverwriteIfTypeMismatch)
 {
 	{
-		FReadScopeLock ScopeLock(AttributeLock);
+		PCG::TSharedScopeLock ReadLock(AttributeLock);
 		if (FPCGMetadataAttribute<T>* Attribute = GetMutableTypedAttribute_Unsafe<T>(AttributeName))
 		{
 			return Attribute;
 		}
 	}
 
-	FWriteScopeLock ScopeLock(AttributeLock);
+	PCG::TUniqueScopeLock WriteLock(AttributeLock);
 	if (FPCGMetadataAttribute<T>* Attribute = GetMutableTypedAttribute_Unsafe<T>(AttributeName))
 	{
 		return Attribute;
diff --git a/Engine/Plugins/PCG/Source/PCG/Public/PCGCommon.h b/Engine/Plugins/PCG/Source/PCG/Public/PCGCommon.h
--- a/Engine/Plugins/PCG/Source/PCG/Public/PCGCommon.h
+++ b/Engine/Plugins/PCG/Source/PCG/Public/PCGCommon.h
@@ -6,6 +6,7 @@
 
 #include "Async/Mutex.h"
 #include "Async/RecursiveMutex.h"
+#include "Async/SharedMutex.h"
 #include "HAL/IConsoleManager.h"
 
 #include "PCGCommon.generated.h"
@@ -29,37 +30,45 @@ namespace PCG
 {
 	using FLock = UE::FMutex;
 	using FRecursiveLock = UE::FRecursiveMutex;
+	using FSharedLock = UE::FSharedMutex;
+
+	template<typename LockType>
+	concept IsSharedLock = requires(LockType Lock)
+	{
+		{ Lock.LockShared() };
+		{ Lock.UnlockShared() };
+	};
 
 	// Similar to existing UE::TScopeLock but allows optional TryLock in which case caller is responsible to 
 	// check if Lock succeeded or not with IsLocked(). This allows execution code to be skipped if non main thread
 	// tasks fail to get the lock instead of waiting.
-	template<typename MutexType>
-	class TScopeLock
+	template<typename LockType> requires (!IsSharedLock<LockType>)
+	class TScopeLock final
 	{
 	public:
 		UE_NONCOPYABLE(TScopeLock);
 
-		UE_NODISCARD_CTOR TScopeLock(MutexType& InMutex, bool bTryLock = false)
-			: Mutex(&InMutex), bLocked(false)
+		UE_NODISCARD_CTOR inline explicit TScopeLock(LockType& InLock, bool bTryLock = false)
+			: Lock(&InLock), bLocked(false)
 		{
-			check(Mutex);
+			check(Lock);
 			if (bTryLock)
 			{
-				bLocked = Mutex->TryLock();
+				bLocked = Lock->TryLock();
 			}
 			else
 			{
-				Mutex->Lock();
+				Lock->Lock();
 				bLocked = true;
 			}
 		}
 
-		~TScopeLock()
+		inline ~TScopeLock()
 		{
-			if (bLocked && Mutex)
+			if (bLocked && Lock)
 			{
-				Mutex->Unlock();
-				Mutex = nullptr;
+				Lock->Unlock();
+				Lock = nullptr;
 			}
 		}
 
@@ -69,9 +78,51 @@ namespace PCG
 		}
 
 	private:
-		MutexType* Mutex;
+		LockType* Lock;
 		bool bLocked;
 	};
+
+	template <typename SharedLockType> requires (IsSharedLock<SharedLockType>)
+	class TSharedScopeLock final
+	{
+	public:
+		UE_NONCOPYABLE(TSharedScopeLock);
+
+		UE_NODISCARD_CTOR inline explicit TSharedScopeLock(SharedLockType& InSharedLock)
+			: SharedLock(InSharedLock)
+		{
+			SharedLock.LockShared();
+		}
+
+		inline ~TSharedScopeLock()
+		{
+			SharedLock.UnlockShared();
+		}
+
+	private:
+		SharedLockType& SharedLock;
+	};
+
+	template <typename SharedLockType> requires (IsSharedLock<SharedLockType>)
+	class TUniqueScopeLock final
+	{
+	public:
+		UE_NONCOPYABLE(TUniqueScopeLock);
+
+		UE_NODISCARD_CTOR inline explicit TUniqueScopeLock(SharedLockType& InSharedLock)
+			: SharedLock(InSharedLock)
+		{
+			SharedLock.Lock();
+		}
+
+		inline ~TUniqueScopeLock()
+		{
+			SharedLock.Unlock();
+		}
+
+	private:
+		SharedLockType& SharedLock;
+	};
 }
 
 namespace PCGLog
diff --git a/Engine/Plugins/PCG/Source/PCG/Public/PCGGraph.h b/Engine/Plugins/PCG/Source/PCG/Public/PCGGraph.h
--- a/Engine/Plugins/PCG/Source/PCG/Public/PCGGraph.h
+++ b/Engine/Plugins/PCG/Source/PCG/Public/PCGGraph.h
@@ -581,7 +581,7 @@ protected:
 	/** All execution grid sizes for nodes (node grid size and all parent grid sizes that it depends on). */
 	mutable TMap<const UPCGNode*, PCGHiGenGrid::FSizeArray> NodeToAllGridSizes;
 
-	mutable FRWLock NodeToGridSizeLock;
+	mutable PCG::FSharedLock NodeToGridSizeLock;
 
 	struct FGridInfo
 	{
diff --git a/Engine/Plugins/PCG/Source/PCG/Public/PCGGraphExecutionInspection.h b/Engine/Plugins/PCG/Source/PCG/Public/PCGGraphExecutionInspection.h
--- a/Engine/Plugins/PCG/Source/PCG/Public/PCGGraphExecutionInspection.h
+++ b/Engine/Plugins/PCG/Source/PCG/Public/PCGGraphExecutionInspection.h
@@ -67,24 +67,24 @@ private:
 
 	TMap<FPCGStack, FPCGDataCollection> InspectionCache;
 
-	mutable FRWLock InspectionCacheLock;
+	mutable PCG::FSharedLock InspectionCacheLock;
 
 	/** Map from nodes to all stacks for which the node produced at least one data item. */
 	TMap<TObjectKey<const UPCGNode>, TSet<FPCGStack>> NodeToStacksThatProducedData;
-	mutable FRWLock NodeToStacksThatProducedDataLock;
+	mutable PCG::FSharedLock NodeToStacksThatProducedDataLock;
 
 	/** Map from nodes to all stacks for which a task for the node was executed. */
 	TMap<TObjectKey<const UPCGNode>, TSet<FNodeExecutedNotificationData>> NodeToStacksInWhichNodeExecuted;
-	mutable FRWLock NodeToStacksInWhichNodeExecutedLock;
+	mutable PCG::FSharedLock NodeToStacksInWhichNodeExecutedLock;
 
 	/** Map from nodes to stacks to mask of output pins that were deactivated during execution. */
 	mutable TMap<TObjectKey<const UPCGNode>, TMap<const FPCGStack, uint64>> NodeToStackToInactivePinMask;
-	mutable FRWLock NodeToStackToInactivePinMaskLock;
+	mutable PCG::FSharedLock NodeToStackToInactivePinMaskLock;
 
 	/** Map from nodes to all stacks for which GPU data transfers occurred. */
 	mutable TMap<TObjectKey<const UPCGNode>, TSet<FPCGStack>> NodeToStacksTriggeringGPUUploads;
 	mutable TMap<TObjectKey<const UPCGNode>, TSet<FPCGStack>> NodeToStacksTriggeringGPUReadbacks;
-	mutable FRWLock NodeToStacksTriggeringGPUTransfersLock;
+	mutable PCG::FSharedLock NodeToStacksTriggeringGPUTransfersLock;
 };
 #endif
 
diff --git a/Engine/Plugins/PCG/Source/PCG/Public/Utils/PCGNodeVisualLogs.h b/Engine/Plugins/PCG/Source/PCG/Public/Utils/PCGNodeVisualLogs.h
--- a/Engine/Plugins/PCG/Source/PCG/Public/Utils/PCGNodeVisualLogs.h
+++ b/Engine/Plugins/PCG/Source/PCG/Public/Utils/PCGNodeVisualLogs.h
@@ -82,7 +82,7 @@ public:
 	static UE_API FText GetSummaryTextWithSources(const FPCGPerNodeVisualLogs& InLogs, const TArray<const IPCGGraphExecutionSource*>* InSources, ELogVerbosity::Type* OutMinimumVerbosity = nullptr);
 private:
 	TMap<FPCGStack, FPCGPerNodeVisualLogs> StackToLogs;
-	mutable FRWLock LogsLock;
+	mutable PCG::FSharedLock LogsLock;
 };
 
 #endif // WITH_EDITOR
diff --git a/Engine/Plugins/PCG/Source/PCGEditor/Private/PCGEditorModule.cpp b/Engine/Plugins/PCG/Source/PCGEditor/Private/PCGEditorModule.cpp
--- a/Engine/Plugins/PCG/Source/PCGEditor/Private/PCGEditorModule.cpp
+++ b/Engine/Plugins/PCG/Source/PCGEditor/Private/PCGEditorModule.cpp
@@ -1016,7 +1016,7 @@ void FPCGEditorModule::OnScheduleGraph(const FPCGStackContext& StackContext)
 	// Record executed stacks.
 	{
 		TRACE_CPUPROFILER_EVENT_SCOPE(UPCGSubsystem::OnScheduleGraph::RecordExecutedStacks);
-		FWriteScopeLock Lock(ExecutedStacksLock);
+		PCG::TUniqueScopeLock WriteLock(ExecutedStacksLock);
 
 		for (const FPCGStack& ExecutedStack : StackContext.GetStacks())
 		{
@@ -1065,7 +1065,7 @@ void FPCGEditorModule::ClearExecutedStacks(const IPCGGraphExecutionSource* InRoo
 	TRACE_CPUPROFILER_EVENT_SCOPE(FPCGEditorModule::ClearExecutedStacksWithComponent);
 	if (InRootSource)
 	{
-		FWriteScopeLock Lock(ExecutedStacksLock);
+		PCG::TUniqueScopeLock WriteLock(ExecutedStacksLock);
 
 		if (TSet<FPCGStackSharedPtr>* AllStackForComponent = ExecutedStacksPerSource.Find(InRootSource))
 		{
@@ -1088,7 +1088,7 @@ void FPCGEditorModule::ClearExecutedStacks(const UPCGGraph* InContainingGraph)
 		return;
 	}
 
-	FWriteScopeLock Lock(ExecutedStacksLock);
+	PCG::TUniqueScopeLock WriteLock(ExecutedStacksLock);
 	for (auto Iterator = ExecutedStacks.CreateIterator(); Iterator; ++Iterator)
 	{
 		FPCGStackSharedPtr StackPtr = *Iterator;
@@ -1110,7 +1110,7 @@ TArray<FPCGStackSharedPtr> FPCGEditorModule::GetExecutedStacksPtrs(const FPCGSta
 	TRACE_CPUPROFILER_EVENT_SCOPE(FPCGEditorModule::GetExecutedStacksPtrsWithBeginning);
 	TArray<FPCGStackSharedPtr> MatchingStacks;
 
-	FReadScopeLock Lock(ExecutedStacksLock);
+	PCG::TSharedScopeLock ReadLock(ExecutedStacksLock);
 	TSet<FPCGStackSharedPtr>* StacksToTest = ExecutedStacksPerSource.Find(BeginningWithStack.GetRootSource());
 
 	// Stack doesn't have a root component (unlikely), we'll check everything then
@@ -1134,7 +1134,7 @@ TArray<FPCGStackSharedPtr> FPCGEditorModule::GetExecutedStacksPtrs(const IPCGGra
 {
 	TRACE_CPUPROFILER_EVENT_SCOPE(FPCGEditorModule::GetExecutedStacksPtrs);
 	TArray<FPCGStackSharedPtr> MatchingStacks;
-	FReadScopeLock Lock(ExecutedStacksLock);
+	PCG::TSharedScopeLock ReadLock(ExecutedStacksLock);
 
 	TSet<FPCGStackSharedPtr>* ComponentSet = nullptr;
 	TSet<FPCGStackSharedPtr>* GraphSet = nullptr;
diff --git a/Engine/Plugins/PCG/Source/PCGEditor/Public/PCGEditorModule.h b/Engine/Plugins/PCG/Source/PCGEditor/Public/PCGEditorModule.h
--- a/Engine/Plugins/PCG/Source/PCGEditor/Public/PCGEditorModule.h
+++ b/Engine/Plugins/PCG/Source/PCGEditor/Public/PCGEditorModule.h
@@ -96,7 +96,7 @@ protected:
 	/** A record of stacks that were executed. Used to populate debugging tool UIs. */
 	TSet<FPCGStackSharedPtr> ExecutedStacks;
 	TMap<FPCGSoftGraphExecutionSource, TSet<FPCGStackSharedPtr>> ExecutedStacksPerSource;
-	mutable FRWLock ExecutedStacksLock;
+	mutable PCG::FSharedLock ExecutedStacksLock;
 
 	EPCGEditorPermissionMode PermissionMode = EPCGEditorPermissionMode::All;
 	FPermissionModeChanged PermissionModeChangedEvent;
