[PCG] Cache player position and direction on player generation source.

We query these from worker threads which is dicey, and may anyway be more efficient to cache.

#rnx

diff --git a/Engine/Plugins/PCG/Source/PCG/Private/RuntimeGen/GenSources/PCGGenSourcePlayer.cpp b/Engine/Plugins/PCG/Source/PCG/Private/RuntimeGen/GenSources/PCGGenSourcePlayer.cpp
--- a/Engine/Plugins/PCG/Source/PCG/Private/RuntimeGen/GenSources/PCGGenSourcePlayer.cpp
+++ b/Engine/Plugins/PCG/Source/PCG/Private/RuntimeGen/GenSources/PCGGenSourcePlayer.cpp
@@ -16,44 +16,37 @@ void UPCGGenSourcePlayer::Tick()
 
 	const ULocalPlayer* LocalPlayer = PlayerController.IsValid() ? PlayerController->GetLocalPlayer() : nullptr;
 
-	if (LocalPlayer && LocalPlayer->ViewportClient)
+	if (LocalPlayer)
 	{
-		FSceneViewProjectionData ProjectionData;
-
-		if (LocalPlayer->GetProjectionData(LocalPlayer->ViewportClient->Viewport, ProjectionData))
+		if (LocalPlayer->ViewportClient)
 		{
-			FConvexVolume ConvexVolume;
-			GetViewFrustumBounds(ConvexVolume, ProjectionData.ComputeViewProjectionMatrix(), /*bUseNearPlane=*/true, /*bUseFarPlane=*/true);
+			FSceneViewProjectionData ProjectionData;
+
+			if (LocalPlayer->GetProjectionData(LocalPlayer->ViewportClient->Viewport, ProjectionData))
+			{
+				FConvexVolume ConvexVolume;
+				GetViewFrustumBounds(ConvexVolume, ProjectionData.ComputeViewProjectionMatrix(), /*bUseNearPlane=*/true, /*bUseFarPlane=*/true);
 
-			ViewFrustum = ConvexVolume;
+				ViewFrustum = ConvexVolume;
+			}
 		}
-	}
-}
 
-TOptional<FVector> UPCGGenSourcePlayer::GetPosition() const
-{
-	if (PlayerController.IsValid())
-	{
 		FVector CameraLocation;
 		FRotator CameraRotation;
 		PlayerController->GetPlayerViewPoint(CameraLocation, CameraRotation);
-		return CameraLocation;
+		CachedPosition = CameraLocation;
+		CachedDirection = CameraRotation.Vector();
 	}
+}
 
-	return TOptional<FVector>();
+TOptional<FVector> UPCGGenSourcePlayer::GetPosition() const
+{
+	return CachedPosition;
 }
 
 TOptional<FVector> UPCGGenSourcePlayer::GetDirection() const
 {
-	if (PlayerController.IsValid())
-	{
-		FVector CameraLocation;
-		FRotator CameraRotation;
-		PlayerController->GetPlayerViewPoint(CameraLocation, CameraRotation);
-		return CameraRotation.Vector();
-	}
-
-	return TOptional<FVector>();
+	return CachedDirection;
 }
 
 void UPCGGenSourcePlayer::SetPlayerController(const APlayerController* InPlayerController)
diff --git a/Engine/Plugins/PCG/Source/PCG/Public/RuntimeGen/GenSources/PCGGenSourcePlayer.h b/Engine/Plugins/PCG/Source/PCG/Public/RuntimeGen/GenSources/PCGGenSourcePlayer.h
--- a/Engine/Plugins/PCG/Source/PCG/Public/RuntimeGen/GenSources/PCGGenSourcePlayer.h
+++ b/Engine/Plugins/PCG/Source/PCG/Public/RuntimeGen/GenSources/PCGGenSourcePlayer.h
@@ -42,6 +42,9 @@ public:
 protected:
 	TWeakObjectPtr<const APlayerController> PlayerController;
 	TOptional<FConvexVolume> ViewFrustum;
+
+	TOptional<FVector> CachedPosition;
+	TOptional<FVector> CachedDirection;
 };
 
 #undef UE_API
