[PCG] Replace remaining UE::FSpinLock & make some locks shared

#rnx

diff --git a/Engine/Plugins/PCG/Source/PCG/Private/Compute/Data/PCGProxyForGPUData.cpp b/Engine/Plugins/PCG/Source/PCG/Private/Compute/Data/PCGProxyForGPUData.cpp
--- a/Engine/Plugins/PCG/Source/PCG/Private/Compute/Data/PCGProxyForGPUData.cpp
+++ b/Engine/Plugins/PCG/Source/PCG/Private/Compute/Data/PCGProxyForGPUData.cpp
@@ -150,7 +150,7 @@ FPCGProxyForGPUDataCollection::FPCGProxyForGPUDataCollection(TRefCountPtr<FRDGPo
 
 bool FPCGProxyForGPUDataCollection::GetCPUData(FPCGContext* InContext, int32 InDataIndex, FPCGTaggedData& OutData)
 {
-	UE::TScopeLock Lock(ReadbackLock);
+	PCG::TScopeLock Lock(ReadbackLock);
 
 	if (bReadbackDataProcessed)
 	{
diff --git a/Engine/Plugins/PCG/Source/PCG/Private/Data/PCGSplineStruct.cpp b/Engine/Plugins/PCG/Source/PCG/Private/Data/PCGSplineStruct.cpp
--- a/Engine/Plugins/PCG/Source/PCG/Private/Data/PCGSplineStruct.cpp
+++ b/Engine/Plugins/PCG/Source/PCG/Private/Data/PCGSplineStruct.cpp
@@ -413,7 +413,7 @@ const TArray<FName>& FPCGSplineStruct::GetAttributeChannelNames() const
 		Channels.Sort([](const FName& A, const FName& B)
 			{ return A.ToString() < B.ToString(); });
 		
-		if (UE::TScopeLock Lock(SpinLock); Channels.Num() != CachedSortedAttributeChannelNames.Num())
+		if (PCG::TScopeLock ScopeLock(Lock); Channels.Num() != CachedSortedAttributeChannelNames.Num())
 		{
 			CachedSortedAttributeChannelNames = MoveTemp(Channels);
 		}
diff --git a/Engine/Plugins/PCG/Source/PCG/Private/Graph/PCGGraphCache.cpp b/Engine/Plugins/PCG/Source/PCG/Private/Graph/PCGGraphCache.cpp
--- a/Engine/Plugins/PCG/Source/PCG/Private/Graph/PCGGraphCache.cpp
+++ b/Engine/Plugins/PCG/Source/PCG/Private/Graph/PCGGraphCache.cpp
@@ -141,7 +141,7 @@ bool FPCGGraphCache::GetFromCache(const FPCGGetFromCacheParams& Params, FPCGData
 
 	{
 		TRACE_CPUPROFILER_EVENT_SCOPE(FPCGGraphCache::GetFromCache);
-		UE::TScopeLock ScopedLock(CacheLock);
+		PCG::TScopeLock ScopedLock(CacheLock);
 
 		FPCGCacheEntryKey CacheKey(InElement, InDependenciesCrc);
 		if (const FPCGDataCollection* Value = const_cast<FPCGGraphCache*>(this)->CacheData.FindAndTouch(CacheKey))
@@ -191,7 +191,7 @@ void FPCGGraphCache::StoreInCache(const FPCGStoreInCacheParams& Params, const FP
 
 	{
 		TRACE_CPUPROFILER_EVENT_SCOPE(FPCGGraphCache::StoreInCache);
-		UE::TScopeLock ScopedLock(CacheLock);
+		PCG::TScopeLock ScopedLock(CacheLock);
 
 		if (CacheData.Num() == CacheData.Max())
 		{
@@ -205,7 +205,7 @@ void FPCGGraphCache::StoreInCache(const FPCGStoreInCacheParams& Params, const FP
 
 void FPCGGraphCache::ClearCache()
 {
-	UE::TScopeLock ScopedLock(CacheLock);
+	PCG::TScopeLock ScopedLock(CacheLock);
 
 	// Remove all entries
 	ClearCacheInternal(CacheData.Max(), /*bClearMemory=*/true);
@@ -231,7 +231,7 @@ bool FPCGGraphCache::EnforceMemoryBudget()
 	}
 
 	{
-		UE::TScopeLock ScopedLock(CacheLock);
+		PCG::TScopeLock ScopedLock(CacheLock);
 		const float MemoryCleanupRatio = FMath::Clamp(CVarCacheMemoryCleanupRatio.GetValueOnAnyThread(), 0.0f, 1.0f);
 		const uint64 TargetCacheMemoryUsage = static_cast<uint64>(MemoryCleanupRatio * MemoryBudget);
 
@@ -259,7 +259,7 @@ void FPCGGraphCache::CleanFromCache(const IPCGElement* InElement, const UPCGSett
 	}
 
 	{
-		UE::TScopeLock ScopedLock(CacheLock);
+		PCG::TScopeLock ScopedLock(CacheLock);
 
 		ValidateElementToCacheEntryKeys();
 
@@ -278,7 +278,7 @@ void FPCGGraphCache::CleanFromCache(const IPCGElement* InElement, const UPCGSett
 uint32 FPCGGraphCache::GetGraphCacheEntryCount(IPCGElement* InElement) const
 {
 	TRACE_CPUPROFILER_EVENT_SCOPE(FPCGGraphCache::GetFromCache);
-	UE::TScopeLock ScopedLock(CacheLock);
+	PCG::TScopeLock ScopedLock(CacheLock);
 
 	if (const TSet<FPCGCacheEntryKey>* ElementCacheEntryKeys = ElementToCacheEntryKeys.Find(InElement))
 	{
@@ -292,7 +292,7 @@ uint32 FPCGGraphCache::GetGraphCacheEntryCount(IPCGElement* InElement) const
 void FPCGGraphCache::AddReferencedObjects(FReferenceCollector& Collector)
 {
 	TRACE_CPUPROFILER_EVENT_SCOPE(FPCGGraphCache::AddReferencedObjects);
-	UE::TScopeLock ScopedLock(CacheLock);
+	PCG::TScopeLock ScopedLock(CacheLock);
 
 	for (FPCGDataCollection& CacheEntry : CacheData)
 	{
diff --git a/Engine/Plugins/PCG/Source/PCG/Private/Graph/PCGGraphCache.h b/Engine/Plugins/PCG/Source/PCG/Private/Graph/PCGGraphCache.h
--- a/Engine/Plugins/PCG/Source/PCG/Private/Graph/PCGGraphCache.h
+++ b/Engine/Plugins/PCG/Source/PCG/Private/Graph/PCGGraphCache.h
@@ -130,5 +130,5 @@ private:
 
 	bool bIsGameWorld = true;
 
-	mutable UE::FSpinLock CacheLock;
+	mutable PCG::FLock CacheLock;
 };
diff --git a/Engine/Plugins/PCG/Source/PCG/Private/Graph/PCGGraphExecutor.cpp b/Engine/Plugins/PCG/Source/PCG/Private/Graph/PCGGraphExecutor.cpp
--- a/Engine/Plugins/PCG/Source/PCG/Private/Graph/PCGGraphExecutor.cpp
+++ b/Engine/Plugins/PCG/Source/PCG/Private/Graph/PCGGraphExecutor.cpp
@@ -765,7 +765,7 @@ void FPCGGraphExecutor::MarkInputResults(TArrayView<const FPCGTaskId> InInputRes
 
 	if (InInputResults.Num() > 0)
 	{
-		PCG::TScopeLock ScopeLock(TaskOutputsLock);
+		PCG::TUniqueScopeLock WriteLock(TaskOutputsLock);
 		for (FPCGTaskId TaskId : InInputResults)
 		{
 			FOutputDataInfo& OutputDataInfo = TaskOutputs[TaskId];
@@ -782,7 +782,7 @@ void FPCGGraphExecutor::MarkInputResults(TArrayView<const FPCGTaskId> InInputRes
 bool FPCGGraphExecutor::GetOutputData(FPCGTaskId TaskId, FPCGDataCollection& OutData)
 {
 	TRACE_CPUPROFILER_EVENT_SCOPE(FPCGGraphExecutor::GetOutputData);
-	PCG::TScopeLock ScopeLock(TaskOutputsLock);
+	PCG::TSharedScopeLock ReadLock(TaskOutputsLock);
 	if (FOutputDataInfo* OutputDataInfo = TaskOutputs.Find(TaskId))
 	{
 		ensure(!OutputDataInfo->bCulled);
@@ -798,7 +798,7 @@ bool FPCGGraphExecutor::GetOutputData(FPCGTaskId TaskId, FPCGDataCollection& Out
 void FPCGGraphExecutor::ClearOutputData(FPCGTaskId TaskId)
 {
 	TRACE_CPUPROFILER_EVENT_SCOPE(FPCGGraphExecutor::ClearOutputData);
-	PCG::TScopeLock ScopeLock(TaskOutputsLock);
+	PCG::TUniqueScopeLock WriteLock(TaskOutputsLock);
 	if (FOutputDataInfo* OutputDataInfo = TaskOutputs.Find(TaskId); OutputDataInfo && ensure(OutputDataInfo->bNeedsManualClear))
 	{
 		TaskOutputs.Remove(TaskId);
@@ -875,7 +875,7 @@ bool FPCGGraphExecutor::ProcessScheduledTasks(double EndTime)
 					TArray<FPCGGraphTask, TInlineAllocator<16>> NewReadyTasks;
 					{
 						PCG::TScopeLock ChildScopeLock(TasksLock);
-						PCG::TScopeLock OutputLock(TaskOutputsLock);
+						PCG::TSharedScopeLock ReadLock(TaskOutputsLock);
 
 						// Every X ScheduledTasks, we want to check the frame budget and early out if we consumed it.
 						const int32 MaxScheduledTasksPerBudgetCheck = 32;
@@ -1080,7 +1080,7 @@ void FPCGGraphExecutor::PostTaskExecute(TSharedPtr<FPCGGraphActiveTask> ActiveTa
 				if (bHasPausedDependencies)
 				{
 					PCG::TScopeLock PauseLock(PausedTaskSuccessorsLock);
-					PCG::TScopeLock OutputLock(TaskOutputsLock);
+					PCG::TSharedScopeLock ReadLock(TaskOutputsLock);
 					for (TSet<FPCGTaskId>::TIterator It = ActiveTask.Context->DynamicDependencies.CreateIterator(); It; ++It)
 					{
 						FPCGTaskId DependencyId = *It;
@@ -1919,7 +1919,7 @@ TArray<FPCGGraphExecutor::FCachedResult*> FPCGGraphExecutor::QueueNextTasksInter
 	TArray<FPCGGraphTask, TInlineAllocator<16>> NewReadyTasks;
 	{
 		PCG::TScopeLock ScopeLock(TasksLock);
-		PCG::TScopeLock OutputLock(TaskOutputsLock);
+		PCG::TSharedScopeLock ReadLock(TaskOutputsLock);
 
 		if (TaskSuccessors.RemoveAndCopyValue(FinishedTask, Successors))
 		{
@@ -2176,7 +2176,7 @@ void FPCGGraphExecutor::BuildTaskInput(FPCGGraphTask& Task, TArray<FPCGTaskId>&
 	
 	if(Task.Inputs.Num() > 0)
 	{
-		PCG::TScopeLock OutputLock(TaskOutputsLock);
+		PCG::TSharedScopeLock ReadLock(TaskOutputsLock);
 		for (const FPCGGraphTaskInput& Input : Task.Inputs)
 		{
 			const FOutputDataInfo& TaskOutput = TaskOutputs.FindChecked(Input.TaskId);
@@ -2360,7 +2360,7 @@ void FPCGGraphExecutor::StoreResults(FPCGTaskId InTaskId, const FPCGDataCollecti
 	}
 
 	// Store output in map
-	PCG::TScopeLock OutputLock(TaskOutputsLock);
+	PCG::TUniqueScopeLock WriteLock(TaskOutputsLock);
 	TaskOutputs.Add(InTaskId, MoveTemp(OutputDataInfo));
 }
 
@@ -2383,7 +2383,7 @@ void FPCGGraphExecutor::ClearResults()
 		}
 	}
 
-	PCG::TScopeLock OutputLock(TaskOutputsLock);
+	PCG::TUniqueScopeLock WriteLock(TaskOutputsLock);
 
 	for (const TPair<FPCGTaskId, FOutputDataInfo>& TaskOutput : TaskOutputs)
 	{
@@ -2525,7 +2525,7 @@ void FPCGGraphExecutor::CullInactiveDownstreamNodes(FPCGTaskId InCompletedTaskId
 	{
 		{
 			// Store culled output
-			PCG::TScopeLock OutputLock(TaskOutputsLock);
+			PCG::TUniqueScopeLock WriteLock(TaskOutputsLock);
 			for (const FPCGTaskId TaskId : AllRemovedTasks)
 			{
 				FOutputDataInfo CullDataInfo;
@@ -2574,7 +2574,7 @@ void FPCGGraphExecutor::AddReferencedObjects(FReferenceCollector& Collector)
 
 	// Go through all data in the cached output map
 	{
-		PCG::TScopeLock ScopeLock(TaskOutputsLock);
+		PCG::TSharedScopeLock ReadLock(TaskOutputsLock);
 		for (auto& OutputDataEntry : TaskOutputs)
 		{
 			OutputDataEntry.Value.DataCollection.AddReferences(Collector);
diff --git a/Engine/Plugins/PCG/Source/PCG/Private/Graph/PCGGraphExecutor.h b/Engine/Plugins/PCG/Source/PCG/Private/Graph/PCGGraphExecutor.h
--- a/Engine/Plugins/PCG/Source/PCG/Private/Graph/PCGGraphExecutor.h
+++ b/Engine/Plugins/PCG/Source/PCG/Private/Graph/PCGGraphExecutor.h
@@ -343,7 +343,7 @@ private:
 	TSet<FPCGContext*> ReleasedContextsForGC;
 
 	/** Lock level leaf */
-	PCG::FLock TaskOutputsLock;
+	PCG::FSharedLock TaskOutputsLock;
 	TMap<FPCGTaskId, FOutputDataInfo> TaskOutputs;
 	
 	/** Monotonically increasing id. Should be reset once all tasks are executed, should be protected by the ScheduleLock */
diff --git a/Engine/Plugins/PCG/Source/PCG/Private/Graph/PCGGraphPerExecutionCache.cpp b/Engine/Plugins/PCG/Source/PCG/Private/Graph/PCGGraphPerExecutionCache.cpp
--- a/Engine/Plugins/PCG/Source/PCG/Private/Graph/PCGGraphPerExecutionCache.cpp
+++ b/Engine/Plugins/PCG/Source/PCG/Private/Graph/PCGGraphPerExecutionCache.cpp
@@ -7,7 +7,7 @@ UPCGData* FPCGPerExecutionCache::GetExecutionCachePCGData(FPCGTaskId InGraphExec
 	bOutFound = false;
 	if (InGraphExecutionTaskId != InvalidPCGTaskId)
 	{
-		UE::TScopeLock ScopeLock(Lock);
+		PCG::TSharedScopeLock ReadLock(Lock);
 		if (FPCGPerExecutionCacheEntry* PerExecutionCacheEntry = Entries.Find(InGraphExecutionTaskId))
 		{
 			TInstancedStruct<FPCGPerExecutionCacheData>& StructData = PerExecutionCacheEntry->Data[(uint8)InExecutionCacheDataType];
@@ -27,7 +27,7 @@ void FPCGPerExecutionCache::SetExecutionCachePCGData(FPCGTaskId InGraphExecution
 {
 	if (InGraphExecutionTaskId != InvalidPCGTaskId)
 	{
-		UE::TScopeLock ScopeLock(Lock);
+		PCG::TUniqueScopeLock WriteLock(Lock);
 		FPCGPerExecutionCacheEntry& PerExecutionCacheEntry = Entries.FindOrAdd(InGraphExecutionTaskId);
 		check(!PerExecutionCacheEntry.Data[(uint8)InExecutionCacheDataType].IsValid());
 		PerExecutionCacheEntry.Data[(uint8)InExecutionCacheDataType].InitializeAs<FPCGPerExecutionCachePCGData>(InData);
@@ -38,7 +38,7 @@ TOptional<FBox> FPCGPerExecutionCache::GetExecutionCacheBounds(FPCGTaskId InGrap
 {
 	if (InGraphExecutionTaskId != InvalidPCGTaskId)
 	{
-		UE::TScopeLock ScopeLock(Lock);
+		PCG::TSharedScopeLock ReadLock(Lock);
 		if (FPCGPerExecutionCacheEntry* PerExecutionCacheEntry = Entries.Find(InGraphExecutionTaskId))
 		{
 			TInstancedStruct<FPCGPerExecutionCacheData>& StructData = PerExecutionCacheEntry->Data[(uint8)InExecutionCacheDataType];
@@ -57,7 +57,7 @@ void FPCGPerExecutionCache::SetExecutionCacheBounds(FPCGTaskId InGraphExecutionT
 {
 	if (InGraphExecutionTaskId != InvalidPCGTaskId)
 	{
-		UE::TScopeLock ScopeLock(Lock);
+		PCG::TUniqueScopeLock WriteLock(Lock);
 		FPCGPerExecutionCacheEntry& PerExecutionCacheEntry = Entries.FindOrAdd(InGraphExecutionTaskId);
 		check(!PerExecutionCacheEntry.Data[(uint8)InExecutionCacheDataType].IsValid());
 		PerExecutionCacheEntry.Data[(uint8)InExecutionCacheDataType].InitializeAs<FPCGPerExecutionCacheBounds>(InBounds);
@@ -66,13 +66,13 @@ void FPCGPerExecutionCache::SetExecutionCacheBounds(FPCGTaskId InGraphExecutionT
 
 void FPCGPerExecutionCache::Clear()
 {
-	UE::TScopeLock ScopeLock(Lock);
+	PCG::TUniqueScopeLock WriteLock(Lock);
 	Entries.Empty();
 }
 
 void FPCGPerExecutionCache::AddStructReferencedObjects(FReferenceCollector& Collector)
 {
-	UE::TScopeLock ScopeLock(Lock);
+	PCG::TSharedScopeLock ReadLock(Lock);
 	for (TMap<FPCGTaskId, FPCGPerExecutionCacheEntry>::TIterator It = Entries.CreateIterator(); It; ++It)
 	{
 		It.Value().AddStructReferencedObjects(Collector);
diff --git a/Engine/Plugins/PCG/Source/PCG/Private/Graph/PCGGraphPerExecutionCache.h b/Engine/Plugins/PCG/Source/PCG/Private/Graph/PCGGraphPerExecutionCache.h
--- a/Engine/Plugins/PCG/Source/PCG/Private/Graph/PCGGraphPerExecutionCache.h
+++ b/Engine/Plugins/PCG/Source/PCG/Private/Graph/PCGGraphPerExecutionCache.h
@@ -77,8 +77,6 @@ struct FPCGPerExecutionCacheEntry
 
 struct FPCGPerExecutionCache
 {
-	using FPCGLock = UE::FSpinLock;
-
 	UPCGData* GetExecutionCachePCGData(FPCGTaskId InGraphExecutionTaskId, EPCGPerExecutionCacheDataType InExecutionCacheDataType, bool& bOutFound);
 	void SetExecutionCachePCGData(FPCGTaskId InGraphExecutionTaskId, EPCGPerExecutionCacheDataType InExecutionCacheDataType, UPCGData* InData);
 
@@ -89,6 +87,6 @@ struct FPCGPerExecutionCache
 	void AddStructReferencedObjects(FReferenceCollector& Collector);
 
 	/** Per graph execution cache, gets emptied when executor has no more work to do */
-	FPCGLock Lock;
+	PCG::FSharedLock Lock;
 	TMap<FPCGTaskId, FPCGPerExecutionCacheEntry> Entries;
 };
\ No newline at end of file
diff --git a/Engine/Plugins/PCG/Source/PCG/Private/Metadata/PCGMetadata.cpp b/Engine/Plugins/PCG/Source/PCG/Private/Metadata/PCGMetadata.cpp
--- a/Engine/Plugins/PCG/Source/PCG/Private/Metadata/PCGMetadata.cpp
+++ b/Engine/Plugins/PCG/Source/PCG/Private/Metadata/PCGMetadata.cpp
@@ -212,7 +212,7 @@ FPCGMetadataDomain* UPCGMetadata::FindOrCreateMetadataDomain(const FPCGMetadataD
 
 	if (!FoundMetadataDomain->IsValid())
 	{
-		UE::TScopeLock<UE::FSpinLock> Lock(MetadataDomainsSpinLock);
+		PCG::TScopeLock Lock(MetadataDomainsLock);
 		
 		FoundMetadataDomain = MetadataDomains.Find(InMetadataDomainID);
 		if (!FoundMetadataDomain->IsValid())
diff --git a/Engine/Plugins/PCG/Source/PCG/Private/PCGGraphExecutionStateInterface.cpp b/Engine/Plugins/PCG/Source/PCG/Private/PCGGraphExecutionStateInterface.cpp
--- a/Engine/Plugins/PCG/Source/PCG/Private/PCGGraphExecutionStateInterface.cpp
+++ b/Engine/Plugins/PCG/Source/PCG/Private/PCGGraphExecutionStateInterface.cpp
@@ -121,7 +121,7 @@ IPCGGraphExecutionSource* FPCGSoftGraphExecutionSource::Get() const
 		IPCGGraphExecutionSource* InterfacePtr = CastChecked<IPCGGraphExecutionSource>(Object);
 		
 		{
-			UE::TScopeLock Lock(SpinLock);
+			PCG::TScopeLock ScopeLock(Lock);
 			if (!CachedWeakPtr.IsValid())
 			{
 				CachedWeakPtr = CastChecked<IPCGGraphExecutionSource>(Object);
@@ -142,7 +142,7 @@ void FPCGSoftGraphExecutionSource::Reset()
 	SoftObjectPtr = nullptr;
 	if (CachedWeakPtr.IsValid())
 	{
-		UE::TScopeLock Lock(SpinLock);
+		PCG::TScopeLock ScopeLock(Lock);
 		CachedWeakPtr.Reset();
 	}
 }
diff --git a/Engine/Plugins/PCG/Source/PCG/Public/Compute/Data/PCGProxyForGPUData.h b/Engine/Plugins/PCG/Source/PCG/Public/Compute/Data/PCGProxyForGPUData.h
--- a/Engine/Plugins/PCG/Source/PCG/Public/Compute/Data/PCGProxyForGPUData.h
+++ b/Engine/Plugins/PCG/Source/PCG/Public/Compute/Data/PCGProxyForGPUData.h
@@ -124,5 +124,5 @@ private:
 
 	bool bReadbackDataArrived = false;
 
-	UE::FSpinLock ReadbackLock;
+	PCG::FLock ReadbackLock;
 };
diff --git a/Engine/Plugins/PCG/Source/PCG/Public/Data/PCGSplineStruct.h b/Engine/Plugins/PCG/Source/PCG/Public/Data/PCGSplineStruct.h
--- a/Engine/Plugins/PCG/Source/PCG/Public/Data/PCGSplineStruct.h
+++ b/Engine/Plugins/PCG/Source/PCG/Public/Data/PCGSplineStruct.h
@@ -2,6 +2,7 @@
 
 #pragma once
 
+#include "PCGCommon.h"
 #include "Metadata/PCGMetadataCommon.h"
 
 #include "CoreMinimal.h"
@@ -175,7 +176,7 @@ private:
 	FSpline Spline;
 	PRAGMA_ENABLE_EXPERIMENTAL_WARNINGS
 
-	mutable UE::FSpinLock SpinLock;
+	mutable PCG::FLock Lock;
 	mutable TArray<FName> CachedSortedAttributeChannelNames;
 
 	UE_API bool ShouldUseSplineCurves() const;
diff --git a/Engine/Plugins/PCG/Source/PCG/Public/Metadata/PCGMetadata.h b/Engine/Plugins/PCG/Source/PCG/Public/Metadata/PCGMetadata.h
--- a/Engine/Plugins/PCG/Source/PCG/Public/Metadata/PCGMetadata.h
+++ b/Engine/Plugins/PCG/Source/PCG/Public/Metadata/PCGMetadata.h
@@ -483,7 +483,7 @@ protected:
 
 	FPCGMetadataDomainID DefaultDomain = PCGMetadataDomainID::Default;
 	TMap<FPCGMetadataDomainID, TSharedPtr<FPCGMetadataDomain>> MetadataDomains;
-	UE::FSpinLock MetadataDomainsSpinLock;
+	PCG::FLock MetadataDomainsLock;
 };
 
 template<typename T>
diff --git a/Engine/Plugins/PCG/Source/PCG/Public/PCGGraphExecutionStateInterface.h b/Engine/Plugins/PCG/Source/PCG/Public/PCGGraphExecutionStateInterface.h
--- a/Engine/Plugins/PCG/Source/PCG/Public/PCGGraphExecutionStateInterface.h
+++ b/Engine/Plugins/PCG/Source/PCG/Public/PCGGraphExecutionStateInterface.h
@@ -213,6 +213,6 @@ private:
 	UPROPERTY()
 	TSoftObjectPtr<UObject> SoftObjectPtr;
 
-	mutable UE::FSpinLock SpinLock;
+	mutable PCG::FLock Lock;
 	mutable TWeakInterfacePtr<IPCGGraphExecutionSource> CachedWeakPtr;
 };
