[PCG] Expose more graph functionalities for Blueprints

#rnx

diff --git a/Engine/Plugins/PCG/Source/PCG/Private/PCGEdge.cpp b/Engine/Plugins/PCG/Source/PCG/Private/PCGEdge.cpp
--- a/Engine/Plugins/PCG/Source/PCG/Private/PCGEdge.cpp
+++ b/Engine/Plugins/PCG/Source/PCG/Private/PCGEdge.cpp
@@ -33,3 +33,23 @@ const UPCGPin* UPCGEdge::GetOtherPin(const UPCGPin* Pin) const
 	check(Pin == InputPin || Pin == OutputPin);
 	return Pin == InputPin ? OutputPin : InputPin;
 }
+
+const UPCGNode* UPCGEdge::GetInputNode() const
+{
+	return InputPin ? InputPin->Node : nullptr;
+}
+
+const UPCGNode* UPCGEdge::GetOutputNode() const
+{
+	return OutputPin ? OutputPin->Node : nullptr;
+}
+
+FName UPCGEdge::GetInputPinLabel() const
+{
+	return InputPin ? InputPin->Properties.Label : NAME_None;
+}
+
+FName UPCGEdge::GetOutputPinLabel() const
+{
+	return OutputPin ? OutputPin->Properties.Label : NAME_None;
+}
diff --git a/Engine/Plugins/PCG/Source/PCG/Private/PCGGraph.cpp b/Engine/Plugins/PCG/Source/PCG/Private/PCGGraph.cpp
--- a/Engine/Plugins/PCG/Source/PCG/Private/PCGGraph.cpp
+++ b/Engine/Plugins/PCG/Source/PCG/Private/PCGGraph.cpp
@@ -1152,6 +1152,27 @@ bool UPCGGraph::AddLabeledEdge(UPCGNode* From, const FName& FromPinLabel, UPCGNo
 	return bToPinBrokeOtherEdges;
 }
 
+TArray<UPCGEdge*> UPCGGraph::GetAllEdges() const
+{
+	TArray<UPCGEdge*> AllEdges;
+
+	auto GatherInputEdgesFromNode = [&AllEdges](const UPCGNode* Node) 
+	{
+		for (const UPCGPin* InPin : Node->InputPins)
+		{
+			AllEdges.Append(InPin->Edges);
+		}
+	};
+
+	// Since we gather edge from their Output side (The input side of the nodes) we have to skip InputNode
+	GatherInputEdgesFromNode(OutputNode);
+	for (const UPCGNode* Node : Nodes)
+	{
+		GatherInputEdgesFromNode(Node);
+	}
+	return AllEdges;
+}
+
 TObjectPtr<UPCGNode> UPCGGraph::ReconstructNewNode(const UPCGNode* InNode)
 {
 	UPCGSettings* NewSettings = nullptr;
diff --git a/Engine/Plugins/PCG/Source/PCG/Public/PCGEdge.h b/Engine/Plugins/PCG/Source/PCG/Public/PCGEdge.h
--- a/Engine/Plugins/PCG/Source/PCG/Public/PCGEdge.h
+++ b/Engine/Plugins/PCG/Source/PCG/Public/PCGEdge.h
@@ -34,16 +34,29 @@ public:
 	TObjectPtr<UPCGNode> OutboundNode_DEPRECATED;
 
 	/** Pin at upstream end of edge. */
-	UPROPERTY()
+	UPROPERTY(BlueprintReadOnly, Category = Properties)
 	TObjectPtr<UPCGPin> InputPin;
 
 	/** Pin at downstream end of edge. */
-	UPROPERTY()
+	UPROPERTY(BlueprintReadOnly, Category = Properties)
 	TObjectPtr<UPCGPin> OutputPin;
 
 	UE_API bool IsValid() const;
 	UE_API UPCGPin* GetOtherPin(const UPCGPin* Pin);
 	UE_API const UPCGPin* GetOtherPin(const UPCGPin* Pin) const;
+
+	UFUNCTION(BlueprintCallable, Category = Properties)
+	const UPCGNode* GetInputNode() const;
+
+	UFUNCTION(BlueprintCallable, Category = Properties)
+	const UPCGNode* GetOutputNode() const;
+
+	UFUNCTION(BlueprintCallable, Category = Properties)
+	FName GetInputPinLabel() const;
+
+	UFUNCTION(BlueprintCallable, Category = Properties)
+	FName GetOutputPinLabel() const;
+
 };
 
 #undef UE_API
diff --git a/Engine/Plugins/PCG/Source/PCG/Public/PCGGraph.h b/Engine/Plugins/PCG/Source/PCG/Public/PCGGraph.h
--- a/Engine/Plugins/PCG/Source/PCG/Public/PCGGraph.h
+++ b/Engine/Plugins/PCG/Source/PCG/Public/PCGGraph.h
@@ -386,6 +386,10 @@ public:
 	UFUNCTION(BlueprintCallable, Category = Graph)
 	UPCGNode* GetOutputNode() const { return OutputNode; }
 
+	/** Returns all the edges of the graph */
+	UFUNCTION(BlueprintCallable, Category = Graph)
+	TArray<UPCGEdge*> GetAllEdges() const;
+
 	/** Duplicate a given node by creating a new node with the same settings and properties, but without any edges and add it to the graph */
 	UE_API TObjectPtr<UPCGNode> ReconstructNewNode(const UPCGNode* InNode);
 
diff --git a/Engine/Plugins/PCG/Source/PCG/Public/PCGSubgraph.h b/Engine/Plugins/PCG/Source/PCG/Public/PCGSubgraph.h
--- a/Engine/Plugins/PCG/Source/PCG/Public/PCGSubgraph.h
+++ b/Engine/Plugins/PCG/Source/PCG/Public/PCGSubgraph.h
@@ -33,6 +33,7 @@ public:
 	virtual bool IsDynamicGraph() const { return false; }
 
 	// Use this method from the outside to set the subgraph, as it will connect editor callbacks
+	UFUNCTION(BlueprintCallable, Category = Graph)
 	UE_API virtual void SetSubgraph(UPCGGraphInterface* InGraph);
 
 protected:
