[PCG] New Pcg graph change for externally modified graph

#rnx

diff --git a/Engine/Plugins/PCG/Source/PCG/Public/PCGCommon.h b/Engine/Plugins/PCG/Source/PCG/Public/PCGCommon.h
--- a/Engine/Plugins/PCG/Source/PCG/Public/PCGCommon.h
+++ b/Engine/Plugins/PCG/Source/PCG/Public/PCGCommon.h
@@ -184,6 +184,8 @@ enum class EPCGChangeType : uint32
 	ShaderSource = 1 << 7,
 	/** Changes in the graph customization that will impact the editor. */
 	GraphCustomization = 1 << 8,
+	/** Changes in the graph caused by external tools, editor graph should be reconstructed. */
+	ExternalModification = 1 << 9,
 };
 ENUM_CLASS_FLAGS(EPCGChangeType);
 
diff --git a/Engine/Plugins/PCG/Source/PCGEditor/Private/PCGEditor.cpp b/Engine/Plugins/PCG/Source/PCGEditor/Private/PCGEditor.cpp
--- a/Engine/Plugins/PCG/Source/PCGEditor/Private/PCGEditor.cpp
+++ b/Engine/Plugins/PCG/Source/PCGEditor/Private/PCGEditor.cpp
@@ -3288,6 +3288,13 @@ void FPCGEditor::OnGraphChanged(UPCGGraphInterface* InGraph, EPCGChangeType Chan
 			}
 		}
 	}
+
+	if (!!(ChangeType & EPCGChangeType::ExternalModification))
+	{
+	    //Changes caused by external tools: Editor graph should be reconstructed.
+		PCGEditorGraph->ReconstructGraph();
+	}
+
 }
 
 void FPCGEditor::OnNodeSourceCompiled(const UPCGNode* InNode, const FPCGCompilerDiagnostics& InDiagnostics)
