Fix warning V623: Consider inspecting the '?:' operator. A temporary object of the 'TRefCountPtr' type is being created and subsequently destroyed. Check third operand.
#rnx

diff --git a/Engine/Plugins/PCG/Source/PCG/Private/Elements/PCGCopyToRenderTarget.cpp b/Engine/Plugins/PCG/Source/PCG/Private/Elements/PCGCopyToRenderTarget.cpp
--- a/Engine/Plugins/PCG/Source/PCG/Private/Elements/PCGCopyToRenderTarget.cpp
+++ b/Engine/Plugins/PCG/Source/PCG/Private/Elements/PCGCopyToRenderTarget.cpp
@@ -205,7 +205,7 @@ bool FPCGCopyToRenderTargetElement::ExecuteInternal(FPCGContext* InContext) cons
 					}
 
 					FTextureRenderTargetResource* RenderTargetResource = RenderTarget ? RenderTarget->GetRenderTargetResource() : nullptr;
-					FRHITexture* TargetTexture = RenderTargetResource ? RenderTargetResource->GetTextureRHI() : nullptr;
+					FRHITexture* TargetTexture = RenderTargetResource ? RenderTargetResource->GetTextureRHI().GetReference() : nullptr;
 
 					if (!TargetTexture)
 					{
