[PCG] Fix crash where string keys values improperly used from an upstream compute graph without a remap.

Disables the raw buffer optimizaiton for by-attribute spawning until we can rework the way data is passed for this DI and get the upstream string table.

#rnx

diff --git a/Engine/Plugins/PCG/Source/PCG/Private/Compute/DataInterfaces/Elements/PCGStaticMeshSpawnerDataInterface.cpp b/Engine/Plugins/PCG/Source/PCG/Private/Compute/DataInterfaces/Elements/PCGStaticMeshSpawnerDataInterface.cpp
--- a/Engine/Plugins/PCG/Source/PCG/Private/Compute/DataInterfaces/Elements/PCGStaticMeshSpawnerDataInterface.cpp
+++ b/Engine/Plugins/PCG/Source/PCG/Private/Compute/DataInterfaces/Elements/PCGStaticMeshSpawnerDataInterface.cpp
@@ -314,7 +314,9 @@ bool UPCGStaticMeshSpawnerDataProvider::PerformPreExecuteReadbacks_GameThread(UP
 	}
 	else if (const UPCGRawBufferData* RawAnalysisData = Cast<UPCGRawBufferData>(InBinding->GetInputDataCollection().TaggedData[AnalysisDataIndex].Data))
 	{
-		// Hardcoded data format - uint array containing pairs of (string key value, instance count) pairs.
+		// Hardcoded data format - uint array containing pairs of (primitive ID, instance count) pairs.
+		// This path is not supported for by-attribute spawning, for which the primitive ID is the string key value, because
+		// we can't currently remap string keys from upstream compute graph.
 		const int32 NumUints = RawAnalysisData->GetNumUint32s();
 		ensure((NumUints % 2) == 0);
 
@@ -463,6 +465,23 @@ bool UPCGStaticMeshSpawnerDataProvider::PrepareForExecute_GameThread(UPCGDataBin
 			}
 		}
 
+		bool bAnyInvalid = false;
+
+		// User could pass us any soft object path so verify static meshes were loaded.
+		for (const FPCGProceduralISMComponentDescriptor& Descriptor : PrimitiveDescriptors)
+		{
+			if (!Cast<UStaticMesh>(Descriptor.StaticMesh.Get()))
+			{
+				UE_LOG(LogPCG, Error, TEXT("Tried to use asset '%s' as a static mesh."), *Descriptor.StaticMesh.ToString());
+				bAnyInvalid = true;
+			}
+		}
+
+		if (bAnyInvalid)
+		{
+			PrimitiveDescriptors.Reset();
+		}
+
 		bStaticMeshesLoaded = true;
 	}
 
diff --git a/Engine/Plugins/PCG/Source/PCG/Private/Elements/PCGStaticMeshSpawner.cpp b/Engine/Plugins/PCG/Source/PCG/Private/Elements/PCGStaticMeshSpawner.cpp
--- a/Engine/Plugins/PCG/Source/PCG/Private/Elements/PCGStaticMeshSpawner.cpp
+++ b/Engine/Plugins/PCG/Source/PCG/Private/Elements/PCGStaticMeshSpawner.cpp
@@ -63,7 +63,11 @@ UPCGStaticMeshSpawnerSettings::UPCGStaticMeshSpawnerSettings(const FObjectInitia
 #if WITH_EDITOR
 void UPCGStaticMeshSpawnerSettings::CreateKernels(FPCGGPUCompilationContext& InOutContext, UObject* InObjectOuter, TArray<UPCGComputeKernel*>& OutKernels, TArray<FPCGKernelEdge>& OutEdges) const
 {
-	const bool bUseRawBuffer = CVarUseRawDataForAnalysis.GetValueOnAnyThread();
+	const UPCGMeshSelectorPrimitiveData* SelectorPrimitiveData = Cast<UPCGMeshSelectorPrimitiveData>(MeshSelectorParameters);
+
+	// Only the by-primitive-data selector supports the raw buffer optimization, for which the analysis data contains table indices.
+	// For by-attribute spawning the analysis data contains upstream string keys which we can't currently remap to the downstream graph string keys.
+	const bool bUseRawBuffer = SelectorPrimitiveData && CVarUseRawDataForAnalysis.GetValueOnAnyThread();
 
 	UPCGStaticMeshSpawnerKernel* SpawnerKernel = nullptr;
 
@@ -83,7 +87,6 @@ void UPCGStaticMeshSpawnerSettings::CreateKernels(FPCGGPUCompilationContext& InO
 
 	// If doing by-attribute selection, add analysis kernel that will count how many instances of each mesh are present.
 	const UPCGMeshSelectorByAttribute* SelectorByAttribute = Cast<UPCGMeshSelectorByAttribute>(MeshSelectorParameters);
-	const UPCGMeshSelectorPrimitiveData* SelectorPrimitiveData = Cast<UPCGMeshSelectorPrimitiveData>(MeshSelectorParameters);
 	if (SelectorByAttribute || SelectorPrimitiveData)
 	{
 		PCGKernelHelpers::FCreateKernelParams CreateParams(InObjectOuter, this);
