[PCG] FPCGGraphExecutor: optims
- Group new ready tasks processing to reduce locking (QueueNextTasks ~125k calls 8.7s to 6.1s
- Move cache param build outside of retrieve so we calculate crcs outside the lock (PrepareForExecute ~80k calls 2.5s to 1.9s)

diff --git a/Engine/Plugins/PCG/Source/PCG/Private/Graph/PCGGraphExecutor.cpp b/Engine/Plugins/PCG/Source/PCG/Private/Graph/PCGGraphExecutor.cpp
--- a/Engine/Plugins/PCG/Source/PCG/Private/Graph/PCGGraphExecutor.cpp
+++ b/Engine/Plugins/PCG/Source/PCG/Private/Graph/PCGGraphExecutor.cpp
@@ -860,7 +860,7 @@ bool FPCGGraphExecutor::ProcessScheduledTasks(double EndTime)
 	TRACE_CPUPROFILER_EVENT_SCOPE(FPCGGraphExecutor::ProcessScheduledTasks);
 	bool bProcessedNewTasks = false;
 	const bool bIsInGameThread = IsInGameThread();
-	TArray<FCachedResult*> CachedResults;
+	FCachedResults CachedResults;
 
 	{
 		// Process any newly scheduled graphs to execute
@@ -920,10 +920,10 @@ bool FPCGGraphExecutor::ProcessScheduledTasks(double EndTime)
 					}
 
 					// Process ready tasks
-					for (FPCGGraphTask& ReadyTask : NewReadyTasks)
+					if (NewReadyTasks.Num() > 0)
 					{
-						OnTaskInputsReady(ReadyTask, CachedResults, bIsInGameThread);
-						--NewReadyTaskCount;
+						OnTaskInputsReady(NewReadyTasks, CachedResults, bIsInGameThread);
+						NewReadyTaskCount -= NewReadyTasks.Num();
 					}
 
 					if (FPlatformTime::Seconds() > EndTime)
@@ -950,7 +950,7 @@ bool FPCGGraphExecutor::ProcessScheduledTasks(double EndTime)
 		}
 	}
 
-	ProcessCachedResults(MoveTemp(CachedResults));
+	ProcessCachedResults(CachedResults);
 
 	CheckState(bIsInGameThread);
 
@@ -1202,10 +1202,11 @@ void FPCGGraphExecutor::PostTaskExecute(TSharedPtr<FPCGGraphActiveTask> ActiveTa
 		}
 
 		// Store output in data map.
-		StoreResults(ActiveTask.NodeId, ActiveTask.Context->OutputData, bNeedsManualClear);
+		FTaskResult NewResult{ ActiveTask.NodeId, &ActiveTask.Context->OutputData, bNeedsManualClear };
+		StoreResults(MakeArrayView(&NewResult, 1));
 
 		// Book-keeping
-		QueueNextTasks(ActiveTask.NodeId);
+		QueueNextTasks(MakeArrayView(&ActiveTask.NodeId, 1));
 	}
 
 	// Erase from ActiveTasks
@@ -1229,38 +1230,56 @@ void FPCGGraphExecutor::PostTaskExecute(TSharedPtr<FPCGGraphActiveTask> ActiveTa
 	}
 }
 
-void FPCGGraphExecutor::OnTaskInputsReady(FPCGGraphTask& Task, TArray<FCachedResult*>& OutCachedResults, bool bIsInGameThread)
+void FPCGGraphExecutor::OnTaskInputsReady(TArrayView<FPCGGraphTask> InNewReadyTasks, FCachedResults& OutCachedResults, bool bIsInGameThread)
 {
 	TRACE_CPUPROFILER_EVENT_SCOPE(FPCGGraphExecutor::OnTaskInputsReady);
-	TArray<FPCGTaskId> ResultsToMarkAsRead;
-	if (SetupTask(Task, ResultsToMarkAsRead))
+	FTaskIds ResultsToMarkAsRead;
+	TArray<FPCGGraphTask, TInlineAllocator<16>> UncachedReadyTasks;
+
+	for (FPCGGraphTask& Task : InNewReadyTasks)
 	{
-		check(Task.Element);
+		// LocalResultsToMarkAsRead does not contain duplicates
+		FTaskIds LocalResultsToMarkAsRead;
+		if (SetupTask(Task, LocalResultsToMarkAsRead))
+		{
+			check(Task.Element);
 
-		FCachedResult* CachedResult = nullptr;
+			FCachedResult* CachedResult = nullptr;
 
-		// Query Cache if we can here, if not the Task will be pushed to ReadyTasks and PrepareForExecute will be called on MainThread if needed
-		if (bIsInGameThread || !Task.CanExecuteOnlyOnMainThread())
-		{
-			PrepareForExecute(Task, CachedResult, /*bLiveTasksLockAlreadyLocked=*/false);
+			// Query Cache if we can here, if not the Task will be pushed to ReadyTasks and PrepareForExecute will be called on MainThread if needed
+			if (bIsInGameThread || !Task.CanExecuteOnlyOnMainThread())
+			{
+				PrepareForExecute(Task, CachedResult, /*bLiveTasksLockAlreadyLocked=*/false);
+			}
+
+			if (CachedResult)
+			{
+				check(CachedResult->TaskId != InvalidPCGTaskId);
+				OutCachedResults.Add(CachedResult);
+			}
+			else
+			{
+				UncachedReadyTasks.Emplace(MoveTemp(Task));
+			}
 		}
 
+		// ResultsToMarkAsRead can contain duplicates as we want to mark the same input as read multiple times if multiple tasks depend on it
+		ResultsToMarkAsRead.Append(LocalResultsToMarkAsRead);
+	}
+
+	if (UncachedReadyTasks.Num() > 0)
+	{
 		PCG::TScopeLock ChildScopeLock(LiveTasksLock);
-		if (CachedResult)
-		{
-			check(CachedResult->TaskId != InvalidPCGTaskId);
-			OutCachedResults.Add(CachedResult);
-		}
-		else 
-		{
-			ReadyTasks.Emplace(MoveTemp(Task));
-		}
+		ReadyTasks.Append(MoveTemp(UncachedReadyTasks));
 	}
 
-	MarkInputResults(ResultsToMarkAsRead);
+	if (ResultsToMarkAsRead.Num() > 0)
+	{
+		MarkInputResults(ResultsToMarkAsRead);
+	}
 }
 
-bool FPCGGraphExecutor::SetupTask(FPCGGraphTask& Task, TArray<FPCGTaskId>& ResultsToMarkAsRead)
+bool FPCGGraphExecutor::SetupTask(FPCGGraphTask& Task, FTaskIds& ResultsToMarkAsRead)
 {
 	TRACE_CPUPROFILER_EVENT_SCOPE(FPCGGraphExecutor::SetupTask);
 	if (Task.bHasDoneSetup)
@@ -1340,24 +1359,33 @@ void FPCGGraphExecutor::PrepareForExecute(FPCGGraphTask& Task, FCachedResult*& O
 	const UPCGSettings* TaskSettings = TaskSettingsInterface ? TaskSettingsInterface->GetSettings() : nullptr;
 
 	// Calculate Crc of dependencies (input data Crcs, settings) and use this as the key in the cache lookup
-	FPCGCrc DependenciesCrc;
+	FPCGGetFromCacheParams CacheParam;
 	EPCGCachingStatus CacheStatus = EPCGCachingStatus::NotInCache;
 	bool bResultAlreadyInCache = false;
 
-	// This section requires the lock on the caching results to prevent interaction between different tasks and the GC.
+	// This section requires the lock on the caching results to prevent interaction between different tasks and the GC so we avoid locking if cache is disabled.
+	if(GraphCache.IsEnabled())
 	{
-		PCG::TScopeLock ScopeLock(CachingResultsLock);
-		FCachedResult LocalCachedResult;
-		CacheStatus = Task.Element->RetrieveResultsFromCache(&GraphCache, Task.Node, Task.TaskInput, Task.ExecutionSource.Get(), LocalCachedResult.Output, &DependenciesCrc);
-
-		bResultAlreadyInCache = (CacheStatus == EPCGCachingStatus::Cached);
-		if (bResultAlreadyInCache)
+		// Build cache param outside of lock
+		if (Task.Element->BuildCacheParams(&GraphCache, Task.Node, Task.TaskInput, Task.ExecutionSource.Get(), CacheParam))
 		{
-			check(!CachingResultsForGC.Contains(Task.NodeId));
-			TUniquePtr<FCachedResult>& CachingResultPtr = CachingResultsForGC.Add(Task.NodeId, MakeUnique<FCachedResult>());
-			OutCachedResult = CachingResultPtr.Get();
-			OutCachedResult->TaskId = Task.NodeId;
-			OutCachedResult->Output = MoveTemp(LocalCachedResult.Output);
+			// Check that we have a valid cache param to avoid locking for nothing
+			if (CacheParam.Crc.IsValid())
+			{
+				PCG::TScopeLock ScopeLock(CachingResultsLock);
+				FCachedResult LocalCachedResult;
+				CacheStatus = Task.Element->RetrieveResultsFromCache(&GraphCache, CacheParam, LocalCachedResult.Output);
+
+				bResultAlreadyInCache = (CacheStatus == EPCGCachingStatus::Cached);
+				if (bResultAlreadyInCache)
+				{
+					check(!CachingResultsForGC.Contains(Task.NodeId));
+					TUniquePtr<FCachedResult>& CachingResultPtr = CachingResultsForGC.Add(Task.NodeId, MakeUnique<FCachedResult>());
+					OutCachedResult = CachingResultPtr.Get();
+					OutCachedResult->TaskId = Task.NodeId;
+					OutCachedResult->Output = MoveTemp(LocalCachedResult.Output);
+				}
+			}
 		}
 	}
 
@@ -1418,7 +1446,7 @@ void FPCGGraphExecutor::PrepareForExecute(FPCGGraphTask& Task, FCachedResult*& O
 		
 		Task.Context->TaskId = Task.NodeId;
 		Task.Context->CompiledTaskId = Task.CompiledTaskId;
-		Task.Context->DependenciesCrc = DependenciesCrc;
+		Task.Context->DependenciesCrc = CacheParam.Crc;
 		Task.Context->StackHandle = Task.GetStackHandle();
 		PRAGMA_DISABLE_DEPRECATION_WARNINGS
 		Task.Context->Stack = Task.GetStack();
@@ -1510,7 +1538,7 @@ bool FPCGGraphExecutor::ExecuteScheduling(double EndTime, TSharedPtr<FPCGGraphAc
 	// Number of threads to use for PCG Async operations 
 	const int32 MaxPCGAsyncThreads = FMath::Min(MaxNumThreads - MaxExecutingThreads, CVarMaxNumTasks.GetValueOnAnyThread());
 
-	TArray<FCachedResult*> CachedResults;
+	FCachedResults CachedResults;
 	{
 		TRACE_CPUPROFILER_EVENT_SCOPE(FPCGGraphExecutor::ExecuteScheduling::ReadyTasks);
 		const bool bTryLock = (OutMainThreadTask == nullptr);
@@ -1615,7 +1643,7 @@ bool FPCGGraphExecutor::ExecuteScheduling(double EndTime, TSharedPtr<FPCGGraphAc
 		}
 	}
 
-	ProcessCachedResults(MoveTemp(CachedResults));
+	ProcessCachedResults(CachedResults);
 
 	// Dispatch Tasks
 	{
@@ -1908,102 +1936,115 @@ void FPCGGraphExecutor::ClearAllTasks()
 	PausedTaskSuccessors.Reset();
 }
 
-TArray<FPCGGraphExecutor::FCachedResult*> FPCGGraphExecutor::QueueNextTasksInternal(FPCGTaskId FinishedTask)
+void FPCGGraphExecutor::QueueNextTasksInternal(TConstArrayView<FPCGTaskId> FinishedTasks, FCachedResults& OutCachedResults)
 {
 	TRACE_CPUPROFILER_EVENT_SCOPE(FPCGGraphExecutor::QueueNextTasks);
-	TArray<FCachedResult*> CachedResults;
 	const bool bIsInGameThread = IsInGameThread();
 
-	
 	TSet<FPCGTaskId> Successors;
 	TArray<FPCGGraphTask, TInlineAllocator<16>> NewReadyTasks;
 	{
 		PCG::TScopeLock ScopeLock(TasksLock);
 		PCG::TSharedScopeLock ReadLock(TaskOutputsLock);
-
-		if (TaskSuccessors.RemoveAndCopyValue(FinishedTask, Successors))
+		for (FPCGTaskId FinishedTask : FinishedTasks)
 		{
-			for (FPCGTaskId Successor : Successors)
+			if (TaskSuccessors.RemoveAndCopyValue(FinishedTask, Successors))
 			{
-				bool bAllPrerequisitesMet = true;
-
-				// When running multithreaded it is possible for the Successor to already have been moved to ReadyTask after the current task stored its results
-				if (FPCGGraphTask* SuccessorTaskPtr = Tasks.Find(Successor))
+				for (FPCGTaskId Successor : Successors)
 				{
-					FPCGGraphTask& SuccessorTask = *SuccessorTaskPtr;
+					bool bAllPrerequisitesMet = true;
 
-					if (SuccessorTask.Inputs.Num() > 0)
+					// When running multithreaded it is possible for the Successor to already have been moved to ReadyTask after the current task stored its results
+					if (FPCGGraphTask* SuccessorTaskPtr = Tasks.Find(Successor))
 					{
-						for (const FPCGGraphTaskInput& Input : SuccessorTask.Inputs)
+						FPCGGraphTask& SuccessorTask = *SuccessorTaskPtr;
+
+						if (SuccessorTask.Inputs.Num() > 0)
 						{
-							bAllPrerequisitesMet &= TaskOutputs.Contains(Input.TaskId);
+							for (const FPCGGraphTaskInput& Input : SuccessorTask.Inputs)
+							{
+								bAllPrerequisitesMet &= TaskOutputs.Contains(Input.TaskId);
+							}
 						}
-					}
 
-					if (bAllPrerequisitesMet)
-					{
-						NewReadyTasks.Add(MoveTemp(SuccessorTask));
-						++NewReadyTaskCount;
-						Tasks.Remove(Successor);
+						if (bAllPrerequisitesMet)
+						{
+							NewReadyTasks.Add(MoveTemp(SuccessorTask));
+							++NewReadyTaskCount;
+							Tasks.Remove(Successor);
+						}
 					}
 				}
 			}
 		}
 	}
 
-	for (FPCGGraphTask& NewReadyTask : NewReadyTasks)
+	if (NewReadyTasks.Num() > 0)
 	{
-		OnTaskInputsReady(NewReadyTask, CachedResults, bIsInGameThread);
-		--NewReadyTaskCount;
+		OnTaskInputsReady(NewReadyTasks, OutCachedResults, bIsInGameThread);
+		NewReadyTaskCount -= NewReadyTasks.Num();
 	}
 
 	{
 		PCG::TScopeLock PauseScopeLock(PausedTaskSuccessorsLock);
-
-		TArray<TSharedPtr<FPCGGraphActiveTask>> ActiveTaskSuccessors;
-
-		if (PausedTaskSuccessors.RemoveAndCopyValue(FinishedTask, ActiveTaskSuccessors))
+		for (FPCGTaskId FinishedTask : FinishedTasks)
 		{
-			for (const TSharedPtr<FPCGGraphActiveTask>& ActiveTask : ActiveTaskSuccessors)
+			TArray<TSharedPtr<FPCGGraphActiveTask>> ActiveTaskSuccessors;
+
+			if (PausedTaskSuccessors.RemoveAndCopyValue(FinishedTask, ActiveTaskSuccessors))
 			{
-				ActiveTask->Context->DynamicDependencies.Remove(FinishedTask);
-				if (ActiveTask->Context->DynamicDependencies.IsEmpty())
+				for (const TSharedPtr<FPCGGraphActiveTask>& ActiveTask : ActiveTaskSuccessors)
 				{
-					ActiveTask->Context->bIsPaused = false;
+					ActiveTask->Context->DynamicDependencies.Remove(FinishedTask);
+					if (ActiveTask->Context->DynamicDependencies.IsEmpty())
+					{
+						ActiveTask->Context->bIsPaused = false;
+					}
 				}
 			}
 		}
 	}
-	
-	return CachedResults;
 }
 
-void FPCGGraphExecutor::QueueNextTasks(FPCGTaskId FinishedTask)
+void FPCGGraphExecutor::QueueNextTasks(TConstArrayView<FPCGTaskId> FinishedTasks)
 {
-	TArray<FCachedResult*> CachedResults = QueueNextTasksInternal(FinishedTask);
+	FCachedResults CachedResults;
+	QueueNextTasksInternal(FinishedTasks, CachedResults);
 
-	ProcessCachedResults(MoveTemp(CachedResults));
+	if (CachedResults.Num() > 0)
+	{
+		ProcessCachedResults(CachedResults);
+	}
 }
 
-void FPCGGraphExecutor::ProcessCachedResults(TArray<FCachedResult*> CachedResults)
+void FPCGGraphExecutor::ProcessCachedResults(TConstArrayView<FCachedResult*> CachedResults)
 {
-	while (CachedResults.Num() > 0)
+	if (CachedResults.IsEmpty())
 	{
-		TArray<FPCGTaskId> NextTasks = ProcessCachedResultsInternal(MoveTemp(CachedResults));
-		check(CachedResults.IsEmpty());
+		return;
+	}
 
-		for (const FPCGTaskId& NextTask : NextTasks)
-		{
-			CachedResults.Append(QueueNextTasksInternal(NextTask));
-		}
+	FTaskIds NextTasks;
+	ProcessCachedResultsInternal(CachedResults, NextTasks);
+	
+	FCachedResults NewCachedResults;
+	QueueNextTasksInternal(NextTasks, NewCachedResults);
+
+	while (NewCachedResults.Num() > 0)
+	{
+		NextTasks.Reset();
+		ProcessCachedResultsInternal(NewCachedResults, NextTasks);
+		
+		NewCachedResults.Reset();
+		QueueNextTasksInternal(NextTasks, NewCachedResults);
 	}
 }
 
-TArray<FPCGTaskId> FPCGGraphExecutor::ProcessCachedResultsInternal(TArray<FCachedResult*> CachedResults)
+void FPCGGraphExecutor::ProcessCachedResultsInternal(TConstArrayView<FCachedResult*> CachedResults, FTaskIds& OutNextTasks)
 {
 	TRACE_CPUPROFILER_EVENT_SCOPE(FPCGGraphExecutor::ProcessCachedResults);
-	TArray<FPCGTaskId> NextTasks;
-	NextTasks.Reserve(CachedResults.Num());
+	
+	TArray<FTaskResult, TInlineAllocator<16>> NewResults;
 
 	for (FCachedResult* CachedResult : CachedResults)
 	{
@@ -2016,10 +2057,13 @@ TArray<FPCGTaskId> FPCGGraphExecutor::ProcessCachedResultsInternal(TArray<FCache
 #endif
 		}
 
-		StoreResults(CachedResult->TaskId, CachedResult->Output, CachedResult->bIsPostGraphTask);
-		NextTasks.Add(CachedResult->TaskId);
+		
+		NewResults.Add({ CachedResult->TaskId, &CachedResult->Output, CachedResult->bIsPostGraphTask });
+		OutNextTasks.Add(CachedResult->TaskId);
 	}
 
+	StoreResults(NewResults);
+
 	if (CachedResults.Num() > 0)
 	{
 		PCG::TScopeLock ScopeLock(CachingResultsLock);
@@ -2032,8 +2076,6 @@ TArray<FPCGTaskId> FPCGGraphExecutor::ProcessCachedResultsInternal(TArray<FCache
 		// Next Scheduling call needs to check if this task completion unblocked some paused task(s)
 		bNeedToCheckPausedTasks = true;
 	}
-
-	return NextTasks;
 }
 
 bool FPCGGraphExecutor::CancelNextTasks(FPCGTaskId CancelledTask, TSet<IPCGGraphExecutionSource*>& OutCancelledExecutionSources)
@@ -2134,7 +2176,7 @@ void FPCGGraphExecutor::RemoveTaskFromInputSuccessorsNoLock(FPCGTaskId Cancelled
 	}
 }
 
-void FPCGGraphExecutor::BuildTaskInput(FPCGGraphTask& Task, TArray<FPCGTaskId>& ResultsToMarkAsRead)
+void FPCGGraphExecutor::BuildTaskInput(FPCGGraphTask& Task, FTaskIds& ResultsToMarkAsRead)
 {
 	TRACE_CPUPROFILER_EVENT_SCOPE(FPCGGraphExecutor::BuildTaskInput);
 
@@ -2340,28 +2382,39 @@ void FPCGGraphExecutor::CombineParams(FPCGGraphTask& Task)
 	}
 }
 
-void FPCGGraphExecutor::StoreResults(FPCGTaskId InTaskId, const FPCGDataCollection& InTaskOutput, bool bNeedsManualClear)
+void FPCGGraphExecutor::StoreResults(TConstArrayView<FTaskResult> InResults)
 {
 	TRACE_CPUPROFILER_EVENT_SCOPE(FPCGGraphExecutor::StoreResults);
 
-	FOutputDataInfo OutputDataInfo;
-	OutputDataInfo.DataCollection = InTaskOutput;
-	OutputDataInfo.bNeedsManualClear = bNeedsManualClear;
-
-	OutputDataInfo.DataCollection.MarkUsage(EPCGDataUsage::GraphExecutorTaskOutput);
+	TArray<TPair<FPCGTaskId, FOutputDataInfo>, TInlineAllocator<16>> NewOutputs;
 
 	{
 		TRACE_CPUPROFILER_EVENT_SCOPE(FPCGGraphExecutor::StoreResults::Successors);
 		PCG::TScopeLock ScopeLock(TasksLock);
-		if (TSet<FPCGTaskId>* Successors = TaskSuccessors.Find(InTaskId))
+		for (const FTaskResult& Result : InResults)
 		{
-			OutputDataInfo.RemainingSuccessorCount = Successors->Num();
+			FOutputDataInfo OutputDataInfo;
+			check(Result.TaskOutput);
+			OutputDataInfo.DataCollection = *Result.TaskOutput;
+			OutputDataInfo.bNeedsManualClear = Result.bNeedsManualClear;
+			OutputDataInfo.DataCollection.MarkUsage(EPCGDataUsage::GraphExecutorTaskOutput);
+
+			if (TSet<FPCGTaskId>* Successors = TaskSuccessors.Find(Result.TaskId))
+			{
+				OutputDataInfo.RemainingSuccessorCount = Successors->Num();
+			}
+
+			NewOutputs.Add({ Result.TaskId, MoveTemp(OutputDataInfo) });
 		}
 	}
 
 	// Store output in map
 	PCG::TUniqueScopeLock WriteLock(TaskOutputsLock);
-	TaskOutputs.Add(InTaskId, MoveTemp(OutputDataInfo));
+
+	for (auto& Pair : NewOutputs)
+	{
+		TaskOutputs.Add(Pair.Key, MoveTemp(Pair.Value));
+	}
 }
 
 void FPCGGraphExecutor::ClearResults()
@@ -2424,7 +2477,7 @@ void FPCGGraphExecutor::CullInactiveDownstreamNodes(FPCGTaskId InCompletedTaskId
 	TSet<FPCGTaskId> AllRemovedTasks;
 
 	// Hoisted out of loop for performance reasons.
-	TArray<FPCGTaskId, TInlineAllocator<64>> TasksToRemove;
+	FTaskIds TasksToRemove;
 
 	if (PinIdsToDeactivate.Num() > 0)
 	{
@@ -2538,10 +2591,7 @@ void FPCGGraphExecutor::CullInactiveDownstreamNodes(FPCGTaskId InCompletedTaskId
 		// Queue downstream tasks in a similar manner to when a task draws from the cache and is skipped.
 		// Some downstream tasks will have been culled which we don't care about (hence the ignore flag),
 		// but some may not be queued and may be ready for queuing.
-		for (const FPCGTaskId TaskId : AllRemovedTasks)
-		{
-			QueueNextTasks(TaskId);
-		}
+		QueueNextTasks(AllRemovedTasks.Array());
 	}
 }
 
diff --git a/Engine/Plugins/PCG/Source/PCG/Private/Graph/PCGGraphExecutor.h b/Engine/Plugins/PCG/Source/PCG/Private/Graph/PCGGraphExecutor.h
--- a/Engine/Plugins/PCG/Source/PCG/Private/Graph/PCGGraphExecutor.h
+++ b/Engine/Plugins/PCG/Source/PCG/Private/Graph/PCGGraphExecutor.h
@@ -222,19 +222,23 @@ private:
 
 	TSet<IPCGGraphExecutionSource*> Cancel(TFunctionRef<bool(TWeakInterfacePtr<IPCGGraphExecutionSource>)> CancelFilter);
 	void ClearAllTasks();
-	void QueueNextTasks(FPCGTaskId FinishedTask);
-	TArray<FCachedResult*> QueueNextTasksInternal(FPCGTaskId FinishedTask);
+	
+	using FCachedResults = TArray<FCachedResult*, TInlineAllocator<16>>;
+	using FTaskIds = TArray<FPCGTaskId, TInlineAllocator<32>>;
+	
+	void QueueNextTasks(TConstArrayView<FPCGTaskId> FinishedTasks);
+	void QueueNextTasksInternal(TConstArrayView<FPCGTaskId> FinishedTasks, FCachedResults& OutCachedResults);
 
 	bool CancelNextTasks(FPCGTaskId CancelledTask, TSet<IPCGGraphExecutionSource*>& OutCancelledExecutionSources);
 	void RemoveTaskFromInputSuccessors(FPCGTaskId CancelledTask, const TArray<FPCGGraphTaskInput>& CancelledTaskInputs);
 	void RemoveTaskFromInputSuccessorsNoLock(FPCGTaskId CancelledTask, const TArray<FPCGGraphTaskInput>& CancelledTaskInputs);
 	
 	/** Called from QueueNextTasks/ProcessScheduledTasks will try and setup/prepare task for execution */
-	void OnTaskInputsReady(FPCGGraphTask& Task, TArray<FCachedResult*>& OutCachedResults, bool bIsInGameThread);
+	void OnTaskInputsReady(TArrayView<FPCGGraphTask> InNewReadyTasks, FCachedResults& OutCachedResults, bool bIsInGameThread);
 
 	/** SetupTask will call BuildTaskInput and assign a IPCGElement to the FPCGGraphTask */
-	bool SetupTask(FPCGGraphTask& Task, TArray<FPCGTaskId>& ResultsToMarkAsRead);
-	void BuildTaskInput(FPCGGraphTask& Task, TArray<FPCGTaskId>& ResultsToMarkAsRead);
+	bool SetupTask(FPCGGraphTask& Task, FTaskIds& ResultsToMarkAsRead);
+	void BuildTaskInput(FPCGGraphTask& Task, FTaskIds& ResultsToMarkAsRead);
 	/** Will check the cache for existing result or create and initialize the FPCGContext to the task */
 	void PrepareForExecute(FPCGGraphTask& Task, FCachedResult*& OutCachedResult, bool bLiveTasksLockAlreadyLocked);
 
@@ -243,12 +247,20 @@ private:
 #endif
 
 	/** Store cache results and Queue next tasks */
-	void ProcessCachedResults(TArray<FCachedResult*> CachedResults);
-	TArray<FPCGTaskId> ProcessCachedResultsInternal(TArray<FCachedResult*> CachedResults);
+	void ProcessCachedResults(TConstArrayView<FCachedResult*> CachedResults);
+	void ProcessCachedResultsInternal(TConstArrayView<FCachedResult*> CachedResults, FTaskIds& OutNextTasks);
 
 	/** Combine all param data into one on the Params pin, if any.*/
 	void CombineParams(FPCGGraphTask& Task);
-	void StoreResults(FPCGTaskId InTaskId, const FPCGDataCollection& InTaskOutput, bool bNeedsManualClear);
+
+	struct FTaskResult
+	{
+		FPCGTaskId TaskId;
+		const FPCGDataCollection* TaskOutput = nullptr;
+		bool bNeedsManualClear = false;
+	};
+
+	void StoreResults(TConstArrayView<FTaskResult> Results);
 	void ClearResults();
 	void MarkInputResults(TArrayView<const FPCGTaskId> InInputResults);
 
diff --git a/Engine/Plugins/PCG/Source/PCG/Private/PCGElement.cpp b/Engine/Plugins/PCG/Source/PCG/Private/PCGElement.cpp
--- a/Engine/Plugins/PCG/Source/PCG/Private/PCGElement.cpp
+++ b/Engine/Plugins/PCG/Source/PCG/Private/PCGElement.cpp
@@ -1207,6 +1207,42 @@ PRAGMA_ENABLE_DEPRECATION_WARNINGS
 	}
 }
 
+bool IPCGElement::BuildCacheParams(IPCGGraphCache* Cache, const UPCGNode* Node, const FPCGDataCollection& Input, IPCGGraphExecutionSource* ExecutionSource, FPCGGetFromCacheParams& OutCacheParams) const
+{
+	if (!Cache || !Cache->IsEnabled())
+	{
+		return false;
+	}
+
+	const UPCGSettingsInterface* SettingsInterface = Input.GetSettingsInterface(Node ? Node->GetSettingsInterface() : nullptr);
+	const UPCGSettings* Settings = SettingsInterface ? SettingsInterface->GetSettings() : nullptr;
+	const bool bCacheable = IsCacheableInstance(SettingsInterface);
+
+	OutCacheParams = { .Node = Node, .Element = this, .ExecutionSource = ExecutionSource };
+
+	if (Settings && bCacheable)
+	{
+		GetDependenciesCrc(FPCGGetDependenciesCrcParams(&Input, Settings, ExecutionSource), OutCacheParams.Crc);
+		return true;
+	}
+	else
+	{
+		return false;
+	}
+}
+
+EPCGCachingStatus IPCGElement::RetrieveResultsFromCache(IPCGGraphCache* Cache, const FPCGGetFromCacheParams& InCacheParams, FPCGDataCollection& Output) const
+{
+	check(Cache && Cache->IsEnabled() && InCacheParams.Crc.IsValid());
+
+	if (Cache->GetFromCache(InCacheParams, Output))
+	{
+		return EPCGCachingStatus::Cached;
+	}
+
+	return EPCGCachingStatus::NotInCache;
+}
+
 EPCGCachingStatus IPCGElement::RetrieveResultsFromCache(IPCGGraphCache* Cache, const UPCGNode* Node, const FPCGDataCollection& Input, IPCGGraphExecutionSource* ExecutionSource, FPCGDataCollection& Output, FPCGCrc* OutCrc) const
 {
 	if (!Cache || !Cache->IsEnabled())
diff --git a/Engine/Plugins/PCG/Source/PCG/Public/Graph/IPCGGraphCache.h b/Engine/Plugins/PCG/Source/PCG/Public/Graph/IPCGGraphCache.h
--- a/Engine/Plugins/PCG/Source/PCG/Public/Graph/IPCGGraphCache.h
+++ b/Engine/Plugins/PCG/Source/PCG/Public/Graph/IPCGGraphCache.h
@@ -15,9 +15,6 @@ struct FPCGGetFromCacheParams
 	const UPCGNode* Node = nullptr;
 	const IPCGElement* Element = nullptr;
 
-	UE_DEPRECATED(5.6, "Use ExecutionSource instead")
-	const UPCGComponent* Component = nullptr;
-
 	const IPCGGraphExecutionSource* ExecutionSource = nullptr;
 	FPCGCrc Crc;
 };
diff --git a/Engine/Plugins/PCG/Source/PCG/Public/PCGElement.h b/Engine/Plugins/PCG/Source/PCG/Public/PCGElement.h
--- a/Engine/Plugins/PCG/Source/PCG/Public/PCGElement.h
+++ b/Engine/Plugins/PCG/Source/PCG/Public/PCGElement.h
@@ -14,6 +14,7 @@ class IPCGGraphExecutionSource;
 class UPCGSettingsInterface;
 struct FPCGContext;
 struct FPCGCrc;
+struct FPCGGetFromCacheParams;
 
 class IPCGElement;
 class IPCGGraphCache;
@@ -168,7 +169,13 @@ public:
 	UE_DEPRECATED(5.6, "Use/Implement version with FPCGGetDependenciesCrcParams parameter instead")
 	UE_API virtual void GetDependenciesCrc(const FPCGDataCollection& InInput, const UPCGSettings* InSettings, UPCGComponent* InComponent, FPCGCrc& OutCrc) const;
 
+	/** Build a query structure to call RetrieveResultsFromCache */
+	UE_API bool BuildCacheParams(IPCGGraphCache* Cache, const UPCGNode* Node, const FPCGDataCollection& Input, IPCGGraphExecutionSource* ExecutionSource, FPCGGetFromCacheParams& OutCacheParams) const;
+
 	/** Gather input data (pre-context creation) and tries to retrieve matching data from the cache if the element is cacheable. */
+	UE_API EPCGCachingStatus RetrieveResultsFromCache(IPCGGraphCache* Cache, const FPCGGetFromCacheParams& InCacheParams, FPCGDataCollection& Output) const;
+
+	UE_DEPRECATED(5.8, "Use BuildCacheParam/RetrieveResultsFromCache combination instead")
 	UE_API EPCGCachingStatus RetrieveResultsFromCache(IPCGGraphCache* Cache, const UPCGNode* Node, const FPCGDataCollection& Input, IPCGGraphExecutionSource* ExecutionSource, FPCGDataCollection& Output, FPCGCrc* OutCrc = nullptr) const;
 
 	/** Public function that executes the element on the appropriately created context.
