[PCG] GPU: Properly empty data binding memory, and flush data description cache after prepare-for-execute.

Expedites freeing of ~1-2mb of memory that otherwise waits for GC.

[RN] [PCG] Optimized memory usage for GPU graphs by proactively freeing memory used in compute graph execution.

diff --git a/Engine/Plugins/PCG/Source/PCG/Private/Compute/Elements/PCGComputeGraphElement.cpp b/Engine/Plugins/PCG/Source/PCG/Private/Compute/Elements/PCGComputeGraphElement.cpp
--- a/Engine/Plugins/PCG/Source/PCG/Private/Compute/Elements/PCGComputeGraphElement.cpp
+++ b/Engine/Plugins/PCG/Source/PCG/Private/Compute/Elements/PCGComputeGraphElement.cpp
@@ -442,6 +442,9 @@ bool FPCGComputeGraphElement::ExecuteInternal(FPCGContext* InContext) const
 			return false;
 		}
 
+		// After prepare for execute flush the data description cache because the cache uses non-negligible memory.
+		Context->DataBinding->FlushDataDescriptionCache();
+
 		if (Context->AsyncState.ShouldStop())
 		{
 			return false;
diff --git a/Engine/Plugins/PCG/Source/PCG/Private/Compute/PCGDataBinding.cpp b/Engine/Plugins/PCG/Source/PCG/Private/Compute/PCGDataBinding.cpp
--- a/Engine/Plugins/PCG/Source/PCG/Private/Compute/PCGDataBinding.cpp
+++ b/Engine/Plugins/PCG/Source/PCG/Private/Compute/PCGDataBinding.cpp
@@ -79,17 +79,20 @@ void UPCGDataBinding::InitializeTables(FPCGContext* InContext)
 void UPCGDataBinding::ReleaseTransientResources()
 {
 	KernelParamsCache.Reset();
-	SourceBufferAttributeToGraphAttributeIndex.Reset();
-	GraphAttributeToSourceBufferAttributeIndex.Reset();
-	SourceBufferToGraphStringKey.Reset();
-	GraphToSourceBufferStringKey.Reset();
-	OutputDataCollection.TaggedData.Reset();
-	InputDataCollection.TaggedData.Reset();
-	MeshSpawnersToPrimitives.Reset();
-	DataToDebug.Reset();
-	DataToInspect.Reset();
-	ContextHandle = nullptr;
-	DataDescriptionCache.Reset();
+	SourceBufferAttributeToGraphAttributeIndex.Empty();
+	GraphAttributeToSourceBufferAttributeIndex.Empty();
+	SourceBufferToGraphStringKey.Empty();
+	GraphToSourceBufferStringKey.Empty();
+	InputDataCollection.TaggedData.Empty();
+	OutputDataCollection.TaggedData.Empty();
+	MeshSpawnersToPrimitives.Empty();
+	CompletedMeshSpawners.Empty();
+	DataToDebug.Empty();
+	DataToInspect.Empty();
+	ContextHandle.Reset();
+	DataDescriptionCache.Empty();
+	AttributeTable.Empty();
+	StringTable.Empty();
 }
 
 void UPCGDataBinding::PrimeDataDescriptionCache()
@@ -119,6 +122,13 @@ void UPCGDataBinding::PrimeDataDescriptionCache()
 	bIsDataDescriptionCachePrimed = true;
 }
 
+void UPCGDataBinding::FlushDataDescriptionCache()
+{
+	TRACE_CPUPROFILER_EVENT_SCOPE(UPCGDataBinding::FlushDataDescriptionCache);
+	DataDescriptionCache.Empty();
+	bIsDataDescriptionCachePrimed = false;
+}
+
 TSharedPtr<const FPCGDataCollectionDesc> UPCGDataBinding::ComputeKernelPinDataDesc(const FPCGKernelPin& InKernelPin)
 {
 	if (!ensure(Graph))
diff --git a/Engine/Plugins/PCG/Source/PCG/Private/Compute/PCGKernelParamsCache.cpp b/Engine/Plugins/PCG/Source/PCG/Private/Compute/PCGKernelParamsCache.cpp
--- a/Engine/Plugins/PCG/Source/PCG/Private/Compute/PCGKernelParamsCache.cpp
+++ b/Engine/Plugins/PCG/Source/PCG/Private/Compute/PCGKernelParamsCache.cpp
@@ -212,8 +212,8 @@ const FPCGKernelParams* FPCGKernelParamsCache::GetCachedKernelParams(int32 InKer
 
 void FPCGKernelParamsCache::Reset()
 {
-	KernelParams.Reset();
-	OverridableParamDataIndices.Reset();
+	KernelParams.Empty();
+	OverridableParamDataIndices.Empty();
 }
 
 void FPCGKernelParamsCache::CacheOverridableParamDataIndices(FPCGContext* InContext, const UPCGDataBinding* InDataBinding)
diff --git a/Engine/Plugins/PCG/Source/PCG/Public/Compute/PCGDataBinding.h b/Engine/Plugins/PCG/Source/PCG/Public/Compute/PCGDataBinding.h
--- a/Engine/Plugins/PCG/Source/PCG/Public/Compute/PCGDataBinding.h
+++ b/Engine/Plugins/PCG/Source/PCG/Public/Compute/PCGDataBinding.h
@@ -91,8 +91,12 @@ public:
 
 	/** Computes a description of data for every pin in the compute graph and caches it. */
 	void PrimeDataDescriptionCache();
+
 	bool IsDataDescriptionCachePrimed() const { return bIsDataDescriptionCachePrimed; }
 
+	/** Flushes computed data descriptions. Subsequent attempts to retrieve from the cache will trigger errors. */
+	void FlushDataDescriptionCache();
+
 	/** Computes a description of data for a single pin in the compute graph and caches it. */
 	PCG_API TSharedPtr<const FPCGDataCollectionDesc> ComputeKernelPinDataDesc(const FPCGKernelPin& InKernelPin);
 	TSharedPtr<const FPCGDataCollectionDesc> ComputeKernelBindingDataDesc(int32 InGraphBindingIndex);
diff --git a/Engine/Plugins/PCG/Source/PCG/Public/Compute/PCGDataDescription.h b/Engine/Plugins/PCG/Source/PCG/Public/Compute/PCGDataDescription.h
--- a/Engine/Plugins/PCG/Source/PCG/Public/Compute/PCGDataDescription.h
+++ b/Engine/Plugins/PCG/Source/PCG/Public/Compute/PCGDataDescription.h
@@ -93,6 +93,8 @@ struct FPCGKernelAttributeTable
 
 	int32 Num() const { return AttributeTable.Num(); }
 
+	void Empty() { AttributeTable.Empty(); }
+
 #if PCG_KERNEL_LOGGING_ENABLED
 	void DebugLog() const;
 #endif
