[PCG] Soften check() that is being hit (likely due to an upstream bug) to ensure() so it does not take down builds.

#rnx

diff --git a/Engine/Plugins/PCG/Source/PCG/Private/Compute/DataInterfaces/PCGInstanceDataInterface.cpp b/Engine/Plugins/PCG/Source/PCG/Private/Compute/DataInterfaces/PCGInstanceDataInterface.cpp
--- a/Engine/Plugins/PCG/Source/PCG/Private/Compute/DataInterfaces/PCGInstanceDataInterface.cpp
+++ b/Engine/Plugins/PCG/Source/PCG/Private/Compute/DataInterfaces/PCGInstanceDataInterface.cpp
@@ -404,7 +404,12 @@ void FPCGInstanceDataProviderProxy::PostSubmit(FComputeContext& Context) const
 
 			FPrimitiveSceneProxy* SceneProxy = Factory->GetSceneProxy(PrimitiveIndex);
 			check(SceneProxy);
-			check(SceneProxy->GetPrimitiveSceneInfo()->GetInstanceSceneDataOffset() != -1);
+
+			if (!ensure(SceneProxy->GetPrimitiveSceneInfo()->GetInstanceSceneDataOffset() >= 0))
+			{
+				UE_LOG(LogPCG, Error, TEXT("FPCGInstanceDataProviderProxy: Scene proxy for primitive %d not ready (InstanceSceneDataOffset is negative), this should not happen."), PrimitiveIndex);
+				continue;
+			}
 
 #if UE_BUILD_SHIPPING
 			const bool bTriggerCapture = false;
