[PCG] Fix wrong lock usage (regression from 46666520)

#rnx

diff --git a/Engine/Plugins/PCG/Source/PCG/Public/Metadata/PCGMetadataAttributeTpl.h b/Engine/Plugins/PCG/Source/PCG/Public/Metadata/PCGMetadataAttributeTpl.h
--- a/Engine/Plugins/PCG/Source/PCG/Public/Metadata/PCGMetadataAttributeTpl.h
+++ b/Engine/Plugins/PCG/Source/PCG/Public/Metadata/PCGMetadataAttributeTpl.h
@@ -281,7 +281,7 @@ protected:
 			{
 				const FPCGMetadataAttribute<T>* Current = Parents[i];
 				
-				PCG::TSharedScopeLock ReadLock(EntryMapLock);
+				PCG::TSharedScopeLock ReadLock(Current->EntryMapLock);
 				AttributeCopy->EntryToValueKeyMap.Append(Current->EntryToValueKeyMap);
 			}
 		}
