[PCG] UObject reduce: avoid allocating a scheduling policy for local components

diff --git a/Engine/Plugins/PCG/Source/PCG/Private/PCGComponent.cpp b/Engine/Plugins/PCG/Source/PCG/Private/PCGComponent.cpp
--- a/Engine/Plugins/PCG/Source/PCG/Private/PCGComponent.cpp
+++ b/Engine/Plugins/PCG/Source/PCG/Private/PCGComponent.cpp
@@ -3411,7 +3411,7 @@ void UPCGComponent::RefreshSchedulingPolicy()
 	}
 
 	// We should never create the scheduling policy when not in runtime generation mode.
-	if (SchedulingPolicyClass && !SchedulingPolicy && IsManagedByRuntimeGenSystem())
+	if (SchedulingPolicyClass && !SchedulingPolicy && IsManagedByRuntimeGenSystem() && !IsLocalComponent())
 	{
 		const EObjectFlags Flags = GetMaskedFlags(RF_PropagateToSubObjects);
 		SchedulingPolicy = NewObject<UPCGSchedulingPolicyBase>(this, SchedulingPolicyClass, NAME_None, Flags);
