[PCG Editor Mode] Implemented Select mode in editor tool mode.

#rnx

diff --git a/Engine/Plugins/PCG/Source/PCG/Internal/Data/Tool/PCGToolActorsData.h b/Engine/Plugins/PCG/Source/PCG/Internal/Data/Tool/PCGToolActorsData.h
new file mode 100644
--- /dev/null
+++ b/Engine/Plugins/PCG/Source/PCG/Internal/Data/Tool/PCGToolActorsData.h
@@ -0,0 +1,33 @@
+// Copyright Epic Games, Inc. All Rights Reserved.
+
+#pragma once
+
+#include "Data/Tool/PCGToolBaseData.h"
+
+#include "PCGToolActorsData.generated.h"
+
+
+/** Tool Data that stores actors. */
+USTRUCT(BlueprintType, DisplayName="Actors Tool Data")
+struct FPCGInteractiveToolWorkingData_Actors : public FPCGInteractiveToolWorkingData
+{
+	GENERATED_BODY()
+	
+public:
+	FPCGInteractiveToolWorkingData_Actors() = default;
+
+	PCG_API virtual void InitializeRuntimeElementData(FPCGContext* InContext) const override;
+
+	PCG_API virtual bool IsValid() const override;
+	
+	PCG_API UPCGParamData* GetParamData() const;
+
+#if WITH_EDITORONLY_DATA
+	PCG_API virtual void InitializeInternal(const FPCGInteractiveToolWorkingDataContext& Context) override;
+#endif
+	
+protected:
+	UPROPERTY(VisibleAnywhere, BlueprintReadOnly, Category="Working Data")
+	TObjectPtr<UPCGParamData> SelectedActorsData;
+
+};
diff --git a/Engine/Plugins/PCG/Source/PCG/Private/Data/Tool/PCGToolActorsData.cpp b/Engine/Plugins/PCG/Source/PCG/Private/Data/Tool/PCGToolActorsData.cpp
new file mode 100644
--- /dev/null
+++ b/Engine/Plugins/PCG/Source/PCG/Private/Data/Tool/PCGToolActorsData.cpp
@@ -0,0 +1,59 @@
+// Copyright Epic Games, Inc. All Rights Reserved.
+
+#include "Data/Tool/PCGToolActorsData.h"
+
+#include "PCGComponent.h"
+#include "PCGContext.h"
+#include "PCGParamData.h"
+#include "Data/PCGBasePointData.h"
+
+#if WITH_EDITOR
+#include "InteractiveTool.h"
+#include "ActorFactories/ActorFactory.h"
+#endif
+
+
+void FPCGInteractiveToolWorkingData_Actors::InitializeRuntimeElementData(FPCGContext* InContext) const
+{
+	FPCGInteractiveToolWorkingData::InitializeRuntimeElementData(InContext);
+
+	TArray<FPCGTaggedData>& Outputs = InContext->OutputData.TaggedData;
+
+	FPCGTaggedData& Output = Outputs.Emplace_GetRef();
+
+#if WITH_EDITOR
+	// Implementation note - we duplicate here to make sure we are not preserving the parent, because it might still change, which would cause issues.
+	check(IsInGameThread());
+	Output.Data = Cast<UPCGData>(StaticDuplicateObject(SelectedActorsData, GetTransientPackage()));
+#else
+	Output.Data = SelectedActorsData;
+#endif
+
+}
+
+bool FPCGInteractiveToolWorkingData_Actors::IsValid() const
+{
+	if (!SelectedActorsData)
+	{
+		return false;
+	}
+
+	const UPCGMetadata* Metadata = SelectedActorsData->ConstMetadata();
+	check(Metadata);
+	const FPCGMetadataAttribute<FSoftObjectPath>* ActorReference = Metadata->GetConstTypedAttribute<FSoftObjectPath>(PCGPointDataConstants::ActorReferenceAttribute);
+	
+	// Should have at least one actor reference to be valid.
+	return ActorReference && ActorReference->GetNumberOfEntries() > 0;
+}
+
+UPCGParamData* FPCGInteractiveToolWorkingData_Actors::GetParamData() const
+{
+	return SelectedActorsData;
+}
+
+#if WITH_EDITOR
+void FPCGInteractiveToolWorkingData_Actors::InitializeInternal(const FPCGInteractiveToolWorkingDataContext& Context)
+{
+	SelectedActorsData = NewObject<UPCGParamData>();
+}
+#endif  //WITH_EDITOR
diff --git a/Engine/Plugins/PCG/Source/PCGEditor/Private/EditorMode/PCGEdMode.cpp b/Engine/Plugins/PCG/Source/PCGEditor/Private/EditorMode/PCGEdMode.cpp
--- a/Engine/Plugins/PCG/Source/PCGEditor/Private/EditorMode/PCGEdMode.cpp
+++ b/Engine/Plugins/PCG/Source/PCGEditor/Private/EditorMode/PCGEdMode.cpp
@@ -9,6 +9,7 @@
 #include "EditorMode/Tools/PCGInteractiveToolSettings.h"
 #include "EditorMode/Tools/Line/PCGDrawSplineTool.h"
 #include "EditorMode/Tools/Paint/PCGPaintTool.h"
+#include "EditorMode/Tools/Select/PCGSelectTool.h"
 #include "EditorMode/Tools/Volume/PCGVolumeTool.h"
 
 #include "ContextObjectStore.h"
@@ -193,6 +194,11 @@ void UPCGEditorMode::Enter()
 		RegisterTool(ToolManagerCommands.EnableVolumeTool, UPCGInteractiveToolSettings_Volume::StaticGetToolTag().ToString(), VolumeToolBuilder);
 	}
 
+	{
+		UPCGSelectToolBuilder* SelectToolBuilder = NewObject<UPCGSelectToolBuilder>();
+		RegisterTool(ToolManagerCommands.EnableSelectTool, UPCGInteractiveToolSettings_Select::StaticGetToolTag().ToString(), SelectToolBuilder);
+	}
+
 	ToolManager->SelectActiveToolType(EToolSide::Left, "DrawSpline");
 }
 
diff --git a/Engine/Plugins/PCG/Source/PCGEditor/Private/EditorMode/PCGEdModeCommands.cpp b/Engine/Plugins/PCG/Source/PCGEditor/Private/EditorMode/PCGEdModeCommands.cpp
--- a/Engine/Plugins/PCG/Source/PCGEditor/Private/EditorMode/PCGEdModeCommands.cpp
+++ b/Engine/Plugins/PCG/Source/PCGEditor/Private/EditorMode/PCGEdModeCommands.cpp
@@ -9,6 +9,7 @@
 #include "EditorMode/Tools/Helpers/PCGEdModeEditorUtilities.h"
 #include "EditorMode/Tools/Line/PCGDrawSplineTool.h"
 #include "EditorMode/Tools/Paint/PCGPaintTool.h"
+#include "EditorMode/Tools/Select/PCGSelectTool.h"
 #include "EditorMode/Tools/Volume/PCGVolumeTool.h"
 #include "Logging/StructuredLog.h"
 #include "Misc/OutputDeviceNull.h"
@@ -52,6 +53,7 @@ void FPCGEditorModePaletteCommands::RegisterCommands()
 	PCG_CONTEXT_COMMAND_INFO(LoadSplineContextPalette, "Spline", "PCGEditorMode.Context.DrawSpline");
 	PCG_CONTEXT_COMMAND_INFO(LoadPaintContextPalette, "Paint", "PCGEditorMode.Context.Paint");
 	PCG_CONTEXT_COMMAND_INFO(LoadVolumeContextPalette, "Volume", "PCGEditorMode.Context.Volume");
+	PCG_CONTEXT_COMMAND_INFO(LoadSelectContextPalette, "Select", "PCGEditorMode.Context.Select");
 }
 
 #undef PCG_CONTEXT_COMMAND_INFO
@@ -90,6 +92,7 @@ void FPCGEditorModeToolCommands::RegisterCommands()
 	PCG_TOOL_COMMAND_INFO(EnableDrawSurfaceTool, "DrawSurface", "Draw Spline Surface", "Use the '{ToolTag}' tool to build a closed spline representing an area for PCG.", UPCGInteractiveToolSettings_SplineSurface::StaticGetToolTag(), "PCGEditorMode.Tools.DrawSurface");
 	PCG_TOOL_COMMAND_INFO(EnablePaintTool, "Paint", "Paint", "Use the '{ToolTag}' tool to apply PCG points with a brush.", UPCGInteractiveToolSettings_PaintTool::StaticGetToolTag(), "PCGEditorMode.Context.Paint");
 	PCG_TOOL_COMMAND_INFO(EnableVolumeTool, "Volume", "Volume", "Create or updates a volume in the scene using the '{ToolTag}' tool.", UPCGInteractiveToolSettings_Volume::StaticGetToolTag(), "PCGEditorMode.Tools.Volume");
+	PCG_TOOL_COMMAND_INFO(EnableSelectTool, "Select", "Select", "Apply PCG graph to selection using the '{ToolTag}' tool.", UPCGInteractiveToolSettings_Select::StaticGetToolTag(), "PCGEditorMode.Tools.Select");
 }
 
 void FPCGEditorModeToolCommands::RegisterAllToolCommands()
@@ -97,10 +100,12 @@ void FPCGEditorModeToolCommands::RegisterAllToolCommands()
 	FPCGDrawSplineCommands::Register();
 	FPCGPaintCommands::Register();
 	FPCGVolumeCommands::Register();
+	FPCGSelectCommands::Register();
 }
 
 void FPCGEditorModeToolCommands::UnregisterAllToolCommands()
 {
+	FPCGSelectCommands::Unregister();
 	FPCGVolumeCommands::Unregister();
 	FPCGPaintCommands::Unregister();
 	FPCGDrawSplineCommands::Unregister();
@@ -126,6 +131,11 @@ TArray<TSharedPtr<FUICommandInfo>> FPCGEditorModeToolCommands::GetVolumeContextT
 	return {EnableVolumeTool};
 }
 
+TArray<TSharedPtr<FUICommandInfo>> FPCGEditorModeToolCommands::GetSelectContextToolCommands()
+{
+	return { EnableSelectTool };
+}
+
 /**
  * Macro for defining a constructor and a helper function to retrieve 
  * @param ToolCommandClassName The tool's class name
@@ -147,6 +157,7 @@ PCG_DEFINE_TOOL_COMMANDS(FPCGDrawSplineCommands, "DrawSpline", "Draw Spline", UP
 PCG_DEFINE_TOOL_COMMANDS(FPCGDrawSurfaceCommands, "DrawSurface", "Draw Surface", UPCGDrawSplineToolBase);
 PCG_DEFINE_TOOL_COMMANDS(FPCGPaintCommands, "Paint", "Paint", UPCGPaintTool);
 PCG_DEFINE_TOOL_COMMANDS(FPCGVolumeCommands, "Volume", "Volume", UPCGVolumeTool);
+PCG_DEFINE_TOOL_COMMANDS(FPCGSelectCommands, "Select", "Select", UPCGSelectTool);
 
 #undef PCG_DEFINE_TOOL_COMMANDS
 
diff --git a/Engine/Plugins/PCG/Source/PCGEditor/Private/EditorMode/PCGEdModeCommands.h b/Engine/Plugins/PCG/Source/PCGEditor/Private/EditorMode/PCGEdModeCommands.h
--- a/Engine/Plugins/PCG/Source/PCGEditor/Private/EditorMode/PCGEdModeCommands.h
+++ b/Engine/Plugins/PCG/Source/PCGEditor/Private/EditorMode/PCGEdModeCommands.h
@@ -26,6 +26,7 @@ public:
 	TSharedPtr<FUICommandInfo> LoadSplineContextPalette;
 	TSharedPtr<FUICommandInfo> LoadPaintContextPalette;
 	TSharedPtr<FUICommandInfo> LoadVolumeContextPalette;
+	TSharedPtr<FUICommandInfo> LoadSelectContextPalette;
 
 	// @todo_pcg: Selection/Filtering actions.
 
@@ -71,6 +72,10 @@ public:
 	/*** PCG EDITOR MODE TOOLS - Volume Context **/
 	TSharedPtr<FUICommandInfo> EnableVolumeTool;
 	TArray<TSharedPtr<FUICommandInfo>> GetVolumeContextToolCommands();
+
+	/*** PCG EDITOR MODE TOOLS - Select Context **/
+	TSharedPtr<FUICommandInfo> EnableSelectTool;
+	TArray<TSharedPtr<FUICommandInfo>> GetSelectContextToolCommands();
 };
 
 /**
@@ -92,5 +97,6 @@ DECLARE_TOOL_COMMANDS(FPCGDrawSplineCommands);
 DECLARE_TOOL_COMMANDS(FPCGDrawSurfaceCommands);
 DECLARE_TOOL_COMMANDS(FPCGPaintCommands);
 DECLARE_TOOL_COMMANDS(FPCGVolumeCommands);
+DECLARE_TOOL_COMMANDS(FPCGSelectCommands);
 
 #undef DECLARE_TOOL_COMMANDS
diff --git a/Engine/Plugins/PCG/Source/PCGEditor/Private/EditorMode/PCGEdModeStyle.cpp b/Engine/Plugins/PCG/Source/PCGEditor/Private/EditorMode/PCGEdModeStyle.cpp
--- a/Engine/Plugins/PCG/Source/PCGEditor/Private/EditorMode/PCGEdModeStyle.cpp
+++ b/Engine/Plugins/PCG/Source/PCGEditor/Private/EditorMode/PCGEdModeStyle.cpp
@@ -64,6 +64,7 @@ FPCGEditorModeStyle::FPCGEditorModeStyle()
 		Set("PCGEditorMode.Context.DrawSpline", new IMAGE_BRUSH_SVG("Icons/EdMode/DrawSpline", Icon20x20));
 		Set("PCGEditorMode.Context.Paint", new IMAGE_BRUSH_SVG("Icons/EdMode/Paint", Icon20x20));
 		Set("PCGEditorMode.Context.Volume", new IMAGE_BRUSH_SVG("Icons/EdMode/Volume", Icon20x20));
+		Set("PCGEditorMode.Context.Select", new IMAGE_BRUSH_SVG("Icons/EdMode/Select", Icon20x20));
 	}
 
 	// Tools
@@ -74,6 +75,8 @@ FPCGEditorModeStyle::FPCGEditorModeStyle()
 		Set("PCGEditorMode.Tools.Paint", new IMAGE_BRUSH_SVG("Icons/EdMode/Paint", Icon20x20));
 
 		Set("PCGEditorMode.Tools.Volume", new IMAGE_BRUSH_SVG("Icons/EdMode/Volume", Icon20x20));
+
+		Set("PCGEditorMode.Tools.Select", new IMAGE_BRUSH_SVG("Icons/EdMode/Select", Icon20x20));
 	}
 
 	// Backdrops and backgrounds
diff --git a/Engine/Plugins/PCG/Source/PCGEditor/Private/EditorMode/PCGEdModeToolkit.cpp b/Engine/Plugins/PCG/Source/PCGEditor/Private/EditorMode/PCGEdModeToolkit.cpp
--- a/Engine/Plugins/PCG/Source/PCGEditor/Private/EditorMode/PCGEdModeToolkit.cpp
+++ b/Engine/Plugins/PCG/Source/PCGEditor/Private/EditorMode/PCGEdModeToolkit.cpp
@@ -22,6 +22,7 @@
 #include "Styling/SlateIconFinder.h"
 #include "Tools/Line/PCGDrawSplineTool.h"
 #include "Tools/Paint/PCGPaintTool.h"
+#include "Tools/Select/PCGSelectTool.h"
 #include "Tools/Volume/PCGVolumeTool.h"
 #include "Widgets/SPCGToolPresetSection.h"
 #include "Widgets/Input/STextComboBox.h"
@@ -507,20 +508,26 @@ void FPCGEditorModeToolkit::RegisterPalettes()
 		
 	{
 		FString ToolIdentifier = UPCGInteractiveToolSettings_Spline::StaticGetToolTag().ToString();
-		const TSharedPtr<FPCGToolPalette> DrawSplineToolSection = MakeShared<FPCGToolPalette>(PaletteCommands.LoadSplineContextPalette, ToolCommands.GetLineContextToolCommands());
-		ToolkitBuilder->AddPalette(std::move(DrawSplineToolSection));
+		TSharedPtr<FPCGToolPalette> DrawSplineToolSection = MakeShared<FPCGToolPalette>(PaletteCommands.LoadSplineContextPalette, ToolCommands.GetLineContextToolCommands());
+		ToolkitBuilder->AddPalette(MoveTemp(DrawSplineToolSection));
 	}
 
 	{
 		FString ToolIdentifier = UPCGInteractiveToolSettings_PaintTool::StaticGetToolTag().ToString();
-		const TSharedPtr<FPCGToolPalette> PaintToolPalette = MakeShared<FPCGToolPalette>(PaletteCommands.LoadPaintContextPalette, ToolCommands.GetPaintContextToolCommands());
-		ToolkitBuilder->AddPalette(std::move(PaintToolPalette));
+		TSharedPtr<FPCGToolPalette> PaintToolPalette = MakeShared<FPCGToolPalette>(PaletteCommands.LoadPaintContextPalette, ToolCommands.GetPaintContextToolCommands());
+		ToolkitBuilder->AddPalette(MoveTemp(PaintToolPalette));
 	}
 
 	{
 		FString ToolIdentifier = UPCGInteractiveToolSettings_Volume::StaticGetToolTag().ToString();
-		const TSharedPtr<FPCGToolPalette> VolumeToolPalette = MakeShared<FPCGToolPalette>(PaletteCommands.LoadVolumeContextPalette, ToolCommands.GetVolumeContextToolCommands());
-		ToolkitBuilder->AddPalette(std::move(VolumeToolPalette));
+		TSharedPtr<FPCGToolPalette> VolumeToolPalette = MakeShared<FPCGToolPalette>(PaletteCommands.LoadVolumeContextPalette, ToolCommands.GetVolumeContextToolCommands());
+		ToolkitBuilder->AddPalette(MoveTemp(VolumeToolPalette));
+	}
+
+	{
+		FString ToolIdentifier = UPCGInteractiveToolSettings_Select::StaticGetToolTag().ToString();
+		TSharedPtr<FPCGToolPalette> SelectToolPalette = MakeShared<FPCGToolPalette>(PaletteCommands.LoadSelectContextPalette, ToolCommands.GetSelectContextToolCommands());
+		ToolkitBuilder->AddPalette(MoveTemp(SelectToolPalette));
 	}
 
 	ToolkitBuilder->SetActivePaletteOnLoad(PaletteCommands.LoadSplineContextPalette.Get());
diff --git a/Engine/Plugins/PCG/Source/PCGEditor/Private/EditorMode/Tools/PCGInteractiveToolSettings.cpp b/Engine/Plugins/PCG/Source/PCGEditor/Private/EditorMode/Tools/PCGInteractiveToolSettings.cpp
--- a/Engine/Plugins/PCG/Source/PCGEditor/Private/EditorMode/Tools/PCGInteractiveToolSettings.cpp
+++ b/Engine/Plugins/PCG/Source/PCGEditor/Private/EditorMode/Tools/PCGInteractiveToolSettings.cpp
@@ -112,8 +112,10 @@ TValueOrError<bool, FText> UPCGInteractiveToolSettings::InitActorAndComponent()
 	// - Have a PCG Component with valid working data we want to edit
 	check(!EditActor.IsValid());
 
-	if (SelectedActors.Num() == 1 && SelectedActors[0] && IsWorkingActorCompatible(SelectedActors[0]))
+	const bool HasRightNumberOfSelectedActors = CanHandleMultipleActorSelection() ? (SelectedActors.Num() >= 1) : (SelectedActors.Num() == 1);
+	if (HasRightNumberOfSelectedActors && SelectedActors[0] && IsWorkingActorCompatible(SelectedActors[0]))
 	{
+		//Edit Actor is always the first one.
 		EditActor = SelectedActors[0];
 	}
 	// If we have no selected actor, we spawn a working actor that needs to fulfill the same actor requirements. This actor is fully instantiated when Applied.
diff --git a/Engine/Plugins/PCG/Source/PCGEditor/Private/EditorMode/Tools/PCGInteractiveToolSettings.h b/Engine/Plugins/PCG/Source/PCGEditor/Private/EditorMode/Tools/PCGInteractiveToolSettings.h
--- a/Engine/Plugins/PCG/Source/PCGEditor/Private/EditorMode/Tools/PCGInteractiveToolSettings.h
+++ b/Engine/Plugins/PCG/Source/PCGEditor/Private/EditorMode/Tools/PCGInteractiveToolSettings.h
@@ -47,7 +47,9 @@ public:
 	
 	/** Register property watchers and callbacks in here. */
 	virtual void RegisterPropertyWatchers();
-	
+
+	virtual bool CanHandleMultipleActorSelection() const { return false; }
+
 	/** Initializes or finds the working data associated with the data instance identifier.
 	 * If NAME_None, the ToolTag is used as identifier.
 	 * If bAllowRetryAfterInit is true, if the data was properly initialized, it will retry (once) after. Otherwise, it will not.
diff --git a/Engine/Plugins/PCG/Source/PCGEditor/Private/EditorMode/Tools/Select/PCGSelectTool.cpp b/Engine/Plugins/PCG/Source/PCGEditor/Private/EditorMode/Tools/Select/PCGSelectTool.cpp
new file mode 100644
--- /dev/null
+++ b/Engine/Plugins/PCG/Source/PCGEditor/Private/EditorMode/Tools/Select/PCGSelectTool.cpp
@@ -0,0 +1,147 @@
+// Copyright Epic Games, Inc. All Rights Reserved.
+
+#include "EditorMode/Tools/Select/PCGSelectTool.h"
+
+#include "InteractiveToolManager.h"
+#include "PCGComponent.h"
+#include "EditorMode/PCGEdModeSettings.h"
+#include "EditorMode/PCGEdModeStyle.h"
+#include "EditorMode/Tools/Helpers/PCGEdModeSceneVisualizeHelpers.h"
+#include "ToolTargetManager.h"
+#include "TargetInterfaces/PrimitiveComponentBackedTarget.h"
+#include "ToolTargets/PrimitiveComponentToolTarget.h"
+
+#include "BaseBehaviors/SingleClickOrDragBehavior.h"
+#include "Settings/LevelEditorViewportSettings.h"
+#include "Subsystems/PCGSubsystem.h"
+#include "Subsystems/PCGEngineSubsystem.h"
+#include "PCGDefaultExecutionSource.h"
+#include "PCGParamData.h"
+#include "Data/PCGBasePointData.h"
+#include "Data/Tool/PCGToolActorsData.h"
+
+#include UE_INLINE_GENERATED_CPP_BY_NAME(PCGSelectTool)
+
+#define LOCTEXT_NAMESPACE "UPCGSelectTool"
+
+namespace PCGSelectTool
+{
+namespace Constants
+{
+	const FText ToolName = LOCTEXT("ToolName", "Select");
+	const FText ToolDescription = LOCTEXT("SelectDisplayMessage", "Applies PCG Graph on selected actors.");
+	const FName SelectToolTag = "SelectTool";
+}
+} // namespace PCGSelectTool
+
+
+FName UPCGInteractiveToolSettings_Select::StaticGetToolTag()
+{
+	return PCGSelectTool::Constants::SelectToolTag;
+}
+
+const UScriptStruct* UPCGInteractiveToolSettings_Select::GetWorkingDataType() const
+{
+	return FPCGInteractiveToolWorkingData_Actors::StaticStruct();
+}
+
+FName UPCGInteractiveToolSettings_Select::GetToolTag() const
+{
+	return StaticGetToolTag();
+}
+
+void UPCGSelectTool::Setup()
+{
+	AddToolPropertySource(ToolSettings);
+
+	ToolSettings->InitializeWorkingData(PCGSelectTool::Constants::SelectToolTag);
+
+	const FPCGInteractiveToolWorkingData_Actors* WorkingData = ToolSettings->GetTypedWorkingData<FPCGInteractiveToolWorkingData_Actors>(ToolSettings->DataInstance);
+	if (WorkingData)
+	{
+		UPCGParamData* ParamData = WorkingData->GetParamData();
+		check(ParamData);
+		UPCGMetadata* Metadata = ParamData->MutableMetadata();
+
+		FPCGMetadataAttribute<FSoftObjectPath>* ActorReference = Metadata->FindOrCreateAttribute(PCGPointDataConstants::ActorReferenceAttribute, FSoftObjectPath(), /*bAllowsInterpolation=*/false, /*bOverrideParent=*/false, /*bOverwriteIfTypeMismatch=*/true);
+		check(ActorReference);
+		for (FSoftObjectPath CurActor : SelectedActors)
+		{
+			PCGMetadataEntryKey Entry = ParamData->MutableMetadata()->AddEntry();
+			ActorReference->SetValue(Entry, CurActor);
+		}
+	}
+}
+
+void UPCGSelectTool::Shutdown(EToolShutdownType ShutdownType)
+{
+	TWeakObjectPtr<UPCGComponent> Component = ToolSettings->GeneratedResources.GeneratedPCGComponent;
+	UE::PCG::EditorMode::Tool::Shutdown(this, ShutdownType);
+
+	Super::Shutdown(ShutdownType);
+
+	if (ShutdownType == EToolShutdownType::Accept)
+	{
+		// If a PCG component was generated, the default ToolSetting does not delete it at apply.
+		// In the case of the Select Tool, we do want to delete it to definitely commit the generated resources.
+		DeleteGeneratedPCGComponent(Component);
+	}
+}
+
+bool UPCGSelectTool::CanAccept() const
+{
+	return true;
+}
+
+void UPCGSelectTool::SetSelectedActors(const TArray<AActor*>& Actors)
+{
+	SelectedActors.Empty(Actors.Num());
+	Algo::Transform(Actors, SelectedActors, [](AActor* Actor) {return FSoftObjectPath{Actor}; });
+}
+
+void UPCGSelectTool::DeleteGeneratedPCGComponent(TWeakObjectPtr<UPCGComponent> GeneratedPCGComponent)
+{
+	if (GeneratedPCGComponent.IsValid())
+	{
+		GeneratedPCGComponent->GetSubsystem()->ScheduleGeneric(
+			[PCGComponent = GeneratedPCGComponent]() -> bool
+			{
+				if (PCGComponent.IsValid())
+				{
+					PCGComponent->DestroyComponent();
+				}
+				return true;
+			}, 
+			GeneratedPCGComponent.Get(), 
+			{ GeneratedPCGComponent->GetGenerationTaskId() });
+	}
+}
+
+bool UPCGSelectToolBuilder::CanBuildTool(const FToolBuilderState& SceneState) const
+{
+	return SceneState.TargetManager->CountSelectedAndTargetable(SceneState, FToolTargetTypeRequirements()) > 0;
+}
+
+UInteractiveTool* UPCGSelectToolBuilder::BuildTool(const FToolBuilderState& SceneState) const
+{
+	UPCGSelectTool* NewTool = NewObject<UPCGSelectTool>(SceneState.ToolManager);
+
+	NewTool->SetToolInfo(
+	{
+		.ToolDisplayName = PCGSelectTool::Constants::ToolDescription,
+		.ToolDisplayMessage = LOCTEXT("SelectToolDisplayMessage", "Todo_pcg."),
+		.ToolIcon = FPCGEditorModeStyle::Get().GetBrush("PCGEditorMode.Tools.Select")
+	});
+
+	UToolTarget* const Target = SceneState.TargetManager->BuildFirstSelectedTargetable(SceneState, FToolTargetTypeRequirements());
+	NewTool->SetTarget(Target);
+	NewTool->SetSelectedActors(SceneState.SelectedActors);
+
+	if (UE::PCG::EditorMode::Tool::BuildTool(NewTool) == false)
+	{
+		return nullptr;
+	}
+	return NewTool;
+}
+
+#undef LOCTEXT_NAMESPACE
diff --git a/Engine/Plugins/PCG/Source/PCGEditor/Private/EditorMode/Tools/Select/PCGSelectTool.h b/Engine/Plugins/PCG/Source/PCGEditor/Private/EditorMode/Tools/Select/PCGSelectTool.h
new file mode 100644
--- /dev/null
+++ b/Engine/Plugins/PCG/Source/PCGEditor/Private/EditorMode/Tools/Select/PCGSelectTool.h
@@ -0,0 +1,73 @@
+// Copyright Epic Games, Inc. All Rights Reserved.
+
+#pragma once
+
+#include "EditorMode/Tools/PCGInteractiveToolBuilder.h"
+#include "EditorMode/Tools/PCGInteractiveToolCommon.h"
+#include "EditorMode/Tools/PCGInteractiveToolSettings.h"
+#include "EditorMode/Tools/Helpers/PCGEdModeSceneQueryHelpers.h"
+
+#include "BaseBehaviors/BehaviorTargetInterfaces.h"
+#include "BaseBehaviors/SingleClickBehavior.h"
+#include "SingleSelectionTool.h"
+
+#include "PCGSelectTool.generated.h"
+
+class USingleClickOrDragInputBehavior;
+class UPCGDefaultExecutionSource;
+
+UCLASS(BlueprintType)
+class PCGEDITOR_API UPCGInteractiveToolSettings_Select : public UPCGInteractiveToolSettings
+{
+	GENERATED_BODY()
+
+public:
+	UPCGInteractiveToolSettings_Select() = default;
+
+	static FName StaticGetToolTag();
+
+	// ~Begin UPCGInteractiveToolSettings interface
+	virtual const UScriptStruct* GetWorkingDataType() const override;
+	virtual FName GetToolTag() const override;
+	virtual bool CanHandleMultipleActorSelection() const override { return true; }
+	// ~End UPCGInteractiveToolSettings interface
+};
+
+UCLASS(Experimental)
+class PCGEDITOR_API UPCGSelectTool : public USingleSelectionTool
+{
+	GENERATED_BODY()
+
+public:
+	// ~Begin UInteractiveTool interface
+	virtual void Setup() override;
+	virtual void Shutdown(EToolShutdownType ShutdownType) override;
+
+	virtual bool HasCancel() const override { return true; }
+	virtual bool HasAccept() const override { return true; }
+	virtual bool CanAccept() const override;
+	// ~End UInteractiveTool interface
+
+	void SetSelectedActors(const TArray<AActor*>& Actors);
+private:
+	void DeleteGeneratedPCGComponent(TWeakObjectPtr<UPCGComponent> GeneratedPCGComponent);
+
+	UPROPERTY(Transient)
+	TObjectPtr<UPCGInteractiveToolSettings_Select> ToolSettings;
+
+	UPROPERTY(Transient)
+	TObjectPtr<UPCGDefaultExecutionSource> PCGDefaultExecutionSource;
+
+	UPROPERTY(Transient)
+	TArray<FSoftObjectPath> SelectedActors;
+};
+
+UCLASS(Transient)
+class PCGEDITOR_API UPCGSelectToolBuilder : public UPCGInteractiveToolWithToolTargetsBuilder
+{
+	GENERATED_BODY()
+
+public:
+	virtual bool CanBuildTool(const FToolBuilderState& SceneState) const override;
+	virtual UInteractiveTool* BuildTool(const FToolBuilderState& SceneState) const override;
+};
