[PCG] Made the clip path operation overridable

#rnx

diff --git a/Engine/Plugins/PCG/Source/PCG/Public/Elements/Polygon/PCGClipPaths.h b/Engine/Plugins/PCG/Source/PCG/Public/Elements/Polygon/PCGClipPaths.h
--- a/Engine/Plugins/PCG/Source/PCG/Public/Elements/Polygon/PCGClipPaths.h
+++ b/Engine/Plugins/PCG/Source/PCG/Public/Elements/Polygon/PCGClipPaths.h
@@ -9,7 +9,7 @@
 UENUM()
 enum class EPCGClipPathOperation : uint8
 {
-	Intersection,
+	Intersection = 0,
 	Difference
 };
 
@@ -48,7 +48,7 @@ protected:
 
 public:
 	/** Controls whether the paths will be clipped to be inside or outside the clip polygons. */
-	UPROPERTY(EditAnywhere, BlueprintReadWrite, Category = Settings)
+	UPROPERTY(EditAnywhere, BlueprintReadWrite, Category = Settings, meta = (PCG_Overridable))
 	EPCGClipPathOperation Operation = EPCGClipPathOperation::Intersection;
 
 	/** Maximum squared distance before we need to subdivide a segment again as part of the spline discretization to a polygon. */
