[PCG] Fixed possibly risky use of a root component

#rnx

diff --git a/Engine/Plugins/PCG/Source/PCGEditor/Private/EditorMode/Tools/Volume/PCGVolumeTool.cpp b/Engine/Plugins/PCG/Source/PCGEditor/Private/EditorMode/Tools/Volume/PCGVolumeTool.cpp
--- a/Engine/Plugins/PCG/Source/PCGEditor/Private/EditorMode/Tools/Volume/PCGVolumeTool.cpp
+++ b/Engine/Plugins/PCG/Source/PCGEditor/Private/EditorMode/Tools/Volume/PCGVolumeTool.cpp
@@ -224,7 +224,10 @@ void UPCGVolumeTool::UpdateVolumeFromBox()
 	else if (UBoxComponent* WorkingBoxComponent = ToolSettings ? ToolSettings->GetWorkingBoxComponent(ToolSettings->DataInstance) : nullptr)
 	{
 		// If we've spawned this actor, we'll assume we can take over the root transform otherwise, everything will behave weirdly.
-		if (ToolSettings->HasSpawnedActor() && WorkingBoxComponent->GetOwner() && WorkingBoxComponent->GetOwner()->GetRootComponent() != WorkingBoxComponent)
+		if (ToolSettings->HasSpawnedActor() && 
+			WorkingBoxComponent->GetOwner() && 
+			WorkingBoxComponent->GetOwner()->GetRootComponent() &&
+			WorkingBoxComponent->GetOwner()->GetRootComponent() != WorkingBoxComponent)
 		{
 			WorkingBoxComponent->GetOwner()->GetRootComponent()->SetWorldTransform(FTransform(BoxCenter));
 
