[PCG] Separate PCGComponent and FPCGComponentExecutionState

#rnx

diff --git a/Engine/Plugins/PCG/Source/PCG/Private/PCGComponent.cpp b/Engine/Plugins/PCG/Source/PCG/Private/PCGComponent.cpp
--- a/Engine/Plugins/PCG/Source/PCG/Private/PCGComponent.cpp
+++ b/Engine/Plugins/PCG/Source/PCG/Private/PCGComponent.cpp
@@ -4515,257 +4515,4 @@ void FPCGComponentInstanceData::AddReferencedObjects(FReferenceCollector& Collec
 	Collector.AddReferencedObject(SourceComponent);
 }
 
-UPCGData* FPCGComponentExecutionState::GetSelfData() const
-{
-	check(Component);
-	return Component->GetActorPCGData();
-}
-
-int32 FPCGComponentExecutionState::GetSeed() const
-{
-	check(Component);
-	return Component->Seed;
-}
-
-FString FPCGComponentExecutionState::GetDebugName() const
-{
-	check(Component && Component->GetOwner());
-	return Component->GetOwner()->GetName();
-}
-
-FTransform FPCGComponentExecutionState::GetTransform() const
-{
-	check(Component && Component->GetOwner());
-	return Component->GetOwner()->GetTransform();
-}
-
-FTransform FPCGComponentExecutionState::GetOriginalTransform() const
-{
-	if (IPCGGraphExecutionSource* OriginalSource = GetOriginalSource())
-	{
-		return OriginalSource->GetExecutionState().GetTransform();
-	}
-	else
-	{
-		return GetTransform();
-	}
-}
-
-UWorld* FPCGComponentExecutionState::GetWorld() const
-{
-	check(Component);
-	return Component->GetWorld();
-}
-
-bool FPCGComponentExecutionState::HasAuthority() const
-{
-	check(Component && Component->GetOwner());
-	return Component->GetOwner()->HasAuthority();
-}
-
-FBox FPCGComponentExecutionState::GetBounds() const
-{
-	check(Component);
-	return Component->GetGridBounds();
-}
-
-FBox FPCGComponentExecutionState::GetOriginalBounds() const
-{
-	check(Component);
-	return Component->GetOriginalGridBounds();
-}
-
-FBox FPCGComponentExecutionState::GetLocalSpaceBounds() const
-{
-	check(Component);
-	return Component->GetLocalSpaceBounds();
-}
-
-FBox FPCGComponentExecutionState::GetOriginalLocalSpaceBounds() const
-{
-	check(Component);
-	return Component->GetOriginalLocalSpaceBounds();
-}
-
-bool FPCGComponentExecutionState::Use2DGrid() const
-{
-	check(Component);
-	return Component->Use2DGrid();
-}
-
-UPCGGraph* FPCGComponentExecutionState::GetGraph() const
-{
-	check(Component);
-	return Component->GetGraph();
-}
-
-UPCGGraphInstance* FPCGComponentExecutionState::GetGraphInstance() const
-{
-	check(Component);
-	return Component->GetGraphInstance();
-}
-
-void FPCGComponentExecutionState::OnGraphExecutionAborted(bool bQuiet, bool bCleanupUnusedResources)
-{
-	check(Component);
-	Component->OnProcessGraphAborted(bQuiet, bCleanupUnusedResources);
-}
-
-void FPCGComponentExecutionState::Cancel()
-{
-	check(Component);
-	Component->CancelGeneration();
-}
-
-bool FPCGComponentExecutionState::IsGenerating() const
-{
-	check(Component);
-	return Component->IsGenerating();
-}
-
-void FPCGComponentExecutionState::ExecutePreGraph(FPCGContext* InContext)
-{
-	check(Component);
-	Component->ExecutePreGraph(InContext);
-}
-
-bool FPCGComponentExecutionState::IsManagedByRuntimeGenSystem() const
-{
-	check(Component);
-	return Component->IsManagedByRuntimeGenSystem();
-}
-
-IPCGGraphExecutionSource* FPCGComponentExecutionState::GetOriginalSource() const
-{
-	check(Component);
-	return Component->GetOriginalComponent();
-}
-
-IPCGGraphExecutionSource* FPCGComponentExecutionState::GetLocalSource(const FPCGGridDescriptor& GridDescriptor, const FIntVector& CellCoords) const
-{
-	check(Component);
-
-	if (UPCGSubsystem* Subsystem = Cast<UPCGSubsystem>(GetSubsystem()))
-	{
-		return Subsystem->GetLocalComponent(GridDescriptor, CellCoords, CastChecked<UPCGComponent>(GetOriginalSource()));
-	}
-
-	return nullptr;
-}
-
-bool FPCGComponentExecutionState::IsPartitioned() const
-{
-	check(Component);
-	return Component->IsPartitioned();
-}
-
-bool FPCGComponentExecutionState::IsLocalSource() const
-{
-	check(Component);
-	return Component->IsLocalComponent();
-}
-
-FPCGGridDescriptor FPCGComponentExecutionState::GetGridDescriptor(uint32 InGridSize) const
-{
-	check(Component);
-	return Component->GetGridDescriptor(InGridSize);
-}
-
-uint32 FPCGComponentExecutionState::GetGenerationGridSize() const
-{
-	check(Component);
-	return (IsPartitioned() || IsLocalSource()) ? Component->GetGenerationGridSize() : PCGHiGenGrid::UninitializedGridSize();
-}
-
-FPCGTaskId FPCGComponentExecutionState::GetGenerationTaskId() const
-{
-	check(Component);
-	return Component->GetGenerationTaskId();
-}
-
-void FPCGComponentExecutionState::StoreOutputDataForPin(const FString& InResourceKey, const FPCGDataCollection& InData) const
-{
-	check(Component);
-	Component->StoreOutputDataForPin(InResourceKey, InData);
-}
-
-const FPCGDataCollection* FPCGComponentExecutionState::RetrieveOutputDataForPin(const FString& InResourceKey) const
-{
-	check(Component);
-	return Component->RetrieveOutputDataForPin(InResourceKey);
-}
-
-void FPCGComponentExecutionState::AddToManagedResources(UPCGManagedResource* InResource)
-{
-	check(Component);
-	return Component->AddToManagedResources(InResource);
-}
-
-FPCGTaskId FPCGComponentExecutionState::GenerateLocalGetTaskId(EPCGHiGenGrid Grid)
-{
-	check(Component)
-
-	const EPCGComponentGenerationTrigger GenTrigger = IsManagedByRuntimeGenSystem() ? EPCGComponentGenerationTrigger::GenerateAtRuntime : EPCGComponentGenerationTrigger::GenerateOnDemand;
-	return Component->GenerateLocalGetTaskId(GenTrigger, /*bForce=*/true, Grid);
-}
-
-#if WITH_EDITOR
-const PCGUtils::FExtraCapture& FPCGComponentExecutionState::GetExtraCapture() const
-{
-	check(Component);
-	return Component->ExtraCapture;
-}
-
-PCGUtils::FExtraCapture& FPCGComponentExecutionState::GetExtraCapture()
-{
-	check(Component);
-	return Component->ExtraCapture;
-}
-
-const FPCGGraphExecutionInspection& FPCGComponentExecutionState::GetInspection() const
-{
-	check(Component);
-	return Component->ExecutionInspection;
-}
-
-FPCGGraphExecutionInspection& FPCGComponentExecutionState::GetInspection()
-{
-	check(Component);
-	return Component->ExecutionInspection;
-}
-
-void FPCGComponentExecutionState::RegisterDynamicTracking(const UPCGSettings* InSettings, const TArrayView<TPair<FPCGSelectionKey, bool>>& InDynamicKeysAndCulling)
-{
-	check(Component);
-	Component->RegisterDynamicTracking(InSettings, InDynamicKeysAndCulling);
-}
-
-void FPCGComponentExecutionState::RegisterDynamicTracking(const FPCGSelectionKeyToSettingsMap& InKeysToSettings)
-{
-	check(Component);
-	Component->RegisterDynamicTracking(InKeysToSettings);
-}
-
-bool FPCGComponentExecutionState::IsRefreshInProgress() const
-{
-	check(Component);
-	return Component->IsRefreshInProgress();
-}
-
-FPCGDynamicTrackingPriority FPCGComponentExecutionState::GetDynamicTrackingPriority() const
-{
-	check(Component);
-
-	// Hash [0, 1] 
-	double Fraction = static_cast<double>(GetTypeHash(Component->GetPathName())) / std::numeric_limits<uint32>::max();
-	// Divide Fraction by 2^14 so that max value most significant decimal is the 5th (1.0 / 16384 -> 0.00006103515625)
-	// TrackingPriority is rounded to 4 decimals so that Fraction (Hash) doesn't interfere with it (or almost never, in which case difference in value would be minimal).
-	Fraction /= 16384;
-
-	// So that leading integer and fractional user part stay the same add or subtract the hash based fraction based on sign
-	return Component->TrackingPriority >= 0 ? (double)Component->TrackingPriority + Fraction : (double)Component->TrackingPriority - Fraction;
-}
-
-#endif // WITH_EDITOR
-
 #undef LOCTEXT_NAMESPACE
diff --git a/Engine/Plugins/PCG/Source/PCG/Private/PCGComponentExecutionState.cpp b/Engine/Plugins/PCG/Source/PCG/Private/PCGComponentExecutionState.cpp
new file mode 100644
--- /dev/null
+++ b/Engine/Plugins/PCG/Source/PCG/Private/PCGComponentExecutionState.cpp
@@ -0,0 +1,261 @@
+// Copyright Epic Games, Inc. All Rights Reserved.
+
+#include "PCGComponentExecutionState.h"
+
+#include "PCGComponent.h"
+#include "PCGSubsystem.h"
+
+#include "GameFramework/Actor.h"
+
+UPCGData* FPCGComponentExecutionState::GetSelfData() const
+{
+	check(Component);
+	return Component->GetActorPCGData();
+}
+
+int32 FPCGComponentExecutionState::GetSeed() const
+{
+	check(Component);
+	return Component->Seed;
+}
+
+FString FPCGComponentExecutionState::GetDebugName() const
+{
+	check(Component && Component->GetOwner());
+	return Component->GetOwner()->GetName();
+}
+
+FTransform FPCGComponentExecutionState::GetTransform() const
+{
+	check(Component && Component->GetOwner());
+	return Component->GetOwner()->GetTransform();
+}
+
+FTransform FPCGComponentExecutionState::GetOriginalTransform() const
+{
+	if (IPCGGraphExecutionSource* OriginalSource = GetOriginalSource())
+	{
+		return OriginalSource->GetExecutionState().GetTransform();
+	}
+	else
+	{
+		return GetTransform();
+	}
+}
+
+UWorld* FPCGComponentExecutionState::GetWorld() const
+{
+	check(Component);
+	return Component->GetWorld();
+}
+
+bool FPCGComponentExecutionState::HasAuthority() const
+{
+	check(Component && Component->GetOwner());
+	return Component->GetOwner()->HasAuthority();
+}
+
+FBox FPCGComponentExecutionState::GetBounds() const
+{
+	check(Component);
+	return Component->GetGridBounds();
+}
+
+FBox FPCGComponentExecutionState::GetOriginalBounds() const
+{
+	check(Component);
+	return Component->GetOriginalGridBounds();
+}
+
+FBox FPCGComponentExecutionState::GetLocalSpaceBounds() const
+{
+	check(Component);
+	return Component->GetLocalSpaceBounds();
+}
+
+FBox FPCGComponentExecutionState::GetOriginalLocalSpaceBounds() const
+{
+	check(Component);
+	return Component->GetOriginalLocalSpaceBounds();
+}
+
+bool FPCGComponentExecutionState::Use2DGrid() const
+{
+	check(Component);
+	return Component->Use2DGrid();
+}
+
+UPCGGraph* FPCGComponentExecutionState::GetGraph() const
+{
+	check(Component);
+	return Component->GetGraph();
+}
+
+UPCGGraphInstance* FPCGComponentExecutionState::GetGraphInstance() const
+{
+	check(Component);
+	return Component->GetGraphInstance();
+}
+
+void FPCGComponentExecutionState::OnGraphExecutionAborted(bool bQuiet, bool bCleanupUnusedResources)
+{
+	check(Component);
+	Component->OnProcessGraphAborted(bQuiet, bCleanupUnusedResources);
+}
+
+void FPCGComponentExecutionState::Cancel()
+{
+	check(Component);
+	Component->CancelGeneration();
+}
+
+bool FPCGComponentExecutionState::IsGenerating() const
+{
+	check(Component);
+	return Component->IsGenerating();
+}
+
+void FPCGComponentExecutionState::ExecutePreGraph(FPCGContext* InContext)
+{
+	check(Component);
+	Component->ExecutePreGraph(InContext);
+}
+
+bool FPCGComponentExecutionState::IsManagedByRuntimeGenSystem() const
+{
+	check(Component);
+	return Component->IsManagedByRuntimeGenSystem();
+}
+
+IPCGGraphExecutionSource* FPCGComponentExecutionState::GetOriginalSource() const
+{
+	check(Component);
+	return Component->GetOriginalComponent();
+}
+
+IPCGGraphExecutionSource* FPCGComponentExecutionState::GetLocalSource(const FPCGGridDescriptor& GridDescriptor, const FIntVector& CellCoords) const
+{
+	check(Component);
+
+	if (UPCGSubsystem* Subsystem = Cast<UPCGSubsystem>(GetSubsystem()))
+	{
+		return Subsystem->GetLocalComponent(GridDescriptor, CellCoords, CastChecked<UPCGComponent>(GetOriginalSource()));
+	}
+
+	return nullptr;
+}
+
+bool FPCGComponentExecutionState::IsPartitioned() const
+{
+	check(Component);
+	return Component->IsPartitioned();
+}
+
+bool FPCGComponentExecutionState::IsLocalSource() const
+{
+	check(Component);
+	return Component->IsLocalComponent();
+}
+
+FPCGGridDescriptor FPCGComponentExecutionState::GetGridDescriptor(uint32 InGridSize) const
+{
+	check(Component);
+	return Component->GetGridDescriptor(InGridSize);
+}
+
+uint32 FPCGComponentExecutionState::GetGenerationGridSize() const
+{
+	check(Component);
+	return (IsPartitioned() || IsLocalSource()) ? Component->GetGenerationGridSize() : PCGHiGenGrid::UninitializedGridSize();
+}
+
+FPCGTaskId FPCGComponentExecutionState::GetGenerationTaskId() const
+{
+	check(Component);
+	return Component->GetGenerationTaskId();
+}
+
+void FPCGComponentExecutionState::StoreOutputDataForPin(const FString& InResourceKey, const FPCGDataCollection& InData) const
+{
+	check(Component);
+	Component->StoreOutputDataForPin(InResourceKey, InData);
+}
+
+const FPCGDataCollection* FPCGComponentExecutionState::RetrieveOutputDataForPin(const FString& InResourceKey) const
+{
+	check(Component);
+	return Component->RetrieveOutputDataForPin(InResourceKey);
+}
+
+void FPCGComponentExecutionState::AddToManagedResources(UPCGManagedResource* InResource)
+{
+	check(Component);
+	return Component->AddToManagedResources(InResource);
+}
+
+FPCGTaskId FPCGComponentExecutionState::GenerateLocalGetTaskId(EPCGHiGenGrid Grid)
+{
+	check(Component)
+
+	const EPCGComponentGenerationTrigger GenTrigger = IsManagedByRuntimeGenSystem() ? EPCGComponentGenerationTrigger::GenerateAtRuntime : EPCGComponentGenerationTrigger::GenerateOnDemand;
+	return Component->GenerateLocalGetTaskId(GenTrigger, /*bForce=*/true, Grid);
+}
+
+#if WITH_EDITOR
+const PCGUtils::FExtraCapture& FPCGComponentExecutionState::GetExtraCapture() const
+{
+	check(Component);
+	return Component->ExtraCapture;
+}
+
+PCGUtils::FExtraCapture& FPCGComponentExecutionState::GetExtraCapture()
+{
+	check(Component);
+	return Component->ExtraCapture;
+}
+
+const FPCGGraphExecutionInspection& FPCGComponentExecutionState::GetInspection() const
+{
+	check(Component);
+	return Component->ExecutionInspection;
+}
+
+FPCGGraphExecutionInspection& FPCGComponentExecutionState::GetInspection()
+{
+	check(Component);
+	return Component->ExecutionInspection;
+}
+
+void FPCGComponentExecutionState::RegisterDynamicTracking(const UPCGSettings* InSettings, const TArrayView<TPair<FPCGSelectionKey, bool>>& InDynamicKeysAndCulling)
+{
+	check(Component);
+	Component->RegisterDynamicTracking(InSettings, InDynamicKeysAndCulling);
+}
+
+void FPCGComponentExecutionState::RegisterDynamicTracking(const FPCGSelectionKeyToSettingsMap& InKeysToSettings)
+{
+	check(Component);
+	Component->RegisterDynamicTracking(InKeysToSettings);
+}
+
+bool FPCGComponentExecutionState::IsRefreshInProgress() const
+{
+	check(Component);
+	return Component->IsRefreshInProgress();
+}
+
+FPCGDynamicTrackingPriority FPCGComponentExecutionState::GetDynamicTrackingPriority() const
+{
+	check(Component);
+
+	// Hash [0, 1] 
+	double Fraction = static_cast<double>(GetTypeHash(Component->GetPathName())) / std::numeric_limits<uint32>::max();
+	// Divide Fraction by 2^14 so that max value most significant decimal is the 5th (1.0 / 16384 -> 0.00006103515625)
+	// TrackingPriority is rounded to 4 decimals so that Fraction (Hash) doesn't interfere with it (or almost never, in which case difference in value would be minimal).
+	Fraction /= 16384;
+
+	// So that leading integer and fractional user part stay the same add or subtract the hash based fraction based on sign
+	return Component->TrackingPriority >= 0 ? (double)Component->TrackingPriority + Fraction : (double)Component->TrackingPriority - Fraction;
+}
+
+#endif // WITH_EDITOR
diff --git a/Engine/Plugins/PCG/Source/PCG/Public/PCGComponent.h b/Engine/Plugins/PCG/Source/PCG/Public/PCGComponent.h
--- a/Engine/Plugins/PCG/Source/PCG/Public/PCGComponent.h
+++ b/Engine/Plugins/PCG/Source/PCG/Public/PCGComponent.h
@@ -3,6 +3,7 @@
 #pragma once
 
 #include "PCGCommon.h"
+#include "PCGComponentExecutionState.h"
 #include "PCGGraphExecutionInspection.h"
 #include "PCGGraphExecutionStateInterface.h"
 #include "PCGNode.h"
@@ -102,60 +103,6 @@ namespace PCGTestsCommon
 	struct FTestData;
 }
 
-class FPCGComponentExecutionState : public IPCGGraphExecutionState
-{
-	friend class UPCGComponent;
-
-public:
-	virtual UPCGData* GetSelfData() const override;
-	virtual int32 GetSeed() const override;
-	virtual FString GetDebugName() const override;
-	virtual FTransform GetTransform() const override;
-	virtual FTransform GetOriginalTransform() const override;
-	virtual UWorld* GetWorld() const override;
-	virtual bool HasAuthority() const override;
-	virtual FBox GetBounds() const override;
-	virtual FBox GetOriginalBounds() const override;
-	virtual FBox GetLocalSpaceBounds() const override;
-	virtual FBox GetOriginalLocalSpaceBounds() const override;
-	virtual bool Use2DGrid() const override;
-	virtual UPCGGraph* GetGraph() const override;
-	virtual UPCGGraphInstance* GetGraphInstance() const override;
-	virtual void OnGraphExecutionAborted(bool bQuiet, bool bCleanupUnusedResources) override;
-	virtual void Cancel() override;
-	virtual bool IsGenerating() const override;
-	virtual void ExecutePreGraph(FPCGContext* InContext) override;
-	virtual bool IsManagedByRuntimeGenSystem() const override;
-	virtual IPCGGraphExecutionSource* GetOriginalSource() const override;
-	virtual IPCGGraphExecutionSource* GetLocalSource(const FPCGGridDescriptor& GridDescriptor, const FIntVector& CellCoords) const override;
-	virtual bool IsPartitioned() const override;
-	virtual bool IsLocalSource() const override;
-	virtual FPCGGridDescriptor GetGridDescriptor(uint32 InGridSize) const override;
-	virtual uint32 GetGenerationGridSize() const override;
-	virtual FPCGTaskId GetGenerationTaskId() const override;
-	virtual void StoreOutputDataForPin(const FString& InResourceKey, const FPCGDataCollection& InData) const override;
-	virtual const FPCGDataCollection* RetrieveOutputDataForPin(const FString& InResourceKey) const override;
-	virtual void AddToManagedResources(UPCGManagedResource* InResource) override;
-	virtual FPCGTaskId GenerateLocalGetTaskId(EPCGHiGenGrid Grid = EPCGHiGenGrid::Uninitialized) override;
-
-#if WITH_EDITOR
-	virtual const PCGUtils::FExtraCapture& GetExtraCapture() const override;
-	virtual PCGUtils::FExtraCapture& GetExtraCapture() override;
-
-	virtual const FPCGGraphExecutionInspection& GetInspection() const override;
-	virtual FPCGGraphExecutionInspection& GetInspection() override;
-
-	virtual void RegisterDynamicTracking(const UPCGSettings* InSettings, const TArrayView<TPair<FPCGSelectionKey, bool>>& InDynamicKeysAndCulling) override;
-	virtual void RegisterDynamicTracking(const FPCGSelectionKeyToSettingsMap& InKeysToSettings) override;
-	
-	virtual bool IsRefreshInProgress() const override;
-	virtual FPCGDynamicTrackingPriority GetDynamicTrackingPriority() const override;
-#endif // WITH_EDITOR
-
-private:
-	UPCGComponent* Component = nullptr;
-};
-
 UCLASS(MinimalAPI, BlueprintType, ClassGroup = (Procedural), meta = (BlueprintSpawnableComponent, PrioritizeCategories = "PCG"))
 class UPCGComponent : public UActorComponent, public IPCGGraphExecutionSource
 {
diff --git a/Engine/Plugins/PCG/Source/PCG/Public/PCGComponentExecutionState.h b/Engine/Plugins/PCG/Source/PCG/Public/PCGComponentExecutionState.h
new file mode 100644
--- /dev/null
+++ b/Engine/Plugins/PCG/Source/PCG/Public/PCGComponentExecutionState.h
@@ -0,0 +1,61 @@
+// Copyright Epic Games, Inc. All Rights Reserved.
+
+#pragma once
+
+#include "PCGGraphExecutionStateInterface.h"
+
+class UPCGComponent;
+
+class FPCGComponentExecutionState : public IPCGGraphExecutionState
+{
+	friend class UPCGComponent;
+
+public:
+	virtual UPCGData* GetSelfData() const override;
+	virtual int32 GetSeed() const override;
+	virtual FString GetDebugName() const override;
+	virtual FTransform GetTransform() const override;
+	virtual FTransform GetOriginalTransform() const override;
+	virtual UWorld* GetWorld() const override;
+	virtual bool HasAuthority() const override;
+	virtual FBox GetBounds() const override;
+	virtual FBox GetOriginalBounds() const override;
+	virtual FBox GetLocalSpaceBounds() const override;
+	virtual FBox GetOriginalLocalSpaceBounds() const override;
+	virtual bool Use2DGrid() const override;
+	virtual UPCGGraph* GetGraph() const override;
+	virtual UPCGGraphInstance* GetGraphInstance() const override;
+	virtual void OnGraphExecutionAborted(bool bQuiet, bool bCleanupUnusedResources) override;
+	virtual void Cancel() override;
+	virtual bool IsGenerating() const override;
+	virtual void ExecutePreGraph(FPCGContext* InContext) override;
+	virtual bool IsManagedByRuntimeGenSystem() const override;
+	virtual IPCGGraphExecutionSource* GetOriginalSource() const override;
+	virtual IPCGGraphExecutionSource* GetLocalSource(const FPCGGridDescriptor& GridDescriptor, const FIntVector& CellCoords) const override;
+	virtual bool IsPartitioned() const override;
+	virtual bool IsLocalSource() const override;
+	virtual FPCGGridDescriptor GetGridDescriptor(uint32 InGridSize) const override;
+	virtual uint32 GetGenerationGridSize() const override;
+	virtual FPCGTaskId GetGenerationTaskId() const override;
+	virtual void StoreOutputDataForPin(const FString& InResourceKey, const FPCGDataCollection& InData) const override;
+	virtual const FPCGDataCollection* RetrieveOutputDataForPin(const FString& InResourceKey) const override;
+	virtual void AddToManagedResources(UPCGManagedResource* InResource) override;
+	virtual FPCGTaskId GenerateLocalGetTaskId(EPCGHiGenGrid Grid = EPCGHiGenGrid::Uninitialized) override;
+
+#if WITH_EDITOR
+	virtual const PCGUtils::FExtraCapture& GetExtraCapture() const override;
+	virtual PCGUtils::FExtraCapture& GetExtraCapture() override;
+
+	virtual const FPCGGraphExecutionInspection& GetInspection() const override;
+	virtual FPCGGraphExecutionInspection& GetInspection() override;
+
+	virtual void RegisterDynamicTracking(const UPCGSettings* InSettings, const TArrayView<TPair<FPCGSelectionKey, bool>>& InDynamicKeysAndCulling) override;
+	virtual void RegisterDynamicTracking(const FPCGSelectionKeyToSettingsMap& InKeysToSettings) override;
+	
+	virtual bool IsRefreshInProgress() const override;
+	virtual FPCGDynamicTrackingPriority GetDynamicTrackingPriority() const override;
+#endif // WITH_EDITOR
+
+private:
+	UPCGComponent* Component = nullptr;
+};
