[PCG] Add non-spatial world streaming query to RGS

#rnx

diff --git a/Engine/Plugins/PCG/Source/PCG/Private/RuntimeGen/PCGRuntimeGenScheduler.cpp b/Engine/Plugins/PCG/Source/PCG/Private/RuntimeGen/PCGRuntimeGenScheduler.cpp
--- a/Engine/Plugins/PCG/Source/PCG/Private/RuntimeGen/PCGRuntimeGenScheduler.cpp
+++ b/Engine/Plugins/PCG/Source/PCG/Private/RuntimeGen/PCGRuntimeGenScheduler.cpp
@@ -402,6 +402,33 @@ bool FPCGRuntimeGenScheduler::ShouldTick()
 		return false;
 	}
 
+	// Check if the non-spatially streamed cell is ready (only if world streaming queries are enabled). We should not do any processing until the world is ready.
+	if (!bNonSpatialCellFinishedStreaming)
+	{
+		const UWorldPartitionSubsystem* WorldPartitionSubsystem = PCGRuntimeGenSchedulerHelpers::CVarEnableWorldStreamingQueries.GetValueOnGameThread() ? UWorld::GetSubsystem<UWorldPartitionSubsystem>(World) : nullptr;
+
+		if (WorldPartitionSubsystem)
+		{
+			TRACE_CPUPROFILER_EVENT_SCOPE(IsNonSpatialWorldStreamingComplete);
+
+			FWorldPartitionStreamingQuerySource NonSpatialStreamingQuerySource;
+			NonSpatialStreamingQuerySource.bSpatialQuery = false;
+
+			bNonSpatialCellFinishedStreaming = WorldPartitionSubsystem->IsStreamingCompleted(EWorldPartitionRuntimeCellState::Activated, { NonSpatialStreamingQuerySource }, /*bExactState*/false);
+
+			if (!bNonSpatialCellFinishedStreaming)
+			{
+				UE_LOG(LogPCG, Verbose, TEXT("Holding back tick of Runtime Gen Scheduler because the world is not yet loaded."));
+				return false;
+			}
+		}
+		else
+		{
+			// Do nothing if we are not in a WP level. The persistent level should already be loaded by now.
+			bNonSpatialCellFinishedStreaming = true;
+		}
+	}
+
 	return true;
 }
 
diff --git a/Engine/Plugins/PCG/Source/PCG/Public/RuntimeGen/PCGRuntimeGenScheduler.h b/Engine/Plugins/PCG/Source/PCG/Public/RuntimeGen/PCGRuntimeGenScheduler.h
--- a/Engine/Plugins/PCG/Source/PCG/Public/RuntimeGen/PCGRuntimeGenScheduler.h
+++ b/Engine/Plugins/PCG/Source/PCG/Public/RuntimeGen/PCGRuntimeGenScheduler.h
@@ -258,6 +258,9 @@ private:
 	using FStreamingCompleteQueryKey = TPair<FVector, float>;
 	TMap<FStreamingCompleteQueryKey, bool> CachedStreamingQueryResults;
 
+	/** We should hold back generation until the non-spatial cell is done streaming, which happens once when the level loads. */
+	bool bNonSpatialCellFinishedStreaming = false;
+
 	/** Used to detect when world state (generation sources, cvars etc) have changed enough to warrant rescanning generation cells. */
 	PCGRuntimeGenChangeDetection::FDetector ChangeDetector;
 
