[PCG] Allow Log capture without a PcgElement or context

#rnx

diff --git a/Engine/Plugins/PCG/Source/PCG/Private/PCGElement.cpp b/Engine/Plugins/PCG/Source/PCG/Private/PCGElement.cpp
--- a/Engine/Plugins/PCG/Source/PCG/Private/PCGElement.cpp
+++ b/Engine/Plugins/PCG/Source/PCG/Private/PCGElement.cpp
@@ -212,7 +212,7 @@ bool IPCGElement::Execute(FPCGContext* Context) const
 
 	while (Context->CurrentPhase != EPCGExecutionPhase::Done)
 	{
-		PCGUtils::FScopedCall ScopedCall(*this, Context, OutputDevice);
+		PCGUtils::FScopedCall ScopedCall(this, Context, OutputDevice);
 		bool bExecutionPostponed = false;
 
 		switch (Context->CurrentPhase)
diff --git a/Engine/Plugins/PCG/Source/PCG/Private/Utils/PCGExtraCapture.cpp b/Engine/Plugins/PCG/Source/PCG/Private/Utils/PCGExtraCapture.cpp
--- a/Engine/Plugins/PCG/Source/PCG/Private/Utils/PCGExtraCapture.cpp
+++ b/Engine/Plugins/PCG/Source/PCG/Private/Utils/PCGExtraCapture.cpp
@@ -261,10 +261,10 @@ PCGUtils::FCallTreeInfo PCGUtils::FExtraCapture::CalculateCallTreeInfo(const IPC
 	return RootInfo;
 }
 
-PCGUtils::FScopedCall::FScopedCall(const IPCGElement& InOwner, FPCGContext* InContext, FScopedCallOutputDevice& InOutputDevice)
+PCGUtils::FScopedCall::FScopedCall(const IPCGElement* InOwner, FPCGContext* InContext, FScopedCallOutputDevice& InOutputDevice)
 	: Owner(InOwner)
 	, Context(InContext)
-	, Phase(InContext->CurrentPhase)
+	, Phase(InContext ? InContext->CurrentPhase : EPCGExecutionPhase::NotExecuted)
 	, OutputDevice(InOutputDevice)
 {
 	StartTime = FPlatformTime::Seconds();
diff --git a/Engine/Plugins/PCG/Source/PCG/Public/Utils/PCGExtraCapture.h b/Engine/Plugins/PCG/Source/PCG/Public/Utils/PCGExtraCapture.h
--- a/Engine/Plugins/PCG/Source/PCG/Public/Utils/PCGExtraCapture.h
+++ b/Engine/Plugins/PCG/Source/PCG/Public/Utils/PCGExtraCapture.h
@@ -62,14 +62,14 @@ namespace PCGUtils
 
 	struct FScopedCallOutputDevice : public FOutputDevice
 	{
-		FScopedCallOutputDevice();
-		virtual ~FScopedCallOutputDevice();
+		UE_API FScopedCallOutputDevice();
+		UE_API virtual ~FScopedCallOutputDevice();
 
 		// FOutputDevice
 		virtual bool IsMemoryOnly() const override { return true; }
 		virtual bool CanBeUsedOnMultipleThreads() const override { return true; }
 		virtual bool CanBeUsedOnAnyThread() const override { return true; }
-		virtual void Serialize(const TCHAR* V, ELogVerbosity::Type Verbosity, const class FName& Category) override;
+		UE_API virtual void Serialize(const TCHAR* V, ELogVerbosity::Type Verbosity, const class FName& Category) override;
 
 		FScopedCall* ScopedCall = nullptr;
 		const uint32 ThreadID;
@@ -77,10 +77,10 @@ namespace PCGUtils
 
 	struct FScopedCall
 	{
-		FScopedCall(const IPCGElement& InOwner, FPCGContext* InContext, FScopedCallOutputDevice& InOutputDevice);
-		virtual ~FScopedCall();
+		UE_API FScopedCall(const IPCGElement* InOwner, FPCGContext* InContext, FScopedCallOutputDevice& InOutputDevice);
+		UE_API virtual ~FScopedCall();
 
-		const IPCGElement& Owner;
+		const IPCGElement* Owner;
 		FPCGContext* Context;
 		double StartTime;
 		EPCGExecutionPhase Phase;
@@ -113,7 +113,7 @@ namespace PCGUtils
 
 	struct FScopedCall
 	{
-		FScopedCall(const IPCGElement& InOwner, FPCGContext* InContext, FScopedCallOutputDevice& InOutputDevice) {}
+		FScopedCall(const IPCGElement* InOwner, FPCGContext* InContext, FScopedCallOutputDevice& InOutputDevice) {}
 	};
 #endif // WITH_EDITOR
 }
