[PCG] Tangent-aware spline offset

diff --git a/Engine/Plugins/PCG/Source/PCG/Private/Elements/Spline/PCGOffsetSpline.cpp b/Engine/Plugins/PCG/Source/PCG/Private/Elements/Spline/PCGOffsetSpline.cpp
new file mode 100644
--- /dev/null
+++ b/Engine/Plugins/PCG/Source/PCG/Private/Elements/Spline/PCGOffsetSpline.cpp
@@ -0,0 +1,114 @@
+// Copyright Epic Games, Inc. All Rights Reserved.
+
+#include "Elements/Spline/PCGOffsetSpline.h"
+
+#include "Data/PCGSplineData.h"
+#include "Metadata/Accessors/PCGAttributeAccessorHelpers.h"
+
+namespace PCGOffsetSpline
+{
+	const FLazyName OffsetMagnitudeLabel = TEXT("Magnitude");
+}
+
+TArray<FPCGPinProperties> UPCGOffsetSplineSettings::InputPinProperties() const
+{
+	TArray<FPCGPinProperties> PinProperties;
+	PinProperties.Emplace_GetRef(PCGPinConstants::DefaultInputLabel, EPCGDataType::Spline).SetRequiredPin();
+	PinProperties.Emplace_GetRef(PCGOffsetSpline::OffsetMagnitudeLabel, EPCGDataType::Spline | EPCGDataType::Param).SetRequiredPin();
+
+	return PinProperties;
+}
+
+TArray<FPCGPinProperties> UPCGOffsetSplineSettings::OutputPinProperties() const
+{
+	TArray<FPCGPinProperties> PinProperties;
+	PinProperties.Emplace(PCGPinConstants::DefaultOutputLabel, EPCGDataType::Spline);
+
+	return PinProperties;
+}
+
+bool FPCGOffsetSplineElement::ExecuteInternal(FPCGContext* Context) const
+{
+	TRACE_CPUPROFILER_EVENT_SCOPE(FPCGOffsetSplineElement::Execute);
+
+	check(Context);
+
+	const UPCGOffsetSplineSettings* Settings = Context->GetInputSettings<UPCGOffsetSplineSettings>();
+	check(Settings);
+
+	// 1. Get inputs
+	TArray<FPCGTaggedData> Inputs = Context->InputData.GetInputsByPin(PCGPinConstants::DefaultInputLabel);
+	TArray<FPCGTaggedData> MagnitudeInputs = Context->InputData.GetInputsByPin(PCGOffsetSpline::OffsetMagnitudeLabel);
+
+	// 2. Validate cardinality
+	if (MagnitudeInputs.Num() != 1 && MagnitudeInputs.Num() != Inputs.Num())
+	{
+		PCGLog::InputOutput::LogInvalidCardinalityError(PCGPinConstants::DefaultInputLabel, PCGOffsetSpline::OffsetMagnitudeLabel, Context);
+		return true;
+	}
+
+	// 3. getting all the direction + magnitude data upfront (can be through extract all values)
+	// 4. then move the control points, then compare the before/after distances and scale the tangents accordingly.
+	for (int InputIndex = 0; InputIndex < Inputs.Num(); ++InputIndex)
+	{
+		const FPCGTaggedData& Input = Inputs[InputIndex];
+		const UPCGSplineData* Spline = Cast<UPCGSplineData>(Input.Data);
+
+		if (!Spline)
+		{
+			PCGLog::InputOutput::LogInvalidInputDataError(Context);
+			continue;
+		}
+
+		TArray<FVector> Directions;
+		if (!PCGAttributeAccessorHelpers::ExtractAllValues(Spline, Settings->DirectionAttribute, Directions, Context) || Directions.IsEmpty())
+		{
+			continue;
+		}
+
+		const FPCGTaggedData& MagnitudeInput = MagnitudeInputs[InputIndex % MagnitudeInputs.Num()];
+
+		TArray<double> Magnitudes;
+		if (!PCGAttributeAccessorHelpers::ExtractAllValues(MagnitudeInput.Data, Settings->MagnitudeAttribute, Magnitudes, Context) || Magnitudes.IsEmpty())
+		{
+			continue;
+		}
+
+		UPCGSplineData* OffsetSpline = FPCGContext::NewObject_AnyThread<UPCGSplineData>(Context);
+		OffsetSpline->InitializeFromData(Spline);
+
+		TArray<FSplinePoint> SplinePoints = Spline->GetSplinePoints();
+		TArray<PCGMetadataEntryKey> SplinePointsEntryKeys;
+
+		TConstArrayView<PCGMetadataEntryKey> ConstEntryKeys = Spline->GetConstVerticesEntryKeys();
+		SplinePointsEntryKeys.Append(ConstEntryKeys);
+
+		TArray<double> PreviousDistances;
+		const int NumSegments = Spline->GetNumSegments();
+		PreviousDistances.Reserve(SplinePoints.Num());
+		for (int SegmentIndex = 0; SegmentIndex < NumSegments; ++SegmentIndex)
+		{
+			PreviousDistances.Add((SplinePoints[SegmentIndex].Position - SplinePoints[(SegmentIndex + 1) % SplinePoints.Num()].Position).Length());
+		}
+
+		for (int PointIndex = 0; PointIndex < SplinePoints.Num(); ++PointIndex)
+		{
+			SplinePoints[PointIndex].Position += Spline->SplineStruct.Transform.InverseTransformVector(Directions[PointIndex % Directions.Num()]) * Magnitudes[PointIndex % Magnitudes.Num()];
+		}
+
+		for (int SegmentIndex = 0; SegmentIndex < NumSegments; ++SegmentIndex)
+		{
+			const double NewDistance = (SplinePoints[SegmentIndex].Position - SplinePoints[(SegmentIndex + 1) % SplinePoints.Num()].Position).Length();
+
+			const double Ratio = (PreviousDistances[SegmentIndex] > UE_DOUBLE_SMALL_NUMBER) ? NewDistance / PreviousDistances[SegmentIndex] : 1.0;
+			SplinePoints[SegmentIndex].LeaveTangent *= Ratio;
+			SplinePoints[(SegmentIndex + 1) % SplinePoints.Num()].ArriveTangent *= Ratio;
+		}
+
+		OffsetSpline->Initialize(MoveTemp(SplinePoints), Spline->IsClosed(), Spline->GetTransform(), MoveTemp(SplinePointsEntryKeys));
+
+		Context->OutputData.TaggedData.Add_GetRef(Input).Data = OffsetSpline;
+	}
+
+	return true;
+}
\ No newline at end of file
diff --git a/Engine/Plugins/PCG/Source/PCG/Public/Elements/Spline/PCGOffsetSpline.h b/Engine/Plugins/PCG/Source/PCG/Public/Elements/Spline/PCGOffsetSpline.h
new file mode 100644
--- /dev/null
+++ b/Engine/Plugins/PCG/Source/PCG/Public/Elements/Spline/PCGOffsetSpline.h
@@ -0,0 +1,48 @@
+// Copyright Epic Games, Inc. All Rights Reserved.
+
+#pragma once
+
+#include "PCGContext.h"
+#include "PCGSettings.h"
+
+#include "PCGOffsetSpline.generated.h"
+
+class FPCGOffsetSplineElement : public IPCGElement
+{
+public:
+	virtual EPCGElementExecutionLoopMode ExecutionLoopMode(const UPCGSettings* Settings) const override { return EPCGElementExecutionLoopMode::PrimaryPinAndBroadcastablePins; }
+	virtual bool ExecuteInternal(FPCGContext* Context) const;
+};
+
+/**
+* Offsets a spline based on provided direction & magnitude, and corrects tangents lengths after the fact.
+* Note that this node is not currently robust to extreme offsets that would make the spline self-intersect.
+*/
+UCLASS(MinimalAPI, BlueprintType, ClassGroup = (Procedural))
+class UPCGOffsetSplineSettings : public UPCGSettings
+{
+	GENERATED_BODY()
+
+#if WITH_EDITOR
+	virtual FName GetDefaultNodeName() const override { return FName(TEXT("OffsetSpline")); }
+	virtual FText GetDefaultNodeTitle() const override { return NSLOCTEXT("PCGOffsetSplineElement", "NodeTitle", "Offset Spline"); }
+	virtual FText GetNodeTooltipText() const override { return NSLOCTEXT("PCGOffsetSplineElement", "NodeTooltip", "Offsets spline control points according to a direction and a magnitude."); }
+	virtual EPCGSettingsType GetType() const override { return EPCGSettingsType::Spatial; }
+#endif // WITH_EDITOR
+
+protected:
+	virtual TArray<FPCGPinProperties> InputPinProperties() const override;
+	virtual TArray<FPCGPinProperties> OutputPinProperties() const override;
+
+	virtual FPCGElementPtr CreateElement() const override { return MakeShared<FPCGOffsetSplineElement>(); }
+	// ~End UPCGSettings interface implementation
+
+public:
+	/** Specifies what attribute from the input that will drive the offset direction. */
+	UPROPERTY(BlueprintReadWrite, EditAnywhere, Category = Settings, meta = (PCG_Overridable))
+	FPCGAttributePropertyInputSelector DirectionAttribute;
+
+	/** Specifies what attribute from the magnitude input that will drive the offset magnitude. */
+	UPROPERTY(BlueprintReadWrite, EditAnywhere, Category = Settings, meta = (PCG_Overridable))
+	FPCGAttributePropertyInputSelector MagnitudeAttribute;
+};
\ No newline at end of file
