[PCG] Add shader feature level switch node.

[RN] [PCG] Added Max Shader Feature Level switch node

diff --git a/Engine/Plugins/PCG/Source/PCG/Private/Elements/ControlFlow/PCGPlatformSwitch.cpp b/Engine/Plugins/PCG/Source/PCG/Private/Elements/ControlFlow/PCGPlatformSwitch.cpp
--- a/Engine/Plugins/PCG/Source/PCG/Private/Elements/ControlFlow/PCGPlatformSwitch.cpp
+++ b/Engine/Plugins/PCG/Source/PCG/Private/Elements/ControlFlow/PCGPlatformSwitch.cpp
@@ -3,11 +3,9 @@
 #include "Elements/ControlFlow/PCGPlatformSwitch.h"
 
 #include "PCGContext.h"
-#include "PCGGraph.h"
 #include "PCGModule.h"
 #include "Elements/PCGGather.h"
 #include "Helpers/PCGHelpers.h"
-#include "Subsystems/PCGSubsystem.h"
 
 #include "DeviceProfiles/DeviceProfile.h"
 #include "DeviceProfiles/DeviceProfileManager.h"
@@ -29,9 +27,9 @@ namespace PCGPlatformSwitch
 	static const FName DefaultLabel = TEXT("Default");
 
 #if WITH_EDITOR
-	bool GetCurrentPlatformInfoFromEditor(FName& OutPlatform, FName& OutPlatformGroup)
+	static bool GetCurrentPlatformInfoFromEditor(FName& OutPlatform, FName& OutPlatformGroup)
 	{
-		// Ask the editor for the active preview platform (PS5, XSX, etc.)
+		// Ask the editor for the active preview platform
 		FName IniPlatform;
 		if (GEditor && GEditor->GetPreviewPlatformName(IniPlatform))
 		{
@@ -42,7 +40,7 @@ namespace PCGPlatformSwitch
 
 		// Single active target platform selected by the user (Project Launcher / Platforms UI)
 		const TArray<ITargetPlatform*>& ActiveTPs = GetTargetPlatformManagerRef().GetActiveTargetPlatforms();
-		if (ActiveTPs.Num() == 1)
+		if (ActiveTPs.Num() == 1 && ActiveTPs[0])
 		{
 			OutPlatform = ActiveTPs[0]->GetPlatformInfo().IniPlatformName;
 			OutPlatformGroup = ActiveTPs[0]->GetPlatformInfo().PlatformGroupName;
@@ -52,7 +50,7 @@ namespace PCGPlatformSwitch
 		// Device Profile (covers Mobile Preview / device-specific previews in-editor)
 		if (UDeviceProfile* ActiveProfile = UDeviceProfileManager::Get().GetActiveProfile())
 		{
-			// DeviceType typically matches the Ini platform name ("Windows", "Android", "PS5", ...)
+			// DeviceType typically matches the Ini platform name
 			OutPlatform = *ActiveProfile->DeviceType;
 			OutPlatformGroup = FDataDrivenPlatformInfoRegistry::GetPlatformInfo(OutPlatform).PlatformGroupName;
 			return true;
@@ -124,34 +122,6 @@ FString UPCGPlatformSwitchSettings::GetAdditionalTitleInformation() const
 }
 
 #if WITH_EDITOR
-void UPCGPlatformSwitchSettings::PostInitProperties()
-{
-	Super::PostInitProperties();
-
-	SetupPreviewPlatformChangeEvents();
-}
-
-void UPCGPlatformSwitchSettings::PostLoad()
-{
-	Super::PostLoad();
-
-	SetupPreviewPlatformChangeEvents();
-}
-
-void UPCGPlatformSwitchSettings::PostDuplicate(bool bDuplicateForPIE)
-{
-	Super::PostDuplicate(bDuplicateForPIE);
-
-	SetupPreviewPlatformChangeEvents();
-}
-
-void UPCGPlatformSwitchSettings::PostEditImport()
-{
-	Super::PostEditImport();
-
-	SetupPreviewPlatformChangeEvents();
-}
-
 void UPCGPlatformSwitchSettings::PreSave(FObjectPreSaveContext ObjectSaveContext)
 {
 	Super::PreSave(ObjectSaveContext);
@@ -162,13 +132,6 @@ void UPCGPlatformSwitchSettings::PreSave(FObjectPreSaveContext ObjectSaveContext
 	}
 }
 
-void UPCGPlatformSwitchSettings::BeginDestroy()
-{
-	TearDownPreviewPlatformChangeEvents();
-
-	Super::BeginDestroy();
-}
-
 FText UPCGPlatformSwitchSettings::GetNodeTooltipText() const
 {
 	return LOCTEXT("NodeTooltip", "Statically activates/deactivates output pins based on the current platform or platform group (desktop/console/mobile), useful for culling downstream nodes.");
@@ -187,50 +150,6 @@ EPCGChangeType UPCGPlatformSwitchSettings::GetChangeTypeForProperty(const FName&
 	return ChangeType;
 }
 
-void UPCGPlatformSwitchSettings::SetupPreviewPlatformChangeEvents()
-{
-	if (GEditor)
-	{
-		if (!PreviewPlatformChangedHandle.IsValid())
-		{
-			PreviewPlatformChangedHandle = GEditor->OnPreviewPlatformChanged().AddUObject(this, &UPCGPlatformSwitchSettings::OnPreviewPlatformChanged);
-		}
-
-		if (!PreviewFeatureLevelChangedHandle.IsValid())
-		{
-			PreviewFeatureLevelChangedHandle = GEditor->OnPreviewFeatureLevelChanged().AddUObject(this, &UPCGPlatformSwitchSettings::OnPreviewFeatureLevelChanged);
-		}
-	}
-}
-
-void UPCGPlatformSwitchSettings::TearDownPreviewPlatformChangeEvents()
-{
-	if (GEditor)
-	{
-		if (PreviewPlatformChangedHandle.IsValid())
-		{
-			GEditor->OnPreviewPlatformChanged().Remove(PreviewPlatformChangedHandle);
-		}
-
-		if (PreviewFeatureLevelChangedHandle.IsValid())
-		{
-			GEditor->OnPreviewFeatureLevelChanged().Remove(PreviewFeatureLevelChangedHandle);
-		}
-	}
-}
-
-void UPCGPlatformSwitchSettings::OnPreviewPlatformChanged()
-{
-	// We don't broadcast structural from here. The PCGSubsystem will refresh all runtime gen components when preview platform changes.
-	OnSettingsChangedDelegate.Broadcast(this, EPCGChangeType::Cosmetic);
-}
-
-void UPCGPlatformSwitchSettings::OnPreviewFeatureLevelChanged(ERHIFeatureLevel::Type)
-{
-	// We don't broadcast structural from here. The PCGSubsystem will refresh all runtime gen components when preview platform changes.
-	OnSettingsChangedDelegate.Broadcast(this, EPCGChangeType::Cosmetic);
-}
-
 TArray<FName> UPCGPlatformSwitchSettings::GetPlatformOptions() const
 {
 	const TArray<const FDataDrivenPlatformInfo*>& SortedPlatforms = FDataDrivenPlatformInfoRegistry::GetSortedPlatformInfos(EPlatformInfoType::TruePlatformsOnly);
@@ -270,7 +189,7 @@ TArray<FPCGPinProperties> UPCGPlatformSwitchSettings::OutputPinProperties() cons
 {
 	TArray<FPCGPinProperties> PinProperties;
 
-	PinProperties.Emplace(TEXT("Default"), EPCGDataType::Any);
+	PinProperties.Emplace(PCGPlatformSwitch::DefaultLabel, EPCGDataType::Any);
 
 	for (FName Platform : GetSanitizedPlatforms())
 	{
@@ -300,9 +219,8 @@ bool UPCGPlatformSwitchSettings::IsPlatformSelected(FName InPlatform) const
 bool UPCGPlatformSwitchSettings::IsPlatformSelected(FName InPlatform, FName InCurrentPlatform, FName InCurrentPlatformGroup) const
 {
 	ensure(!InCurrentPlatform.IsNone());
-	ensure(!InCurrentPlatformGroup.IsNone());
 
-	return InPlatform == InCurrentPlatform || InPlatform == InCurrentPlatformGroup;
+	return InPlatform == InCurrentPlatform || (InPlatform == InCurrentPlatformGroup && !InCurrentPlatformGroup.IsNone());
 }
 
 bool UPCGPlatformSwitchSettings::GetCurrentPlatformInfo(FName& OutPlatform, FName& OutPlatformGroup) const
@@ -332,16 +250,6 @@ TArray<FName> UPCGPlatformSwitchSettings::GetSanitizedPlatforms() const
 	return Platforms;
 }
 
-bool UPCGPlatformSwitchSettings::IsPinUsedByNodeExecution(const UPCGPin* InPin) const
-{
-	if (!InPin->IsOutputPin())
-	{
-		return Super::IsPinUsedByNodeExecution(InPin);
-	}
-
-	return IsPinStaticallyActive(InPin->Properties.Label);
-}
-
 bool FPCGPlatformSwitchElement::ExecuteInternal(FPCGContext* InContext) const
 {
 	TRACE_CPUPROFILER_EVENT_SCOPE(FPCGPlatformSwitchElement::ExecuteInternal);
@@ -360,7 +268,7 @@ bool FPCGPlatformSwitchElement::ExecuteInternal(FPCGContext* InContext) const
 		return true;
 	}
 
-	bool bAnyPlatformSelected = false;
+	bool bAnyOutputSelected = false;
 
 	// Copy of data for each platform output.
 	for (FName Platform : Platforms)
@@ -373,11 +281,11 @@ bool FPCGPlatformSwitchElement::ExecuteInternal(FPCGContext* InContext) const
 				InContext->OutputData.TaggedData.Add(TaggedData);
 			}
 
-			bAnyPlatformSelected = true;
+			bAnyOutputSelected = true;
 		}
 	}
 
-	if (!bAnyPlatformSelected)
+	if (!bAnyOutputSelected)
 	{
 		// Fall back to Default pin.
 		InContext->OutputData.TaggedData.Append(InputData.TaggedData);
diff --git a/Engine/Plugins/PCG/Source/PCG/Private/Elements/ControlFlow/PCGPlatformSwitchBase.cpp b/Engine/Plugins/PCG/Source/PCG/Private/Elements/ControlFlow/PCGPlatformSwitchBase.cpp
new file mode 100644
--- /dev/null
+++ b/Engine/Plugins/PCG/Source/PCG/Private/Elements/ControlFlow/PCGPlatformSwitchBase.cpp
@@ -0,0 +1,100 @@
+// Copyright Epic Games, Inc. All Rights Reserved.
+
+#include "Elements/ControlFlow/PCGPlatformSwitchBase.h"
+
+#if WITH_EDITOR
+#include "Editor.h"
+#endif
+
+#include UE_INLINE_GENERATED_CPP_BY_NAME(PCGPlatformSwitchBase)
+
+#if WITH_EDITOR
+void UPCGPlatformSwitchSettingsBase::PostInitProperties()
+{
+	Super::PostInitProperties();
+
+	SetupPreviewPlatformChangeEvents();
+}
+
+void UPCGPlatformSwitchSettingsBase::PostLoad()
+{
+	Super::PostLoad();
+
+	SetupPreviewPlatformChangeEvents();
+}
+
+void UPCGPlatformSwitchSettingsBase::PostDuplicate(bool bDuplicateForPIE)
+{
+	Super::PostDuplicate(bDuplicateForPIE);
+
+	SetupPreviewPlatformChangeEvents();
+}
+
+void UPCGPlatformSwitchSettingsBase::PostEditImport()
+{
+	Super::PostEditImport();
+
+	SetupPreviewPlatformChangeEvents();
+}
+
+void UPCGPlatformSwitchSettingsBase::BeginDestroy()
+{
+	TearDownPreviewPlatformChangeEvents();
+
+	Super::BeginDestroy();
+}
+
+void UPCGPlatformSwitchSettingsBase::SetupPreviewPlatformChangeEvents()
+{
+	if (GEditor)
+	{
+		if (!PreviewPlatformChangedHandle.IsValid())
+		{
+			PreviewPlatformChangedHandle = GEditor->OnPreviewPlatformChanged().AddUObject(this, &UPCGPlatformSwitchSettingsBase::OnPreviewPlatformChanged);
+		}
+
+		if (!PreviewFeatureLevelChangedHandle.IsValid())
+		{
+			PreviewFeatureLevelChangedHandle = GEditor->OnPreviewFeatureLevelChanged().AddUObject(this, &UPCGPlatformSwitchSettingsBase::OnPreviewFeatureLevelChanged);
+		}
+	}
+}
+
+void UPCGPlatformSwitchSettingsBase::TearDownPreviewPlatformChangeEvents()
+{
+	if (GEditor)
+	{
+		if (PreviewPlatformChangedHandle.IsValid())
+		{
+			GEditor->OnPreviewPlatformChanged().Remove(PreviewPlatformChangedHandle);
+		}
+
+		if (PreviewFeatureLevelChangedHandle.IsValid())
+		{
+			GEditor->OnPreviewFeatureLevelChanged().Remove(PreviewFeatureLevelChangedHandle);
+		}
+	}
+}
+
+void UPCGPlatformSwitchSettingsBase::OnPreviewPlatformChanged()
+{
+	// We don't broadcast structural from here. The PCGSubsystem will refresh all runtime gen components when preview platform changes.
+	OnSettingsChangedDelegate.Broadcast(this, EPCGChangeType::Cosmetic);
+}
+
+void UPCGPlatformSwitchSettingsBase::OnPreviewFeatureLevelChanged(ERHIFeatureLevel::Type)
+{
+	// We don't broadcast structural from here. The PCGSubsystem will refresh all runtime gen components when preview platform changes.
+	OnSettingsChangedDelegate.Broadcast(this, EPCGChangeType::Cosmetic);
+}
+#endif // WITH_EDITOR
+
+bool UPCGPlatformSwitchSettingsBase::IsPinUsedByNodeExecution(const UPCGPin* InPin) const
+{
+	if (!InPin->IsOutputPin())
+	{
+		return Super::IsPinUsedByNodeExecution(InPin);
+	}
+
+	return IsPinStaticallyActive(InPin->Properties.Label);
+}
diff --git a/Engine/Plugins/PCG/Source/PCG/Private/Elements/ControlFlow/PCGShaderFeatureLevelSwitch.cpp b/Engine/Plugins/PCG/Source/PCG/Private/Elements/ControlFlow/PCGShaderFeatureLevelSwitch.cpp
new file mode 100644
--- /dev/null
+++ b/Engine/Plugins/PCG/Source/PCG/Private/Elements/ControlFlow/PCGShaderFeatureLevelSwitch.cpp
@@ -0,0 +1,293 @@
+// Copyright Epic Games, Inc. All Rights Reserved.
+
+#include "Elements/ControlFlow/PCGShaderFeatureLevelSwitch.h"
+
+#include "PCGContext.h"
+#include "PCGModule.h"
+#include "Elements/PCGGather.h"
+#include "Helpers/PCGHelpers.h"
+
+#include "DataDrivenShaderPlatformInfo.h"
+#include "RHIStrings.h"
+#include "Interfaces/ITargetPlatform.h"
+#include "UObject/ObjectSaveContext.h"
+
+#if WITH_EDITOR
+#include "Editor.h"
+#endif
+
+#include UE_INLINE_GENERATED_CPP_BY_NAME(PCGShaderFeatureLevelSwitch)
+
+#define LOCTEXT_NAMESPACE "PCGShaderFeatureLevelSwitchElement"
+
+namespace PCGShaderFeatureLevelSwitch
+{
+	static const FName DefaultLabel = TEXT("Default");
+
+	static bool GetFeatureLevelFromName(FName InFeatureLevelName, ERHIFeatureLevel::Type& OutFeatureLevel)
+	{
+		if (InFeatureLevelName == FName(TEXT("ES3_1")))
+		{
+			OutFeatureLevel = ERHIFeatureLevel::ES3_1;
+			return true;
+		}
+		else if (InFeatureLevelName == FName(TEXT("SM5")))
+		{
+			OutFeatureLevel = ERHIFeatureLevel::SM5;
+			return true;
+		}
+		else if (InFeatureLevelName == FName(TEXT("SM6")))
+		{
+			OutFeatureLevel = ERHIFeatureLevel::SM6;
+			return true;
+		}
+		else
+		{
+			OutFeatureLevel = ERHIFeatureLevel::Num;
+			return false;
+		}
+	}
+
+#if WITH_EDITOR
+	static ERHIFeatureLevel::Type GetCurrentFeatureLevel()
+	{
+		if (GEditor && GEditor->IsFeatureLevelPreviewActive())
+		{
+			return GEditor->GetActiveFeatureLevelPreviewType();
+		}
+		else
+		{
+			ERHIFeatureLevel::Type FeatureLevel = ERHIFeatureLevel::ES3_1;
+
+			TArray<FName> ShaderFormats;
+
+			for (const ITargetPlatform* Platform : GetTargetPlatformManagerRef().GetActiveTargetPlatforms())
+			{
+				if (!Platform)
+				{
+					continue;
+				}
+
+				ShaderFormats.Reset();
+				Platform->GetAllTargetedShaderFormats(ShaderFormats);
+
+				for (FName Format : ShaderFormats)
+				{
+					const EShaderPlatform ShaderPlatform = ShaderFormatToLegacyShaderPlatform(Format);
+					const ERHIFeatureLevel::Type PlatformFeatureLevel = GetMaxSupportedFeatureLevel(ShaderPlatform);
+					FeatureLevel = FMath::Max(FeatureLevel, PlatformFeatureLevel);
+				}
+			}
+
+			return FeatureLevel;
+		}
+	}
+#endif // WITH_EDITOR
+}
+
+bool UPCGShaderFeatureLevelSwitchSettings::IsPinStaticallyActive(const FName& PinLabel) const
+{
+	ERHIFeatureLevel::Type CurrentFeatureLevel = GetCurrentShaderFeatureLevel();
+
+	if (PinLabel != PCGShaderFeatureLevelSwitch::DefaultLabel)
+	{
+		// Active if the platform for this pin is selected.
+		return IsRHIFeatureLevelSelected(PinLabel, CurrentFeatureLevel);
+	}
+	else
+	{
+		// Default only active if no other pin is active.
+		for (FName FeatureLevel : GetSanitizedRHIFeatureLevelNames())
+		{
+			if (IsRHIFeatureLevelSelected(FeatureLevel, CurrentFeatureLevel))
+			{
+				return false;
+			}
+		}
+
+		return true;
+	}
+}
+
+FString UPCGShaderFeatureLevelSwitchSettings::GetAdditionalTitleInformation() const
+{
+#if WITH_EDITOR
+	FString Selection;
+	if (RHIFeatureLevelOutputs.IsEmpty())
+	{
+		Selection = TEXT("None");
+	}
+	else if (RHIFeatureLevelOutputs.Num() == 1)
+	{
+		Selection = RHIFeatureLevelOutputs[0].ToString();
+	}
+	else
+	{
+		Selection = TEXT("Multiple");
+	}
+
+	return FString::Format(TEXT("Max Feature Level: {0}"), { Selection });
+#else
+	return {};
+#endif
+}
+
+#if WITH_EDITOR
+void UPCGShaderFeatureLevelSwitchSettings::PreSave(FObjectPreSaveContext ObjectSaveContext)
+{
+	Super::PreSave(ObjectSaveContext);
+
+	if (ObjectSaveContext.IsCooking())
+	{
+		CookedFeatureLevel = PCGShaderFeatureLevelSwitch::GetCurrentFeatureLevel();
+	}
+}
+
+FText UPCGShaderFeatureLevelSwitchSettings::GetNodeTooltipText() const
+{
+	return LOCTEXT("NodeTooltip", "Statically activates/deactivates output pins based on the maximum supported shader feature level, useful for culling downstream nodes.");
+}
+
+EPCGChangeType UPCGShaderFeatureLevelSwitchSettings::GetChangeTypeForProperty(const FName& InPropertyName) const
+{
+	EPCGChangeType ChangeType = Super::GetChangeTypeForProperty(InPropertyName) | EPCGChangeType::Cosmetic;
+
+	if (InPropertyName == GET_MEMBER_NAME_CHECKED(UPCGShaderFeatureLevelSwitchSettings, bEnabled)
+		|| InPropertyName == GET_MEMBER_NAME_CHECKED(UPCGShaderFeatureLevelSwitchSettings, RHIFeatureLevelOutputs))
+	{
+		ChangeType |= EPCGChangeType::Structural;
+	}
+
+	return ChangeType;
+}
+
+TArray<FName> UPCGShaderFeatureLevelSwitchSettings::GetRHIFeatureLevelOptions() const
+{
+	return { TEXT("ES3_1"), TEXT("SM5"), TEXT("SM6") };
+}
+#endif // WITH_EDITOR
+
+TArray<FPCGPinProperties> UPCGShaderFeatureLevelSwitchSettings::InputPinProperties() const
+{
+	TArray<FPCGPinProperties> PinProperties;
+	PinProperties.Emplace(PCGPinConstants::DefaultInputLabel, EPCGDataType::Any, /*bInAllowMultipleConnections=*/true, /*bAllowMultipleData=*/true,
+		LOCTEXT("OutputPinTooltip", "All input will be forwarded directly to the active output pin."));
+
+	return PinProperties;
+}
+
+TArray<FPCGPinProperties> UPCGShaderFeatureLevelSwitchSettings::OutputPinProperties() const
+{
+	TArray<FPCGPinProperties> PinProperties;
+
+	PinProperties.Emplace(PCGShaderFeatureLevelSwitch::DefaultLabel, EPCGDataType::Any);
+
+	for (FName FeatureLevelName : GetSanitizedRHIFeatureLevelNames())
+	{
+		PinProperties.Emplace(FeatureLevelName, EPCGDataType::Any);
+	}
+
+	return PinProperties;
+}
+
+FPCGElementPtr UPCGShaderFeatureLevelSwitchSettings::CreateElement() const
+{
+	return MakeShared<FPCGShaderFeatureLevelSwitchElement>();
+}
+
+bool UPCGShaderFeatureLevelSwitchSettings::IsRHIFeatureLevelSelected(ERHIFeatureLevel::Type InFeatureLevel, ERHIFeatureLevel::Type InCurrentFeatureLevel) const
+{
+	return InFeatureLevel == InCurrentFeatureLevel;
+}
+
+bool UPCGShaderFeatureLevelSwitchSettings::IsRHIFeatureLevelSelected(FName InFeatureLevelName, ERHIFeatureLevel::Type InCurrentFeatureLevel) const
+{
+	ERHIFeatureLevel::Type FeatureLevel;
+	return PCGShaderFeatureLevelSwitch::GetFeatureLevelFromName(InFeatureLevelName, FeatureLevel) && IsRHIFeatureLevelSelected(FeatureLevel, InCurrentFeatureLevel);
+}
+
+ERHIFeatureLevel::Type UPCGShaderFeatureLevelSwitchSettings::GetCurrentShaderFeatureLevel() const
+{
+#if WITH_EDITOR
+	return PCGShaderFeatureLevelSwitch::GetCurrentFeatureLevel();
+#else
+	return static_cast<ERHIFeatureLevel::Type>(CookedFeatureLevel);
+#endif
+}
+
+TArray<ERHIFeatureLevel::Type> UPCGShaderFeatureLevelSwitchSettings::GetSanitizedRHIFeatureLevels() const
+{
+	TArray<FName> FeatureLevelNames = GetSanitizedRHIFeatureLevelNames();
+
+	TArray<ERHIFeatureLevel::Type> FeatureLevels;
+	FeatureLevels.Reserve(FeatureLevelNames.Num());
+
+	for (FName FeatureLevelName : FeatureLevelNames)
+	{
+		ERHIFeatureLevel::Type FeatureLevel;
+		if (GetFeatureLevelFromName(FeatureLevelName, FeatureLevel))
+		{
+			FeatureLevels.Add(FeatureLevel);
+		}
+	}
+
+	return FeatureLevels;
+}
+
+TArray<FName> UPCGShaderFeatureLevelSwitchSettings::GetSanitizedRHIFeatureLevelNames() const
+{
+	TArray<FName> FeatureLevels;
+	FeatureLevels.Reserve(RHIFeatureLevelOutputs.Num());
+
+	for (FName FeatureLevel : RHIFeatureLevelOutputs)
+	{
+		if (!FeatureLevel.IsNone() && !FeatureLevels.Contains(FeatureLevel))
+		{
+			FeatureLevels.Add(FeatureLevel);
+		}
+	}
+
+	return FeatureLevels;
+}
+
+bool FPCGShaderFeatureLevelSwitchElement::ExecuteInternal(FPCGContext* InContext) const
+{
+	TRACE_CPUPROFILER_EVENT_SCOPE(FPCGShaderFeatureLevelSwitchElement::ExecuteInternal);
+
+	const UPCGShaderFeatureLevelSwitchSettings* Settings = InContext->GetInputSettings<UPCGShaderFeatureLevelSwitchSettings>();
+	check(Settings);
+
+	const FPCGDataCollection InputData = PCGGather::GatherDataForPin(InContext->InputData, PCGPinConstants::DefaultInputLabel, PCGShaderFeatureLevelSwitch::DefaultLabel);
+
+	const TArray<FName> FeatureLevelNames = Settings->GetSanitizedRHIFeatureLevelNames();
+	InContext->OutputData.TaggedData.Reserve(InputData.TaggedData.Num() * (1 + FeatureLevelNames.Num()));
+
+	ERHIFeatureLevel::Type CurrentFeatureLevel = Settings->GetCurrentShaderFeatureLevel();
+
+	bool bAnyOutputSelected = false;
+
+	// Copy of data for each supported feature level.
+	for (FName FeatureLevelName : FeatureLevelNames)
+	{
+		if (Settings->IsRHIFeatureLevelSelected(FeatureLevelName, CurrentFeatureLevel))
+		{
+			for (FPCGTaggedData TaggedData : InputData.TaggedData)
+			{
+				TaggedData.Pin = FeatureLevelName;
+				InContext->OutputData.TaggedData.Add(TaggedData);
+			}
+
+			bAnyOutputSelected = true;
+		}
+	}
+
+	if (!bAnyOutputSelected)
+	{
+		// Fall back to Default pin.
+		InContext->OutputData.TaggedData.Append(InputData.TaggedData);
+	}
+
+	return true;
+}
+
+#undef LOCTEXT_NAMESPACE
diff --git a/Engine/Plugins/PCG/Source/PCG/Public/Elements/ControlFlow/PCGPlatformSwitch.h b/Engine/Plugins/PCG/Source/PCG/Public/Elements/ControlFlow/PCGPlatformSwitch.h
--- a/Engine/Plugins/PCG/Source/PCG/Public/Elements/ControlFlow/PCGPlatformSwitch.h
+++ b/Engine/Plugins/PCG/Source/PCG/Public/Elements/ControlFlow/PCGPlatformSwitch.h
@@ -2,27 +2,20 @@
 
 #pragma once
 
-#include "PCGControlFlow.h"
+#include "Elements/ControlFlow/PCGPlatformSwitchBase.h"
 
 #include "PCGPlatformSwitch.generated.h"
 
-namespace ERHIFeatureLevel { enum Type : int; }
-
 /** Statically activates/deactivates output pins based on the current platform/platform group. */
 UCLASS(BlueprintType, ClassGroup = (Procedural), meta=(Keywords = "if branch platform"))
-class UPCGPlatformSwitchSettings : public UPCGControlFlowSettings
+class UPCGPlatformSwitchSettings : public UPCGPlatformSwitchSettingsBase
 {
 	GENERATED_BODY()
 
 public:
 	//~ Begin UObject interface
 #if WITH_EDITOR
-	virtual void PostInitProperties();
-	virtual void PostLoad() override;
-	virtual void PostDuplicate(bool bDuplicateForPIE) override;
-	virtual void PostEditImport() override;
 	virtual void PreSave(FObjectPreSaveContext ObjectSaveContext) override;
-	virtual void BeginDestroy() override;
 #endif
 	//~ End UObject interface
 
@@ -31,14 +24,9 @@ public:
 	virtual FName GetDefaultNodeName() const override { return FName(TEXT("PlatformSwitch")); }
 	virtual FText GetDefaultNodeTitle() const override { return NSLOCTEXT("FPCGPlatformSwitchElement", "NodeTitle", "Platform Switch"); }
 	virtual FText GetNodeTooltipText() const override;
-	virtual EPCGSettingsType GetType() const override { return EPCGSettingsType::ControlFlow; }
 #endif // WITH_EDITOR
-	virtual bool HasDynamicPins() const override { return true; }
 	virtual bool IsPinStaticallyActive(const FName& PinLabel) const override;
-	virtual bool IsPinUsedByNodeExecution(const UPCGPin* InPin) const override;
 	virtual FString GetAdditionalTitleInformation() const override;
-	virtual bool HasFlippedTitleLines() const override { return true; }
-	virtual bool CanCullTaskIfUnwired() const override { return false; } // Useful to run unwired to only cull downstream nodes.
 
 protected:
 #if WITH_EDITOR
@@ -49,13 +37,6 @@ protected:
 	virtual FPCGElementPtr CreateElement() const override;
 	//~ End UPCGSettings interface
 
-#if WITH_EDITOR
-	void SetupPreviewPlatformChangeEvents();
-	void TearDownPreviewPlatformChangeEvents();
-	void OnPreviewPlatformChanged();
-	void OnPreviewFeatureLevelChanged(ERHIFeatureLevel::Type);
-#endif
-
 public:
 	/** Get the users selection, sanitized to remove None and duplicates. */
 	TArray<FName> GetSanitizedPlatforms() const;
@@ -69,7 +50,6 @@ public:
 
 protected:
 #if WITH_EDITOR
-	/** List of all non-advanced input pin names. */
 	UFUNCTION()
 	TArray<FName> GetPlatformOptions() const;
 #endif
@@ -85,12 +65,6 @@ protected:
 	/** The current platform group used in cooked builds. */
 	UPROPERTY()
 	FName CookedPlatformGroup;
-
-private:
-#if WITH_EDITOR
-	FDelegateHandle PreviewPlatformChangedHandle;
-	FDelegateHandle PreviewFeatureLevelChangedHandle;
-#endif
 };
 
 class FPCGPlatformSwitchElement : public IPCGElement
diff --git a/Engine/Plugins/PCG/Source/PCG/Public/Elements/ControlFlow/PCGPlatformSwitchBase.h b/Engine/Plugins/PCG/Source/PCG/Public/Elements/ControlFlow/PCGPlatformSwitchBase.h
new file mode 100644
--- /dev/null
+++ b/Engine/Plugins/PCG/Source/PCG/Public/Elements/ControlFlow/PCGPlatformSwitchBase.h
@@ -0,0 +1,50 @@
+// Copyright Epic Games, Inc. All Rights Reserved.
+
+#pragma once
+
+#include "PCGControlFlow.h"
+
+#include "PCGPlatformSwitchBase.generated.h"
+
+namespace ERHIFeatureLevel { enum Type : int; }
+
+UCLASS(MinimalAPI, BlueprintType, Abstract, ClassGroup = (Procedural))
+class UPCGPlatformSwitchSettingsBase : public UPCGControlFlowSettings
+{
+	GENERATED_BODY()
+
+public:
+	//~ Begin UObject interface
+#if WITH_EDITOR
+	virtual void PostInitProperties();
+	virtual void PostLoad() override;
+	virtual void PostDuplicate(bool bDuplicateForPIE) override;
+	virtual void PostEditImport() override;
+	virtual void BeginDestroy() override;
+#endif
+	//~ End UObject interface
+
+	//~ Begin UPCGSettings interface
+#if WITH_EDITOR
+	virtual EPCGSettingsType GetType() const override { return EPCGSettingsType::ControlFlow; }
+#endif // WITH_EDITOR
+	virtual bool HasDynamicPins() const override { return true; }
+	virtual bool HasFlippedTitleLines() const override { return true; }
+	virtual bool CanCullTaskIfUnwired() const override { return false; } // Useful to run unwired to only cull downstream nodes.
+	virtual bool IsPinUsedByNodeExecution(const UPCGPin* InPin) const override;
+	//~ End UPCGSettings interface
+
+protected:
+#if WITH_EDITOR
+	void SetupPreviewPlatformChangeEvents();
+	void TearDownPreviewPlatformChangeEvents();
+	void OnPreviewPlatformChanged();
+	void OnPreviewFeatureLevelChanged(ERHIFeatureLevel::Type);
+#endif
+
+private:
+#if WITH_EDITOR
+	FDelegateHandle PreviewPlatformChangedHandle;
+	FDelegateHandle PreviewFeatureLevelChangedHandle;
+#endif
+};
diff --git a/Engine/Plugins/PCG/Source/PCG/Public/Elements/ControlFlow/PCGShaderFeatureLevelSwitch.h b/Engine/Plugins/PCG/Source/PCG/Public/Elements/ControlFlow/PCGShaderFeatureLevelSwitch.h
new file mode 100644
--- /dev/null
+++ b/Engine/Plugins/PCG/Source/PCG/Public/Elements/ControlFlow/PCGShaderFeatureLevelSwitch.h
@@ -0,0 +1,75 @@
+// Copyright Epic Games, Inc. All Rights Reserved.
+
+#pragma once
+
+#include "Elements/ControlFlow/PCGPlatformSwitchBase.h"
+
+#include "RHIFeatureLevel.h"
+
+#include "PCGShaderFeatureLevelSwitch.generated.h"
+
+/** Statically activates/deactivates output pins based on the maximum shader feature level of the current platform. */
+UCLASS(BlueprintType, ClassGroup = (Procedural), meta=(Keywords = "if branch platform shader"))
+class UPCGShaderFeatureLevelSwitchSettings : public UPCGPlatformSwitchSettingsBase
+{
+	GENERATED_BODY()
+
+public:
+	//~ Begin UObject interface
+#if WITH_EDITOR
+	virtual void PreSave(FObjectPreSaveContext ObjectSaveContext) override;
+#endif
+	//~ End UObject interface
+
+	//~ Begin UPCGSettings interface
+#if WITH_EDITOR
+	virtual FName GetDefaultNodeName() const override { return FName(TEXT("ShaderFeatureLevelSwitch")); }
+	virtual FText GetDefaultNodeTitle() const override { return NSLOCTEXT("PCGShaderFeatureLevelSwitchElement", "NodeTitle", "Max Shader Feature Level Switch"); }
+	virtual FText GetNodeTooltipText() const override;
+#endif // WITH_EDITOR
+	virtual bool IsPinStaticallyActive(const FName& PinLabel) const override;
+	virtual FString GetAdditionalTitleInformation() const override;
+
+protected:
+#if WITH_EDITOR
+	virtual EPCGChangeType GetChangeTypeForProperty(const FName& InPropertyName) const override;
+#endif
+	virtual TArray<FPCGPinProperties> InputPinProperties() const override;
+	virtual TArray<FPCGPinProperties> OutputPinProperties() const override;
+	virtual FPCGElementPtr CreateElement() const override;
+	//~ End UPCGSettings interface
+
+public:
+	/** Get the users selection, sanitized to remove None and duplicates. */
+	TArray<ERHIFeatureLevel::Type> GetSanitizedRHIFeatureLevels() const;
+	TArray<FName> GetSanitizedRHIFeatureLevelNames() const;
+
+	/** Returns current platform and platform group, either by querying from editor or by using cooked values. */
+	ERHIFeatureLevel::Type GetCurrentShaderFeatureLevel() const;
+
+	/** Returns true if given shader feature level included in users selection. */
+	bool IsRHIFeatureLevelSelected(ERHIFeatureLevel::Type InFeatureLevel, ERHIFeatureLevel::Type InCurrentFeatureLevel) const;
+	bool IsRHIFeatureLevelSelected(FName InFeatureLevelName, ERHIFeatureLevel::Type InCurrentFeatureLevel) const;
+
+protected:
+#if WITH_EDITOR
+	UFUNCTION()
+	TArray<FName> GetRHIFeatureLevelOptions() const;
+#endif
+
+protected:
+	UPROPERTY(EditAnywhere, Category=Settings, DisplayName = "Shader Feature Levels", meta = (GetOptions = "GetRHIFeatureLevelOptions"))
+	TArray<FName> RHIFeatureLevelOutputs;
+
+	/** The current platform group used in cooked builds. */
+	UPROPERTY()
+	int CookedFeatureLevel = INDEX_NONE;
+};
+
+class FPCGShaderFeatureLevelSwitchElement : public IPCGElement
+{
+protected:
+	virtual bool ExecuteInternal(FPCGContext* InContext) const override;
+	virtual bool SupportsGPUResidentData(FPCGContext* InContext) const override { return true; }
+	virtual bool SupportsBasePointDataInputs(FPCGContext* InContext) const { return true; }
+};
