[PCG] Support overriding NumElements2D in CustomHLSL node

Split NumElements2D property into NumElements2DX and NumElements2DY to allow integer overrides instead of relying on struct override support.

Deprecate the old NumElements2D property.

diff --git a/Engine/Plugins/PCG/Source/PCG/Private/Compute/DataInterfaces/Elements/PCGCustomHLSLDataInterface.cpp b/Engine/Plugins/PCG/Source/PCG/Private/Compute/DataInterfaces/Elements/PCGCustomHLSLDataInterface.cpp
--- a/Engine/Plugins/PCG/Source/PCG/Private/Compute/DataInterfaces/Elements/PCGCustomHLSLDataInterface.cpp
+++ b/Engine/Plugins/PCG/Source/PCG/Private/Compute/DataInterfaces/Elements/PCGCustomHLSLDataInterface.cpp
@@ -17,7 +17,8 @@
 
 BEGIN_SHADER_PARAMETER_STRUCT(FPCGCustomHLSLDataInterfaceParameters,)
 	SHADER_PARAMETER(int32, NumElements)
-	//SHADER_PARAMETER(FIntVector2, NumElements2D)
+	SHADER_PARAMETER(int32, NumElementsX)
+	SHADER_PARAMETER(int32, NumElementsY)
 	SHADER_PARAMETER(int32, ThreadCountMultiplier)
 	SHADER_PARAMETER(int32, FixedThreadCount)
 END_SHADER_PARAMETER_STRUCT()
@@ -39,7 +40,8 @@ FComputeDataProviderRenderProxy* UPCGCustomHLSLDataProvider::GetRenderProxy()
 	if (ensure(KernelParams))
 	{
 		ProxyData.NumElements = KernelParams->GetValueInt(GET_MEMBER_NAME_CHECKED(UPCGCustomHLSLSettings, NumElements));
-		//ProxyData.NumElements2D = KernelParams->GetValueInt2(GET_MEMBER_NAME_CHECKED(UPCGCustomHLSLSettings, NumElements2D));
+		ProxyData.NumElementsX = KernelParams->GetValueInt(GET_MEMBER_NAME_CHECKED(UPCGCustomHLSLSettings, NumElementsX));
+		ProxyData.NumElementsY = KernelParams->GetValueInt(GET_MEMBER_NAME_CHECKED(UPCGCustomHLSLSettings, NumElementsY));
 		ProxyData.ThreadCountMultiplier = KernelParams->GetValueInt(GET_MEMBER_NAME_CHECKED(UPCGCustomHLSLSettings, ThreadCountMultiplier));
 		ProxyData.FixedThreadCount = KernelParams->GetValueInt(GET_MEMBER_NAME_CHECKED(UPCGCustomHLSLSettings, FixedThreadCount));
 	}
@@ -82,7 +84,8 @@ void FPCGCustomHLSLDataProviderProxy::GatherDispatchData(FDispatchData const& In
 	{
 		FParameters& Parameters = ParameterArray[InvocationIndex];
 		Parameters.NumElements = Data.NumElements;
-		//Parameters.NumElements2D = static_cast<FVector2f>(Data.NumElements2D);
+		Parameters.NumElementsX = Data.NumElementsX;
+		Parameters.NumElementsY = Data.NumElementsY;
 		Parameters.ThreadCountMultiplier = Data.ThreadCountMultiplier;
 		Parameters.FixedThreadCount = Data.FixedThreadCount;
 	}
diff --git a/Engine/Plugins/PCG/Source/PCG/Private/Compute/DataInterfaces/Elements/PCGCustomHLSLDataInterface.h b/Engine/Plugins/PCG/Source/PCG/Private/Compute/DataInterfaces/Elements/PCGCustomHLSLDataInterface.h
--- a/Engine/Plugins/PCG/Source/PCG/Private/Compute/DataInterfaces/Elements/PCGCustomHLSLDataInterface.h
+++ b/Engine/Plugins/PCG/Source/PCG/Private/Compute/DataInterfaces/Elements/PCGCustomHLSLDataInterface.h
@@ -48,8 +48,8 @@ public:
 	struct FCustomHLSLData
 	{
 		int32 NumElements = 0;
-		// @todo_pcg: Support param override for NumElements2D. Requires struct support for GPU overridable params.
-		//FIntPoint NumElements2D = FIntPoint::ZeroValue;
+		int32 NumElementsX = 0;
+		int32 NumElementsY = 0;
 		int32 ThreadCountMultiplier = 0;
 		int32 FixedThreadCount = 0;
 	};
diff --git a/Engine/Plugins/PCG/Source/PCG/Private/Compute/Elements/PCGCustomHLSL.cpp b/Engine/Plugins/PCG/Source/PCG/Private/Compute/Elements/PCGCustomHLSL.cpp
--- a/Engine/Plugins/PCG/Source/PCG/Private/Compute/Elements/PCGCustomHLSL.cpp
+++ b/Engine/Plugins/PCG/Source/PCG/Private/Compute/Elements/PCGCustomHLSL.cpp
@@ -56,6 +56,13 @@ void UPCGCustomHLSLSettings::PostLoad()
 		NumElements = PointCount_DEPRECATED;
 		PointCount_DEPRECATED = 0;
 	}
+
+	if (NumElements2D_DEPRECATED.X > 0 || NumElements2D_DEPRECATED.Y > 0)
+	{
+		NumElementsX = NumElements2D_DEPRECATED.X;
+		NumElementsY = NumElements2D_DEPRECATED.Y;
+		NumElements2D_DEPRECATED = FIntPoint(0, 0);
+	}
 #endif
 
 	// Note: We update here so that Custom HLSL nodes will have the correct pin settings & declarations on load.
diff --git a/Engine/Plugins/PCG/Source/PCG/Private/Compute/Elements/PCGCustomHLSL.h b/Engine/Plugins/PCG/Source/PCG/Private/Compute/Elements/PCGCustomHLSL.h
--- a/Engine/Plugins/PCG/Source/PCG/Private/Compute/Elements/PCGCustomHLSL.h
+++ b/Engine/Plugins/PCG/Source/PCG/Private/Compute/Elements/PCGCustomHLSL.h
@@ -150,8 +150,11 @@ public:
 	UPROPERTY(EditAnywhere, Category = "Settings", meta = (EditCondition = "KernelType == EPCGKernelType::PointGenerator", EditConditionHides, PCG_OverridableCPUAndGPUWithReadback))
 	int NumElements = 256;
 
-	UPROPERTY(EditAnywhere, DisplayName = "Num Elements", Category = "Settings", meta = (EditCondition = "KernelType == EPCGKernelType::TextureGenerator", EditConditionHides))
-	FIntPoint NumElements2D = FIntPoint(64, 64);
+	UPROPERTY(EditAnywhere, Category = "Settings", meta = (EditCondition = "KernelType == EPCGKernelType::TextureGenerator", EditConditionHides, PCG_OverridableCPUAndGPUWithReadback))
+	int NumElementsX = 64;
+
+	UPROPERTY(EditAnywhere, Category = "Settings", meta = (EditCondition = "KernelType == EPCGKernelType::TextureGenerator", EditConditionHides, PCG_OverridableCPUAndGPUWithReadback))
+	int NumElementsY = 64;
 
 	UPROPERTY(EditAnywhere, Category = "Settings|Thread Count", meta = (EditCondition = "KernelType == EPCGKernelType::Custom", EditConditionHides))
 	EPCGDispatchThreadCount DispatchThreadCount = EPCGDispatchThreadCount::FromFirstOutputPin;
@@ -216,6 +219,9 @@ protected:
 
 	UPROPERTY()
 	int PointCount_DEPRECATED = 0;
+
+	UPROPERTY()
+	FIntPoint NumElements2D_DEPRECATED = FIntPoint(0, 0);
 #endif
 
 	friend class UPCGCustomHLSLKernel;
diff --git a/Engine/Plugins/PCG/Source/PCG/Private/Compute/Elements/PCGCustomHLSLKernel.cpp b/Engine/Plugins/PCG/Source/PCG/Private/Compute/Elements/PCGCustomHLSLKernel.cpp
--- a/Engine/Plugins/PCG/Source/PCG/Private/Compute/Elements/PCGCustomHLSLKernel.cpp
+++ b/Engine/Plugins/PCG/Source/PCG/Private/Compute/Elements/PCGCustomHLSLKernel.cpp
@@ -270,12 +270,22 @@ TSharedPtr<const FPCGDataCollectionDesc> UPCGCustomHLSLKernel::ComputeOutputBind
 	}
 	else if (InOutputPinLabel == FirstOutputPinLabel && CustomHLSLSettings->KernelType == EPCGKernelType::TextureGenerator)
 	{
+		const FPCGKernelParams* KernelParams = InBinding->GetCachedKernelParams(this);
+
+		if (!ensure(KernelParams))
+		{
+			return nullptr;
+		}
+
 		OutDataDesc = FPCGDataCollectionDesc::MakeShared();
 
+		const int NumElementsX = KernelParams->GetValueInt(GET_MEMBER_NAME_CHECKED(UPCGCustomHLSLSettings, NumElementsX));
+		const int NumElementsY = KernelParams->GetValueInt(GET_MEMBER_NAME_CHECKED(UPCGCustomHLSLSettings, NumElementsY));
+
 		// Generators always produce a single texture data with known size.
-		OutDataDesc->GetDataDescriptionsMutable().Emplace(EPCGDataType::BaseTexture, FIntVector4(CustomHLSLSettings->NumElements2D.X, CustomHLSLSettings->NumElements2D.Y, 0, 0));
+		OutDataDesc->GetDataDescriptionsMutable().Emplace(EPCGDataType::BaseTexture, FIntVector4(NumElementsX, NumElementsY, 0, 0));
 	}
-	else if (ensure(OutputPinProperties))
+	else
 	{
 		OutDataDesc = FPCGDataCollectionDesc::MakeShared();
 
@@ -532,7 +542,7 @@ FString UPCGCustomHLSLKernel::GetCookedSource(FPCGGPUCompilationContext& InOutCo
 	}
 	else if (CustomHLSLSettings->KernelType == EPCGKernelType::TextureGenerator)
 	{
-		KernelSpecificPreamble += FString::Format(TEXT("    const uint2 NumElements = uint2({0}, {1});\n"), { CustomHLSLSettings->NumElements2D.X, CustomHLSLSettings->NumElements2D.Y });
+		KernelSpecificPreamble += TEXT("    const uint2 NumElements = uint2(NumElementsX_GetOverridableValue(), NumElementsY_GetOverridableValue());\n");
 
 		if (const FPCGPinPropertiesGPU* OutputPin = GetFirstOutputPin())
 		{
@@ -636,8 +646,11 @@ int UPCGCustomHLSLKernel::ComputeThreadCount(const UPCGDataBinding* InBinding) c
 	}
 	else if (CustomHLSLSettings->KernelType == EPCGKernelType::TextureGenerator)
 	{
+		const int NumElementsX = KernelParams->GetValueInt(GET_MEMBER_NAME_CHECKED(UPCGCustomHLSLSettings, NumElementsX));
+		const int NumElementsY = KernelParams->GetValueInt(GET_MEMBER_NAME_CHECKED(UPCGCustomHLSLSettings, NumElementsY));
+
 		// Texture generator has fixed thread count.
-		ThreadCount = CustomHLSLSettings->NumElements2D.X * CustomHLSLSettings->NumElements2D.Y;
+		ThreadCount = NumElementsX * NumElementsY;
 	}
 	else if (CustomHLSLSettings->IsProcessorKernel())
 	{
