[PCG] Update RGS to use execution sources rather than local components

#rnx

diff --git a/Engine/Plugins/PCG/Source/PCG/Private/Graph/PCGGraphExecutor.cpp b/Engine/Plugins/PCG/Source/PCG/Private/Graph/PCGGraphExecutor.cpp
--- a/Engine/Plugins/PCG/Source/PCG/Private/Graph/PCGGraphExecutor.cpp
+++ b/Engine/Plugins/PCG/Source/PCG/Private/Graph/PCGGraphExecutor.cpp
@@ -3235,7 +3235,7 @@ namespace PCGGraphExecutor
 				PCGGraphExecutionLogging::LogGridLinkageTaskExecuteRetrieveScheduleGraph(InContext, ExecSourceWithData, InResourceKey);
 
 				// Wake up this task after graph has generated.
-				const FPCGTaskId GraphTaskId = ExecSourceWithData->GetExecutionState().GenerateLocalGetTaskId();
+				const FPCGTaskId GraphTaskId = ExecSourceWithData->GetExecutionState().Generate();
 
 				// Update state and go to paused.
 				InContext->DynamicDependencies.Add(GraphTaskId);
diff --git a/Engine/Plugins/PCG/Source/PCG/Private/PCGComponentExecutionState.cpp b/Engine/Plugins/PCG/Source/PCG/Private/PCGComponentExecutionState.cpp
--- a/Engine/Plugins/PCG/Source/PCG/Private/PCGComponentExecutionState.cpp
+++ b/Engine/Plugins/PCG/Source/PCG/Private/PCGComponentExecutionState.cpp
@@ -163,6 +163,12 @@ FPCGGridDescriptor FPCGComponentExecutionState::GetGridDescriptor(uint32 InGridS
 	return Component->GetGridDescriptor(InGridSize);
 }
 
+EPCGHiGenGrid FPCGComponentExecutionState::GetGenerationGrid() const
+{
+	check(Component);
+	return (IsPartitioned() || IsLocalSource()) ? Component->GetGenerationGrid() : EPCGHiGenGrid::Uninitialized;
+}
+
 uint32 FPCGComponentExecutionState::GetGenerationGridSize() const
 {
 	check(Component);
@@ -193,12 +199,43 @@ void FPCGComponentExecutionState::AddToManagedResources(UPCGManagedResource* InR
 	return Component->AddToManagedResources(InResource);
 }
 
-FPCGTaskId FPCGComponentExecutionState::GenerateLocalGetTaskId(EPCGHiGenGrid Grid)
+FPCGTaskId FPCGComponentExecutionState::Generate(const FGenerateParams& InGenerateParams)
 {
-	check(Component)
+	check(Component);
 
+	const EPCGHiGenGrid Grid = InGenerateParams.bGenerateLocalSources ? EPCGHiGenGrid::Uninitialized : GetGenerationGrid();
 	const EPCGComponentGenerationTrigger GenTrigger = IsManagedByRuntimeGenSystem() ? EPCGComponentGenerationTrigger::GenerateAtRuntime : EPCGComponentGenerationTrigger::GenerateOnDemand;
-	return Component->GenerateLocalGetTaskId(GenTrigger, /*bForce=*/true, Grid);
+	return Component->GenerateLocalGetTaskId(GenTrigger, /*bForce=*/true, Grid, InGenerateParams.Dependencies);
+}
+
+// UE_DEPRECATED(5.8)
+FPCGTaskId FPCGComponentExecutionState::GenerateLocalGetTaskId(EPCGHiGenGrid Grid)
+{
+	FGenerateParams GenerateParams;
+	GenerateParams.bGenerateLocalSources = Grid == EPCGHiGenGrid::Uninitialized;
+
+	return Generate(GenerateParams);
+}
+
+FPCGTaskId FPCGComponentExecutionState::Cleanup(const FCleanupParams& InCleanupParams)
+{
+	check(Component);
+
+	if (InCleanupParams.bImmediate)
+	{
+		Component->CleanupLocalImmediate(InCleanupParams.bReleaseManagedResources, InCleanupParams.bCleanupLocalSources);
+		return InvalidPCGTaskId;
+	}
+
+	// @todo_pcg: There is a difference in behavior for CleanupLocal vs CleanupLocalImmediate, where CleanupLocal cannot be
+	// called on local components directly, and will ALWAYS cleanup local components when called on an original component.
+	if (InCleanupParams.bCleanupLocalSources || IsLocalSource())
+	{
+		UE_LOG(LogPCG, Warning, TEXT("Tried to schedule cleanup on a local execution source. This is not supported, call Cleanup() with bImmediate=true instead."))
+		return InvalidPCGTaskId;
+	}
+
+	return Component->CleanupLocal(InCleanupParams.bReleaseManagedResources, InCleanupParams.Dependencies);
 }
 
 #if WITH_EDITOR
diff --git a/Engine/Plugins/PCG/Source/PCG/Private/RuntimeGen/PCGRuntimeGenScheduler.cpp b/Engine/Plugins/PCG/Source/PCG/Private/RuntimeGen/PCGRuntimeGenScheduler.cpp
--- a/Engine/Plugins/PCG/Source/PCG/Private/RuntimeGen/PCGRuntimeGenScheduler.cpp
+++ b/Engine/Plugins/PCG/Source/PCG/Private/RuntimeGen/PCGRuntimeGenScheduler.cpp
@@ -5,6 +5,7 @@
 
 #include "PCGActorAndComponentMapping.h"
 #include "PCGCommon.h"
+#include "PCGComponent.h"
 #include "PCGGraph.h"
 #include "PCGModule.h"
 #include "PCGSubsystem.h"
@@ -912,7 +913,7 @@ void FPCGRuntimeGenScheduler::TickCleanup(const TSet<IPCGGenSourceBase*>& InGenS
 	const int32 NumGeneratedComponents = GeneratedComponentsArray.Num();
 
 	// Generated Entry Key, Local/Generated Component
-	using PCGComponentToClean = TTuple<FGridGenerationKey, UPCGComponent*>;
+	using PCGComponentToClean = TTuple<FGridGenerationKey, IPCGGraphExecutionSource*>;
 	TArray<PCGComponentToClean, TInlineAllocator<256>> ComponentsToClean;
 	TArray<FGridGenerationKey, TInlineAllocator<16>> InvalidKeys;
 
@@ -971,13 +972,13 @@ void FPCGRuntimeGenScheduler::TickCleanup(const TSet<IPCGGenSourceBase*>& InGenS
 		// Otherwise, we have a local component.
 		else
 		{
-			UPCGComponent* LocalComponent = GenerationKey.GetCachedLocalComponent();
+			IPCGGraphExecutionSource* LocalComponent = GenerationKey.GetCachedLocalComponent();
 			if (!LocalComponent)
 			{
 				LocalComponent = ActorAndComponentMapping->GetLocalComponent(GridDescriptor, GridCoords, OriginalComponent);
 			}
 
-			APCGPartitionActor* PartitionActor = LocalComponent ? Cast<APCGPartitionActor>(LocalComponent->GetOwner()) : nullptr;
+			APCGPartitionActor* PartitionActor = LocalComponent ? Cast<APCGPartitionActor>(CastChecked<UPCGComponent>(LocalComponent)->GetOwner()) : nullptr;
 			if (!PartitionActor || !OriginalComponent->bActivated)
 			{
 				// Attempt to clean even in failure case to avoid leaking resources.
@@ -986,7 +987,7 @@ void FPCGRuntimeGenScheduler::TickCleanup(const TSet<IPCGGenSourceBase*>& InGenS
 				return;
 			}
 
-			const FBox GridBounds = LocalComponent->GetGridBounds();
+			const FBox GridBounds = LocalComponent->GetExecutionState().GetBounds();
 
 			// Only clean up if all generation sources agreed to clean up.
 			if (CheckIfAllGenSourcesWantToCleanup(Policy, GridDescriptor, GridBounds, CleanupRadiusSquared))
@@ -1037,9 +1038,9 @@ void FPCGRuntimeGenScheduler::TickScheduleGeneration(TMap<FGridGenerationKey, do
 				continue;
 			}
 
-			UPCGComponent* Component = (Key.GetGridDescriptor().GetHiGenGrid() == EPCGHiGenGrid::Unbounded) ? Key.GetOriginalComponent() : Key.GetCachedLocalComponent();
+			const IPCGGraphExecutionSource* ExecutionSource = (Key.GetGridDescriptor().GetHiGenGrid() == EPCGHiGenGrid::Unbounded) ? Key.GetOriginalComponent() : Key.GetCachedLocalComponent();
 
-			if (Component && Component->IsGenerating())
+			if (ExecutionSource && ExecutionSource->GetExecutionState().IsGenerating())
 			{
 				++NumGenerating;
 			}
@@ -1110,8 +1111,8 @@ void FPCGRuntimeGenScheduler::TickScheduleGeneration(TMap<FGridGenerationKey, do
 		else
 		{
 			// Grab local component and PA if they exist already.
-			UPCGComponent* LocalComponent = ActorAndComponentMapping->GetLocalComponent(GridDescriptor, GridCoords, OriginalComponent);
-			TObjectPtr<APCGPartitionActor> PartitionActor = LocalComponent ? Cast<APCGPartitionActor>(LocalComponent->GetOwner()) : nullptr;
+			IPCGGraphExecutionSource* LocalComponent = ActorAndComponentMapping->GetLocalComponent(GridDescriptor, GridCoords, OriginalComponent);
+			TObjectPtr<APCGPartitionActor> PartitionActor = LocalComponent ? Cast<APCGPartitionActor>(CastChecked<UPCGComponent>(LocalComponent)->GetOwner()) : nullptr;
 
 			if (!LocalComponent || !ensure(PartitionActor))
 			{
@@ -1177,7 +1178,7 @@ void FPCGRuntimeGenScheduler::TickScheduleGeneration(TMap<FGridGenerationKey, do
 				LocalComponent = PartitionActor->GetLocalComponent(OriginalComponent);
 			}
 
-			if (ensure(LocalComponent) && !LocalComponent->IsGenerating())
+			if (ensure(LocalComponent) && !LocalComponent->GetExecutionState().IsGenerating())
 			{
 				if (PCGRuntimeGenSchedulerHelpers::CVarRuntimeGenerationEnableDebugging.GetValueOnAnyThread())
 				{
@@ -1202,7 +1203,10 @@ void FPCGRuntimeGenScheduler::TickScheduleGeneration(TMap<FGridGenerationKey, do
 #endif
 
 				// Force to refresh if the component is already generated.
-				LocalComponent->GenerateLocal(EPCGComponentGenerationTrigger::GenerateAtRuntime, /*bForce=*/true, LocalComponent->GetGenerationGrid(), Dependencies);
+				IPCGGraphExecutionState::FGenerateParams GenerateParams;
+				GenerateParams.Dependencies = MoveTemp(Dependencies);
+
+				LocalComponent->GetExecutionState().Generate(GenerateParams);
 
 				Key.SetCachedLocalComponent(LocalComponent);
 			}
@@ -1504,7 +1508,7 @@ void FPCGRuntimeGenScheduler::OnOriginalComponentUnregistered(UPCGComponent* InO
 		const bool bIsOriginalComponent = GridDescriptor.GetGridSize() == PCGHiGenGrid::UnboundedGridSize();
 
 		// Get the generated component for this key (might be a local component).
-		UPCGComponent* ComponentToCleanup = bIsOriginalComponent ? InOriginalComponent : ActorAndComponentMapping->GetLocalComponent(GridDescriptor, GenerationKey.GetGridCoords(), InOriginalComponent);
+		IPCGGraphExecutionSource* ComponentToCleanup = bIsOriginalComponent ? InOriginalComponent : ActorAndComponentMapping->GetLocalComponent(GridDescriptor, GenerationKey.GetGridCoords(), InOriginalComponent);
 
 		// It is possible for a PartitionActor's LocalComponent to have been cleaned up by the APCGPartitionActor::EndPlay call depending on the order in which actors get called
 		if (ComponentToCleanup)
@@ -1554,7 +1558,7 @@ void FPCGRuntimeGenScheduler::CleanupRemainingComponents(UPCGComponent* InOrigin
 		TSet<TObjectPtr<APCGPartitionActor>> PartitionActors = ActorAndComponentMapping->GetPCGComponentPartitionActorMappings(InOriginalComponent);
 		for (TObjectPtr<APCGPartitionActor> PartitionActor : PartitionActors)
 		{
-			if (UPCGComponent* ComponentToCleanup = PartitionActor->GetLocalComponent(InOriginalComponent))
+			if (IPCGGraphExecutionSource* ComponentToCleanup = PartitionActor->GetLocalComponent(InOriginalComponent))
 			{
 				CleanupLocalComponent(PartitionActor, ComponentToCleanup);
 			}
@@ -1569,7 +1573,7 @@ void FPCGRuntimeGenScheduler::CleanupLocalComponents(const APCGWorldActor* InPCG
 	check(ActorAndComponentMapping && InPCGWorldActor);
 
 	// Generated Entry Key, Local/Generated Component
-	using PCGComponentToClean = TTuple<FGridGenerationKey, UPCGComponent*>;
+	using PCGComponentToClean = TTuple<FGridGenerationKey, IPCGGraphExecutionSource*>;
 	TArray<PCGComponentToClean, TInlineAllocator<16>> ComponentsToClean;
 
 	TSet<UPCGComponent*> OriginalComponents;
@@ -1596,7 +1600,7 @@ void FPCGRuntimeGenScheduler::CleanupLocalComponents(const APCGWorldActor* InPCG
 		// Only operate on LocalComponents.
 		if (Grid != EPCGHiGenGrid::Unbounded)
 		{
-			UPCGComponent* LocalComponent = ActorAndComponentMapping->GetLocalComponent(GridDescriptor, GridCoords, OriginalComponent);
+			IPCGGraphExecutionSource* LocalComponent = ActorAndComponentMapping->GetLocalComponent(GridDescriptor, GridCoords, OriginalComponent);
 			ComponentsToClean.Add({ GenerationKey, LocalComponent });
 		}
 	}
@@ -1626,14 +1630,17 @@ void FPCGRuntimeGenScheduler::OnLevelStreamingStateChanged(UWorld* InWorld, cons
 	}
 }
 
-void FPCGRuntimeGenScheduler::CleanupLocalComponent(APCGPartitionActor* PartitionActor, UPCGComponent* LocalComponent)
+void FPCGRuntimeGenScheduler::CleanupLocalComponent(APCGPartitionActor* PartitionActor, IPCGGraphExecutionSource* LocalComponent)
 {
 	if (!PartitionActor)
 	{
 		return;
 	}
 
-	if (UPCGComponent* OriginalComponent = PartitionActor->GetOriginalComponent(LocalComponent))
+	IPCGGraphExecutionSource* OriginalExecutionSource = LocalComponent ? LocalComponent->GetExecutionState().GetOriginalSource() : nullptr;
+	UPCGComponent* OriginalComponent = Cast<UPCGComponent>(OriginalExecutionSource);
+
+	if (OriginalComponent)
 	{
 		if (PCGRuntimeGenSchedulerHelpers::CVarRuntimeGenerationEnableDebugging.GetValueOnAnyThread())
 		{
@@ -1697,7 +1704,7 @@ void FPCGRuntimeGenScheduler::CleanupLocalComponent(APCGPartitionActor* Partitio
 	}
 }
 
-void FPCGRuntimeGenScheduler::CleanupComponent(const FGridGenerationKey& GenerationKey, UPCGComponent* GeneratedComponent)
+void FPCGRuntimeGenScheduler::CleanupComponent(const FGridGenerationKey& GenerationKey, IPCGGraphExecutionSource* GeneratedComponent)
 {
 	TRACE_CPUPROFILER_EVENT_SCOPE(FPCGRuntimeGenScheduler::CleanupComponent);
 
@@ -1719,21 +1726,26 @@ void FPCGRuntimeGenScheduler::CleanupComponent(const FGridGenerationKey& Generat
 	}
 	else // If the GeneratedComponent does exist, we can clean it up.
 	{
-		PartitionActor = Cast<APCGPartitionActor>(GeneratedComponent->GetOwner());
+		PartitionActor = Cast<APCGPartitionActor>(CastChecked<UPCGComponent>(GeneratedComponent)->GetOwner());
 		ensure(PartitionActor || Grid == EPCGHiGenGrid::Unbounded);
 
 		if (Grid == EPCGHiGenGrid::Unbounded)
 		{
-			if (PCGRuntimeGenSchedulerHelpers::CVarRuntimeGenerationEnableDebugging.GetValueOnAnyThread() && GeneratedComponent->GetOwner())
+			if (PCGRuntimeGenSchedulerHelpers::CVarRuntimeGenerationEnableDebugging.GetValueOnAnyThread() && PartitionActor)
 			{
-				UE_LOG(LogPCG, Warning, TEXT("[RUNTIMEGEN] CLEANUP: '%s'"), *GeneratedComponent->GetOwner()->GetActorNameOrLabel());
+				UE_LOG(LogPCG, Warning, TEXT("[RUNTIMEGEN] CLEANUP: '%s'"), *PartitionActor->GetActorNameOrLabel());
 			}
 
 #if PCG_RGS_ONSCREENDEBUGMESSAGES
 			++PCGRuntimeGenSchedulerHelpers::Stats.CleanupCallCounter;
 #endif
 
-			GeneratedComponent->CleanupLocalImmediate(/*bRemoveComponents=*/true);
+			IPCGGraphExecutionState::FCleanupParams CleanupParams;
+			CleanupParams.bCleanupLocalSources = false; // Allow local sources to clean up themselves.
+			CleanupParams.bReleaseManagedResources = true;
+			CleanupParams.bImmediate = true;
+
+			GeneratedComponent->GetExecutionState().Cleanup(CleanupParams);
 		}
 	}
 
@@ -1773,8 +1785,8 @@ void FPCGRuntimeGenScheduler::CleanupDelayedRefreshComponents()
 			continue;
 		}
 
-		UPCGComponent* LocalComponent = OriginalComponent ? ActorAndComponentMapping->GetLocalComponent(GridDescriptor, GridCoords, OriginalComponent) : nullptr;
-		APCGPartitionActor* PartitionActor = LocalComponent ? Cast<APCGPartitionActor>(LocalComponent->GetOwner()) : nullptr;
+		IPCGGraphExecutionSource* LocalComponent = OriginalComponent ? ActorAndComponentMapping->GetLocalComponent(GridDescriptor, GridCoords, OriginalComponent) : nullptr;
+		APCGPartitionActor* PartitionActor = LocalComponent ? Cast<APCGPartitionActor>(CastChecked<UPCGComponent>(LocalComponent)->GetOwner()) : nullptr;
 
 		if (LocalComponent && PartitionActor)
 		{
@@ -1798,16 +1810,17 @@ void FPCGRuntimeGenScheduler::CleanupDelayedRefreshComponents()
 	GeneratedComponentsToRemove.Empty();
 }
 
-void FPCGRuntimeGenScheduler::RefreshComponent(UPCGComponent* InComponent, bool bRemovePartitionActors)
+void FPCGRuntimeGenScheduler::RefreshComponent(IPCGGraphExecutionSource* InComponent, bool bRemovePartitionActors)
 {
 	if (!InComponent || !ensure(IsInGameThread()))
 	{
 		return;
 	}
 
-	APCGPartitionActor* PartitionActor = Cast<APCGPartitionActor>(InComponent->GetOwner());
+	IPCGGraphExecutionState& ExecutionState = InComponent->GetExecutionState();
+	APCGPartitionActor* PartitionActor = Cast<APCGPartitionActor>(CastChecked<UPCGComponent>(InComponent)->GetOwner());
 	const bool bIsLocalComponent = PartitionActor != nullptr;
-	UPCGComponent* OriginalComponent = bIsLocalComponent ? PartitionActor->GetOriginalComponent(InComponent) : InComponent;
+	UPCGComponent* OriginalComponent = Cast<UPCGComponent>(ExecutionState.GetOriginalSource());
 
 	// If we are mid way through setting up an original component, early out from this refresh.
 	if (!OriginalComponent || OriginalComponent == OriginalComponentBeingGenerated)
@@ -1821,9 +1834,9 @@ void FPCGRuntimeGenScheduler::RefreshComponent(UPCGComponent* InComponent, bool
 	const bool bLoggingEnabled = PCGRuntimeGenSchedulerHelpers::CVarRuntimeGenerationEnableDebugging.GetValueOnAnyThread();
 
 	// Useful because we can run into generation order issues if components are left to continue generating.
-	if (InComponent->IsGenerating())
+	if (ExecutionState.IsGenerating())
 	{
-		InComponent->CancelGeneration();
+		ExecutionState.Cancel();
 	}
 
 	if (!bRemovePartitionActors)
@@ -1840,7 +1853,12 @@ void FPCGRuntimeGenScheduler::RefreshComponent(UPCGComponent* InComponent, bool
 			}
 
 			GeneratedComponentsToRemove.Emplace({ PartitionActor->GetPCGGridSize(), PartitionActor->GetGridCoord(), OriginalComponent, bIsLocalComponent ? InComponent : nullptr });
-			InComponent->CleanupLocalImmediate(/*bRemoveComponents=*/false);
+
+			IPCGGraphExecutionState::FCleanupParams CleanupParams;
+			CleanupParams.bReleaseManagedResources = false;
+			CleanupParams.bImmediate = true;
+
+			ExecutionState.Cleanup(CleanupParams);
 		}
 		else
 		{
@@ -1855,22 +1873,29 @@ void FPCGRuntimeGenScheduler::RefreshComponent(UPCGComponent* InComponent, bool
 					const FPCGGridDescriptor GridDescriptor = Key.GetGridDescriptor();
 					
 					// TODO - clean up local immediate will have a flag in the future to clean up the local components on its own, so this call to CleanupLocalImmediate will not be required
-					UPCGComponent* LocalComponent = ActorAndComponentMapping->GetLocalComponent(
+					IPCGGraphExecutionSource* LocalComponent = ActorAndComponentMapping->GetLocalComponent(
 						GridDescriptor,
 						Key.GetGridCoords(),
 						OriginalComponent);
 
-					if (bLoggingEnabled && LocalComponent && LocalComponent->GetOwner())
-					{
-						UE_LOG(LogPCG, Warning, TEXT("[RUNTIMEGEN] SHALLOW REFRESH LOCAL COMPONENT: '%s'"), *LocalComponent->GetOwner()->GetActorNameOrLabel());
-					}
-
 					if (LocalComponent)
 					{
-						LocalComponent->CleanupLocalImmediate(/*bRemoveComponents=*/false);
+						IPCGGraphExecutionState& LocalExecutionState = LocalComponent->GetExecutionState();
+
+						if (bLoggingEnabled && CastChecked<UPCGComponent>(LocalComponent)->GetOwner())
+						{
+							UE_LOG(LogPCG, Warning, TEXT("[RUNTIMEGEN] SHALLOW REFRESH LOCAL COMPONENT: '%s'"), *CastChecked<UPCGComponent>(LocalComponent)->GetOwner()->GetActorNameOrLabel());
+						}
 
+						IPCGGraphExecutionState::FCleanupParams CleanupParams;
+						CleanupParams.bReleaseManagedResources = false;
+						CleanupParams.bImmediate = true;
+
+						LocalExecutionState.Cleanup(CleanupParams);
+
+						// @todo_pcg: Instead of casting to PCGComponent, we should be able to cast to a RuntimeGenExecutionSource which provides this behavior.
 						// We need to make sure that the next time this is generated that it matches the original
-						LocalComponent->SetPropertiesFromOriginal(OriginalComponent);
+						CastChecked<UPCGComponent>(LocalComponent)->SetPropertiesFromOriginal(OriginalComponent);
 					}
 
 					GeneratedComponentsToRemove.Add(Key);
@@ -1884,17 +1909,21 @@ void FPCGRuntimeGenScheduler::RefreshComponent(UPCGComponent* InComponent, bool
 					OriginalComponent->IsPartitioned() ? 1 : 0);
 			}
 
-			InComponent->CleanupLocalImmediate(/*bRemoveComponents=*/false);
+			IPCGGraphExecutionState::FCleanupParams CleanupParams;
+			CleanupParams.bReleaseManagedResources = false;
+			CleanupParams.bImmediate = true;
+
+			ExecutionState.Cleanup(CleanupParams);
 		}
 	}
 	else
 	{
 		// Full cleanout path - cleanup existing components and return actors to the pool.
 
-		auto RefreshLocalComponent = [this, OriginalComponent, bLoggingEnabled](UPCGComponent* LocalComponent)
+		auto RefreshLocalComponent = [this, OriginalComponent, bLoggingEnabled](IPCGGraphExecutionSource* LocalComponent)
 		{
 			check(LocalComponent);
-			APCGPartitionActor* PartitionActor = CastChecked<APCGPartitionActor>(LocalComponent->GetOwner());
+			APCGPartitionActor* PartitionActor = CastChecked<APCGPartitionActor>(Cast<UPCGComponent>(LocalComponent)->GetOwner());
 
 			// Find the specific generation key for this component, if it exists, cleanup and generate.
 			FGridGenerationKey LocalComponentKey(PartitionActor->GetPCGGridSize(), PartitionActor->GetGridCoord(), OriginalComponent, LocalComponent);
@@ -1949,7 +1978,7 @@ void FPCGRuntimeGenScheduler::RefreshComponent(UPCGComponent* InComponent, bool
 				{
 					const FIntVector GridCoords = GenerationKey.GetGridCoords();
 
-					if (UPCGComponent* LocalComponent = ActorAndComponentMapping->GetLocalComponent(GridDescriptor, GridCoords, OriginalComponent))
+					if (IPCGGraphExecutionSource* LocalComponent = ActorAndComponentMapping->GetLocalComponent(GridDescriptor, GridCoords, OriginalComponent))
 					{
 						RefreshLocalComponent(LocalComponent);
 					}
diff --git a/Engine/Plugins/PCG/Source/PCG/Public/PCGComponentExecutionState.h b/Engine/Plugins/PCG/Source/PCG/Public/PCGComponentExecutionState.h
--- a/Engine/Plugins/PCG/Source/PCG/Public/PCGComponentExecutionState.h
+++ b/Engine/Plugins/PCG/Source/PCG/Public/PCGComponentExecutionState.h
@@ -35,13 +35,20 @@ public:
 	virtual bool IsPartitioned() const override;
 	virtual bool IsLocalSource() const override;
 	virtual FPCGGridDescriptor GetGridDescriptor(uint32 InGridSize) const override;
+	virtual EPCGHiGenGrid GetGenerationGrid() const override;
 	virtual uint32 GetGenerationGridSize() const override;
 	virtual FPCGTaskId GetGenerationTaskId() const override;
 	virtual void StoreOutputDataForPin(const FString& InResourceKey, const FPCGDataCollection& InData) const override;
 	virtual const FPCGDataCollection* RetrieveOutputDataForPin(const FString& InResourceKey) const override;
 	virtual void AddToManagedResources(UPCGManagedResource* InResource) override;
+	
+	virtual FPCGTaskId Generate(const FGenerateParams& InGenerateParams = {}) override;
+
+	UE_DEPRECATED(5.8, "Call/Implement Generate instead.")
 	virtual FPCGTaskId GenerateLocalGetTaskId(EPCGHiGenGrid Grid = EPCGHiGenGrid::Uninitialized) override;
 
+	virtual FPCGTaskId Cleanup(const FCleanupParams& InCleanupParams = {}) override;
+
 #if WITH_EDITOR
 	virtual const PCGUtils::FExtraCapture& GetExtraCapture() const override;
 	virtual PCGUtils::FExtraCapture& GetExtraCapture() override;
diff --git a/Engine/Plugins/PCG/Source/PCG/Public/PCGGraphExecutionStateInterface.h b/Engine/Plugins/PCG/Source/PCG/Public/PCGGraphExecutionStateInterface.h
--- a/Engine/Plugins/PCG/Source/PCG/Public/PCGGraphExecutionStateInterface.h
+++ b/Engine/Plugins/PCG/Source/PCG/Public/PCGGraphExecutionStateInterface.h
@@ -33,6 +33,30 @@ using FPCGDynamicTrackingPriority = double;
 class IPCGGraphExecutionState
 {
 public:
+	struct FGenerateParams
+	{
+		/** If the execution source is partitioned, will forward the generate call to its registered local execution sources. */
+		bool bGenerateLocalSources = true;
+
+		/** Dependencies to wait on before the generation can execute. */
+		TArray<FPCGTaskId> Dependencies = {};
+	};
+
+	struct FCleanupParams
+	{
+		/** If the execution source is partitioned, will forward the clean up call to its registered local execution sources. */
+		bool bCleanupLocalSources = true;
+
+		/** Performs a hard release (no re-use) of generated managed resources. */
+		bool bReleaseManagedResources = true;
+
+		/** Performs the cleanup immediately without scheduling. Ignores dependencies. */
+		bool bImmediate = false;
+
+		/** Dependencies to wait on before the cleanup can execute. */
+		TArray<FPCGTaskId> Dependencies = {};
+	};
+
 	virtual ~IPCGGraphExecutionState() = default;
 
 	/** Returns a UPCGData representation of the ExecutionState. */
@@ -86,9 +110,17 @@ public:
 	/** Gets the task ID for the current generation task if it exists. */
 	virtual FPCGTaskId GetGenerationTaskId() const { return InvalidPCGTaskId; }
 
-	/** Requests the execution source to generate only on the specified grid level (all grid levels if EPCGHiGenGrid::Uninitialized). */
+	/** Requests the execution source to generate. */
+	virtual FPCGTaskId Generate(const FGenerateParams& InGenerateParams) { return InvalidPCGTaskId; }
+	FPCGTaskId Generate() { return Generate(FGenerateParams()); }
+
+	UE_DEPRECATED(5.8, "Call/Implement Generate instead.")
 	virtual FPCGTaskId GenerateLocalGetTaskId(EPCGHiGenGrid Grid = EPCGHiGenGrid::Uninitialized) { return InvalidPCGTaskId; }
 
+	/** Requests the execution source to clean up. */
+	virtual FPCGTaskId Cleanup(const FCleanupParams& InCleanupParams) { return InvalidPCGTaskId; }
+	FPCGTaskId Cleanup() { return Cleanup(FCleanupParams()); }
+
 	/** Store data with a resource key that identifies the pin. */
 	virtual void StoreOutputDataForPin(const FString& InResourceKey, const FPCGDataCollection& InData) const {}
 
@@ -151,6 +183,9 @@ public:
 	/** Returns a GridDescriptor based on this execution source for the specified grid size. */
 	virtual FPCGGridDescriptor GetGridDescriptor(uint32 InGridSize) const;
 
+	/** Get the HiGen grid this execution source executes on. */
+	virtual EPCGHiGenGrid GetGenerationGrid() const { return EPCGHiGenGrid::Uninitialized; }
+
 	/** Get the size for the HiGen grid this execution source executes on. */
 	virtual uint32 GetGenerationGridSize() const { return PCGHiGenGrid::UninitializedGridSize(); }
 };
diff --git a/Engine/Plugins/PCG/Source/PCG/Public/RuntimeGen/PCGRuntimeGenScheduler.h b/Engine/Plugins/PCG/Source/PCG/Public/RuntimeGen/PCGRuntimeGenScheduler.h
--- a/Engine/Plugins/PCG/Source/PCG/Public/RuntimeGen/PCGRuntimeGenScheduler.h
+++ b/Engine/Plugins/PCG/Source/PCG/Public/RuntimeGen/PCGRuntimeGenScheduler.h
@@ -3,7 +3,6 @@
 #pragma once
 
 #include "PCGCommon.h"
-#include "PCGComponent.h"
 #include "Grid/PCGGridDescriptor.h"
 #include "RuntimeGen/PCGRuntimeGenChangeDetection.h"
 
@@ -58,7 +57,7 @@ struct FPCGVirtualTexturePrimingInfo
  * APCGPartitionActors can be created/destroyed on demand or provided by a dynamically growing pool of actors. If
  * enabled, the pool will double in capacity anytime the number of available PAs reaches zero.
  * 
- * Components and PartitionActors created by the Runtime Generation Scheduler should be managed exclusively by the
+ * Execution sources and partition actors created by the Runtime Generation Scheduler should be managed exclusively by the
  * runtime gen scheduling system.
  */
 class FPCGRuntimeGenScheduler
@@ -87,7 +86,7 @@ public:
 #endif
 
 protected:
-	// Grid size, grid coords, original component, local componnent
+	// Grid size, grid coords, original component, local source
 	struct FGridGenerationKey
 	{
 		FGridGenerationKey(uint32 InGridSize, const FIntVector& InGridCoords, UPCGComponent* InOriginalComponent)
@@ -95,7 +94,7 @@ protected:
 		{
 		}
 
-		FGridGenerationKey(uint32 InGridSize, const FIntVector& InGridCoords, UPCGComponent* InOriginalComponent, UPCGComponent* InLocalComponent)
+		FGridGenerationKey(uint32 InGridSize, const FIntVector& InGridCoords, UPCGComponent* InOriginalComponent, IPCGGraphExecutionSource* InLocalComponent)
 			: GridSize(InGridSize)
 			, GridCoords(InGridCoords)
 			, OriginalComponent(InOriginalComponent)
@@ -119,8 +118,8 @@ protected:
 		FIntVector GetGridCoords() const { return GridCoords; }
 		UPCGComponent* GetOriginalComponent() const { return OriginalComponent.ResolveObjectPtr(); }
 
-		UPCGComponent* GetCachedLocalComponent() const { return CachedLocalComponent.Get(); }
-		void SetCachedLocalComponent(UPCGComponent* InComponent) const { CachedLocalComponent = InComponent; }
+		IPCGGraphExecutionSource* GetCachedLocalComponent() const { return CachedLocalComponent.Get(); }
+		void SetCachedLocalComponent(IPCGGraphExecutionSource* InLocalComponent) const { CachedLocalComponent = InLocalComponent; }
 
 		FPCGGridDescriptor GetGridDescriptor() const;
 
@@ -129,8 +128,8 @@ protected:
 		FIntVector GridCoords = FIntVector(0);
 		TObjectKey<UPCGComponent> OriginalComponent;
 
-		// Optional/opportunistic cached local component if one has been created.
-		mutable TWeakObjectPtr<UPCGComponent> CachedLocalComponent;
+		// Optional/opportunistic cached local execution source if one has been created.
+		mutable TWeakInterfacePtr<IPCGGraphExecutionSource> CachedLocalComponent;
 
 		friend uint32 GetTypeHash(const FGridGenerationKey& In)
 		{
@@ -152,9 +151,7 @@ protected:
 	};
 
 	/** Queue nearby components for generation. */
-	void TickQueueComponentsForGeneration(
-		const FTickQueueComponentsForGenerationInputs& Inputs,
-		TMap<FGridGenerationKey, double>& OutComponentsToGenerate);
+	void TickQueueComponentsForGeneration(const FTickQueueComponentsForGenerationInputs& Inputs, TMap<FGridGenerationKey, double>& OutComponentsToGenerate);
 
 	/** Perform immediate cleanup on components that become out of range. */
 	void TickCleanup(const TSet<IPCGGenSourceBase*>& GenSources, const APCGWorldActor* InPCGWorldActor, double InEndTime);
@@ -172,13 +169,13 @@ protected:
 	void CleanupLocalComponents(const APCGWorldActor* InPCGWorldActor);
 
 	/** Cleanup a component and remove it from the GeneratedComponents set. */
-	void CleanupComponent(const FGridGenerationKey& GenerationKey, UPCGComponent* GeneratedComponent);
+	void CleanupComponent(const FGridGenerationKey& GenerationKey, IPCGGraphExecutionSource* GeneratedComponent);
 
 	/** Remove components from the GeneratedComponents set that have been marked for delayed refresh. Fully cleanup any that would be leaked otherwise. */
 	void CleanupDelayedRefreshComponents();
 
 	/** Refresh a generated component. bRemovePartitionActors will also perform a full cleanup of PAs and local components. */
-	void RefreshComponent(UPCGComponent* InComponent, bool bRemovePartitionActors = false);
+	void RefreshComponent(IPCGGraphExecutionSource* InComponent, bool bRemovePartitionActors = false);
 	
 	/** Grabs an empty RuntimeGen PA from the PartitionActorPool and initializes it at the given GridSize and GridCoords. If no PAs are available in the pool,
 	* the pool capacity will double and new PAs will be created.
@@ -198,10 +195,10 @@ private:
 	/** Called on world streaming events. */
 	void OnLevelStreamingStateChanged(UWorld* InWorld, const ULevelStreaming* InLevelStreaming, ULevel* InLevelIfLoaded, ELevelStreamingState InPreviousState, ELevelStreamingState InNewState);
 
-	/** Cleanup a local compnent from the Partition actor. */
-	void CleanupLocalComponent(APCGPartitionActor* PartitionActor, UPCGComponent* LocalComponent);
+	/** Cleanup a local component from the Partition actor. */
+	void CleanupLocalComponent(APCGPartitionActor* PartitionActor, IPCGGraphExecutionSource* LocalComponent);
 
-	/** Cleanup the remaining comopnents that are not part of the GeneratedComponents */
+	/** Cleanup the remaining components that are not part of the GeneratedComponents */
 	void CleanupRemainingComponents(UPCGComponent* InOriginalComponent);
 
 	/** Tracks the generated components managed by the RuntimeGenScheduler. For local components, this generation key will hold the original component.
