[PCG] UPCGSubsystem::IsComponentBeingUnloaded
- Fix runtime gen Original Component not cleaning up Local Components when being unloaded through WP or Level Instance unload
- Replace complex conditions now that we can ask the PCGSubsystem

diff --git a/Engine/Plugins/PCG/Source/PCG/Private/PCGComponent.cpp b/Engine/Plugins/PCG/Source/PCG/Private/PCGComponent.cpp
--- a/Engine/Plugins/PCG/Source/PCG/Private/PCGComponent.cpp
+++ b/Engine/Plugins/PCG/Source/PCG/Private/PCGComponent.cpp
@@ -1952,27 +1952,29 @@ void UPCGComponent::EndPlay(const EEndPlayReason::Type EndPlayReason)
 void UPCGComponent::OnUnregister()
 {
 #if WITH_EDITOR
-	// Track that we were unregistered through loading to avoid clean up
-	bUnregisteredThroughLoading = UE::GetIsEditorLoadingPackage();
-
 	if (UPCGSubsystem* Subsystem = GetSubsystem())
 	{
+		// Track if Component got unregistered through unloading (not property changes / blueprint construction)
+		bUnregisteredThroughLoading = Subsystem->IsComponentBeingUnloaded(this);
+
 		if (!PCGHelpers::IsRuntimeOrPIE())
 		{
 			Subsystem->CancelGeneration(this);
 		}
 
 		Subsystem->OnPCGComponentUnregistered.Broadcast(IsValidChecked(this) ? this : nullptr);
+
+		// If Original Component is managed by runtime and being unloaded, unregister immediatly so that runtime generated Local components also get cleaned up. (BeginDestroy is too late)
+		if (bUnregisteredThroughLoading && IsManagedByRuntimeGenSystem() && !IsLocalComponent())
+		{
+			Subsystem->UnregisterPCGComponent(this, /*bForce=*/true);
+		}
 	}
 
-	// We shouldn't cleanup resources in OnUnregister in most cases.
+	// We shouldn't cleanup resources in OnUnregister in most persistent cases.
 	// This specific case is to handle World Partition unloading of actors where our actor might have a PCG Component which generated some Preview Resources.
 	// In this case we do want to release those resources and the actors they might hold on to. If we don't those actors become orphaned and we might end up with duplicates when reloading the component.
-	if (PCGComponent::CVarDeletePreviewResourcesWhenUnloading.GetValueOnAnyThread() &&
-		CurrentEditingMode == EPCGEditorDirtyMode::Preview && // We have preview resources
-		UE::GetIsEditorLoadingPackage() && // We are loading (we don't want to release resources if we are getting unregistered outside of a loading scope)
-		(!IsCreatedByConstructionScript() || !GIsReconstructingBlueprintInstances) && // But not inside a construction script
-		(GetWorld() && GetWorld()->IsPartitionedWorld())) // and in a Partitioned World because if not the preview resources will get unloaded with us (no partial unloading of levels)
+	if (bUnregisteredThroughLoading && PCGComponent::CVarDeletePreviewResourcesWhenUnloading.GetValueOnAnyThread() && CurrentEditingMode == EPCGEditorDirtyMode::Preview)
 	{
 		TSet<TSoftObjectPtr<AActor>> ActorsToDelete;
 		for (UPCGManagedResource* ManagedResource : GeneratedResources)
diff --git a/Engine/Plugins/PCG/Source/PCG/Private/Subsystems/PCGSubsystem.cpp b/Engine/Plugins/PCG/Source/PCG/Private/Subsystems/PCGSubsystem.cpp
--- a/Engine/Plugins/PCG/Source/PCG/Private/Subsystems/PCGSubsystem.cpp
+++ b/Engine/Plugins/PCG/Source/PCG/Private/Subsystems/PCGSubsystem.cpp
@@ -217,6 +217,15 @@ void UPCGSubsystem::Deinitialize()
 		GEditor->OnPreviewPlatformChanged().RemoveAll(this);
 		GEditor->OnPreviewFeatureLevelChanged().RemoveAll(this);
 	}
+
+	if (UWorld* World = GetWorld(); World && !World->IsGameWorld())
+	{
+		// Unregister persistent level events
+		OnLevelRemovedFromWorld(World->PersistentLevel, World);
+
+		FWorldDelegates::LevelAddedToWorld.RemoveAll(this);
+		FWorldDelegates::LevelRemovedFromWorld.RemoveAll(this);
+	}
 #endif
 
 	Super::Deinitialize();
@@ -257,6 +266,16 @@ void UPCGSubsystem::PostInitialize()
 		GEditor->OnPreviewPlatformChanged().AddUObject(this, &UPCGSubsystem::OnPreviewPlatformChanged);
 		GEditor->OnPreviewFeatureLevelChanged().AddUObject(this, &UPCGSubsystem::OnPreviewFeatureLevelChanged);
 	}
+
+	if (UWorld* World = GetWorld(); World && !World->IsGameWorld())
+	{
+		// Register persistent level events
+		OnLevelAddedToWorld(World->PersistentLevel, World);
+
+		// Listen for sub-levels being added/removed to register per-level events (ex: Level Instances)
+		FWorldDelegates::LevelAddedToWorld.AddUObject(this, &UPCGSubsystem::OnLevelAddedToWorld);
+		FWorldDelegates::LevelRemovedFromWorld.AddUObject(this, &UPCGSubsystem::OnLevelRemovedFromWorld);
+	}
 #endif
 }
 
@@ -1951,6 +1970,61 @@ void UPCGSubsystem::OnObjectsReplaced(const TMap<UObject*, UObject*>& InOldToNew
 	}
 }
 
+void UPCGSubsystem::OnLevelAddedToWorld(ULevel* InLevel, UWorld* InWorld)
+{
+	if (InLevel && InWorld == GetWorld())
+	{
+		InLevel->OnLoadedActorRemovedFromLevelPreEvent.AddUObject(this, &UPCGSubsystem::OnLoadedActorRemovedFromLevelPreEvent);
+		InLevel->OnLoadedActorRemovedFromLevelPostEvent.AddUObject(this, &UPCGSubsystem::OnLoadedActorRemovedFromLevelPostEvent);
+	}
+}
+
+void UPCGSubsystem::OnLevelRemovedFromWorld(ULevel* InLevel, UWorld* InWorld)
+{
+	if (InLevel && InWorld == GetWorld())
+	{
+		InLevel->OnLoadedActorRemovedFromLevelPreEvent.RemoveAll(this);
+		InLevel->OnLoadedActorRemovedFromLevelPostEvent.RemoveAll(this);
+	}
+}
+
+void UPCGSubsystem::OnLoadedActorRemovedFromLevelPreEvent(const TArray<AActor*>& InActors)
+{
+	check(IsInGameThread());
+	++NumLevelUnloadingActors;
+}
+
+void UPCGSubsystem::OnLoadedActorRemovedFromLevelPostEvent(const TArray<AActor*>& InActors)
+{
+	check(IsInGameThread());
+	check(NumLevelUnloadingActors > 0);
+	--NumLevelUnloadingActors;
+}
+
+bool UPCGSubsystem::IsComponentBeingUnloaded(UActorComponent* InComponent) const
+{
+	// World Partition actor unload
+	if (NumLevelUnloadingActors > 0)
+	{
+		return true;
+	}
+
+	// Sub-level being removed 
+	ULevel* OuterLevel = InComponent ? InComponent->GetTypedOuter<ULevel>() : nullptr;
+	if (OuterLevel && OuterLevel->bIsBeingRemoved)
+	{
+		return true;
+	}
+
+	// World being unloaded (map change)
+	if (UWorld* World = GetWorld(); World && World->IsBeingCleanedUp())
+	{
+		return true;
+	}
+
+	return false;
+}
+
 void UPCGSubsystem::OnPCGGraphCancelled(UPCGComponent* InComponent)
 {
 	ActorAndComponentMapping->OnPCGGraphCancelled(InComponent);
diff --git a/Engine/Plugins/PCG/Source/PCG/Public/Subsystems/PCGSubsystem.h b/Engine/Plugins/PCG/Source/PCG/Public/Subsystems/PCGSubsystem.h
--- a/Engine/Plugins/PCG/Source/PCG/Public/Subsystems/PCGSubsystem.h
+++ b/Engine/Plugins/PCG/Source/PCG/Public/Subsystems/PCGSubsystem.h
@@ -111,6 +111,9 @@ public:
 
 	UE_API void SetConstructionScriptSourceComponent(UPCGComponent* InComponent);
 	UE_API bool RemoveAndCopyConstructionScriptSourceComponent(AActor* InComponentOwner, FName InComponentName, UPCGComponent*& OutSourceComponent);
+
+	/** Returns true if the Component is being unloaded through World Partition editor unloading or if the level the component is part of is being removed from world */
+	UE_API bool IsComponentBeingUnloaded(UActorComponent* InActorComponent) const;
 #endif
 
 	UE_API APCGWorldActor* GetPCGWorldActor();
@@ -372,6 +375,10 @@ private:
 	UE_API void SetChainedDispatchToLocalComponents(bool bInChainedDispatch);
 
 	void OnObjectsReplaced(const TMap<UObject*, UObject*>& InOldToNewInstances);
+	void OnLevelAddedToWorld(ULevel* InLevel, UWorld* InWorld);
+	void OnLevelRemovedFromWorld(ULevel* InLevel, UWorld* InWorld);
+	void OnLoadedActorRemovedFromLevelPreEvent(const TArray<AActor*>& InActors);
+	void OnLoadedActorRemovedFromLevelPostEvent(const TArray<AActor*>& InActors);
 #endif // WITH_EDITOR
 	
 private:
@@ -413,6 +420,8 @@ private:
 	PCG::FLock PCGWorldActorLock;
 
 #if WITH_EDITOR
+	int32 NumLevelUnloadingActors = 0;
+
 	using FConstructionScriptSourceComponents = TMap<FName, TObjectKey<UPCGComponent>>;
 	TMap<TObjectKey<AActor>, FConstructionScriptSourceComponents> PerActorConstructionScriptSourceComponents;
 
