[PCG][Bughawk] Update cvar description

#rnx

diff --git a/Engine/Plugins/PCG/Source/PCG/Private/Compute/DataInterfaces/Elements/PCGStaticMeshSpawnerDataInterface.cpp b/Engine/Plugins/PCG/Source/PCG/Private/Compute/DataInterfaces/Elements/PCGStaticMeshSpawnerDataInterface.cpp
--- a/Engine/Plugins/PCG/Source/PCG/Private/Compute/DataInterfaces/Elements/PCGStaticMeshSpawnerDataInterface.cpp
+++ b/Engine/Plugins/PCG/Source/PCG/Private/Compute/DataInterfaces/Elements/PCGStaticMeshSpawnerDataInterface.cpp
@@ -50,7 +50,7 @@ namespace PCGStaticMeshSpawnerDataInterface
 	static TAutoConsoleVariable<bool> CVarCreatePrimitivesComponentless(
 		TEXT("pcg.RuntimeGeneration.ISM.ComponentlessPrimitives"),
 		false,
-		TEXT("Uses a component-less path for creating primitives in game worlds (PIE and standalone)."));
+		TEXT("Uses a component-less path for creating primitives."));
 }
 
 void UPCGStaticMeshSpawnerDataInterface::GetSupportedInputs(TArray<FShaderFunctionDefinition>& OutFunctions) const
