[PCG] Added selection refresh after generation/apply.

diff --git a/Engine/Plugins/PCG/Source/PCG/Public/PCGComponent.h b/Engine/Plugins/PCG/Source/PCG/Public/PCGComponent.h
--- a/Engine/Plugins/PCG/Source/PCG/Public/PCGComponent.h
+++ b/Engine/Plugins/PCG/Source/PCG/Public/PCGComponent.h
@@ -392,6 +392,7 @@ public:
 
 	/** Returns task ids to do internal chaining */
 	FPCGTaskId GetGenerationTaskId() const { return CurrentGenerationTask; }
+	FPCGTaskId GetCleanupTaskId() const { return CurrentCleanupTask; }
 
 #if WITH_EDITOR
 	/** Schedules refresh of the component. If bCancelExistingRefresh is true, any existing refresh is cancelled and a new one is scheduled. */
diff --git a/Engine/Plugins/PCG/Source/PCGEditor/Private/Details/PCGComponentDetails.cpp b/Engine/Plugins/PCG/Source/PCGEditor/Private/Details/PCGComponentDetails.cpp
--- a/Engine/Plugins/PCG/Source/PCGEditor/Private/Details/PCGComponentDetails.cpp
+++ b/Engine/Plugins/PCG/Source/PCGEditor/Private/Details/PCGComponentDetails.cpp
@@ -13,7 +13,9 @@
 #include "DetailWidgetRow.h"
 #include "Modules/ModuleManager.h"
 #include "LevelEditor.h"
+#include "LevelEditorSubsystem.h"
 #include "Algo/NoneOf.h"
+#include "Elements/Framework/TypedElementSelectionSet.h"
 
 #define LOCTEXT_NAMESPACE "PCGComponentDetails"
 
@@ -22,6 +24,37 @@ static TAutoConsoleVariable<bool> CVarBroadcastComponentChanged(
 		false,
 		TEXT("Controls whether we should broadcast component changes to the level editor"));
 
+namespace PCGComponentDetails
+{
+	void TriggerSelectionUpdate(UPCGComponent* Component, FPCGTaskId TaskId)
+	{
+		if (!GEditor || !GEditor->GetEditorSubsystem<ULevelEditorSubsystem>() || !Component || TaskId == InvalidPCGTaskId)
+		{
+			return;
+		}
+
+		if (UPCGSubsystem* PCGSubsystem = Component->GetSubsystem())
+		{
+			PCGSubsystem->ScheduleGeneric([]()
+			{
+				if (!GEditor)
+				{
+					return true;
+				}
+
+				// Since we might have been moving things around (component wise) we need to force a selection set change, 
+				// Which in turn will trigger a component visualizer update.
+				if(UTypedElementSelectionSet* SelectionSet = GEditor->GetEditorSubsystem<ULevelEditorSubsystem>()->GetSelectionSet())
+				{
+					SelectionSet->OnChanged().Broadcast(SelectionSet);
+				}
+
+				return true;
+			}, Component, { TaskId });
+		}
+	}
+} // namespace PCGComponentDetails
+
 TSharedRef<IDetailCustomization> FPCGComponentDetails::MakeInstance()
 {
 	return MakeShareable(new FPCGComponentDetails());
@@ -321,6 +354,7 @@ FReply FPCGComponentDetails::OnGenerateClicked()
 			}
 
 			PCGComponent->Generate(/*bForce=*/bIsControlDown);
+			PCGComponentDetails::TriggerSelectionUpdate(PCGComponent, PCGComponent->GetGenerationTaskId());
 		}
 	}
 
@@ -397,6 +431,8 @@ FReply FPCGComponentDetails::OnCleanupClicked()
 			{
 				Component.Get()->Cleanup();
 			}
+
+			PCGComponentDetails::TriggerSelectionUpdate(Component.Get(), Component->GetCleanupTaskId());
 		}
 	}
 
