[PCG] Adjusted condition to be a bit more robust

#rnx

diff --git a/Engine/Plugins/PCG/Source/PCGEditor/Private/Details/PCGComponentDetails.cpp b/Engine/Plugins/PCG/Source/PCGEditor/Private/Details/PCGComponentDetails.cpp
--- a/Engine/Plugins/PCG/Source/PCGEditor/Private/Details/PCGComponentDetails.cpp
+++ b/Engine/Plugins/PCG/Source/PCGEditor/Private/Details/PCGComponentDetails.cpp
@@ -28,7 +28,7 @@ namespace PCGComponentDetails
 {
 	void TriggerSelectionUpdate(UPCGComponent* Component, FPCGTaskId TaskId)
 	{
-		if (!GEditor || !GEditor->GetEditorSubsystem<ULevelEditorSubsystem>() || !Component || TaskId == InvalidPCGTaskId)
+		if (!GEditor || !Component || TaskId == InvalidPCGTaskId)
 		{
 			return;
 		}
@@ -37,16 +37,14 @@ namespace PCGComponentDetails
 		{
 			PCGSubsystem->ScheduleGeneric([]()
 			{
-				if (!GEditor)
+				if (GEditor && GEditor->GetEditorSubsystem<ULevelEditorSubsystem>())
 				{
-					return true;
-				}
-
-				// Since we might have been moving things around (component wise) we need to force a selection set change, 
-				// Which in turn will trigger a component visualizer update.
-				if(UTypedElementSelectionSet* SelectionSet = GEditor->GetEditorSubsystem<ULevelEditorSubsystem>()->GetSelectionSet())
-				{
-					SelectionSet->OnChanged().Broadcast(SelectionSet);
+					// Since we might have been moving things around (component wise) we need to force a selection set change, 
+					// Which in turn will trigger a component visualizer update.
+					if (UTypedElementSelectionSet* SelectionSet = GEditor->GetEditorSubsystem<ULevelEditorSubsystem>()->GetSelectionSet())
+					{
+						SelectionSet->OnChanged().Broadcast(SelectionSet);
+					}
 				}
 
 				return true;
