[PCG] Made sure the components are properly visible after 'Apply' on a new actor in the pcg editor toolmode

#rnx

diff --git a/Engine/Plugins/PCG/Source/PCGEditor/Private/EditorMode/Tools/PCGInteractiveToolSettings.cpp b/Engine/Plugins/PCG/Source/PCGEditor/Private/EditorMode/Tools/PCGInteractiveToolSettings.cpp
--- a/Engine/Plugins/PCG/Source/PCGEditor/Private/EditorMode/Tools/PCGInteractiveToolSettings.cpp
+++ b/Engine/Plugins/PCG/Source/PCGEditor/Private/EditorMode/Tools/PCGInteractiveToolSettings.cpp
@@ -12,6 +12,7 @@
 #include "Subsystems/PCGSubsystem.h"
 
 #include "InteractiveToolManager.h"
+#include "LevelEditorSubsystem.h"
 #include "BaseBehaviors/KeyInputBehavior.h"
 #include "Logging/StructuredLog.h"
 
@@ -216,6 +217,16 @@ void UPCGInteractiveToolSettings::Apply(UInteractiveTool* OwningTool)
 	}
 
 	OwningTool->GetToolManager()->EndUndoTransaction();
+
+	// Since we might have been moving things around (component wise) we need to force a selection set change, 
+	// Which in turn will trigger a component visualizer update.
+	if (GEditor && GEditor->GetEditorSubsystem<ULevelEditorSubsystem>())
+	{
+		if (UTypedElementSelectionSet* SelectionSet = GEditor->GetEditorSubsystem<ULevelEditorSubsystem>()->GetSelectionSet())
+		{
+			SelectionSet->OnChanged().Broadcast(SelectionSet);
+		}
+	}
 }
 
 void UPCGInteractiveToolSettings::Cancel(UInteractiveTool* OwningTool)
