[PCG] Expose API on Graph and Edge descriptive functions

#rnx

diff --git a/Engine/Plugins/PCG/Source/PCG/Public/PCGEdge.h b/Engine/Plugins/PCG/Source/PCG/Public/PCGEdge.h
--- a/Engine/Plugins/PCG/Source/PCG/Public/PCGEdge.h
+++ b/Engine/Plugins/PCG/Source/PCG/Public/PCGEdge.h
@@ -52,10 +52,10 @@ public:
 	const UPCGNode* GetOutputNode() const;
 
 	UFUNCTION(BlueprintCallable, Category = Properties)
-	FName GetInputPinLabel() const;
+	UE_API FName GetInputPinLabel() const;
 
 	UFUNCTION(BlueprintCallable, Category = Properties)
-	FName GetOutputPinLabel() const;
+	UE_API FName GetOutputPinLabel() const;
 
 };
 
diff --git a/Engine/Plugins/PCG/Source/PCG/Public/PCGGraph.h b/Engine/Plugins/PCG/Source/PCG/Public/PCGGraph.h
--- a/Engine/Plugins/PCG/Source/PCG/Public/PCGGraph.h
+++ b/Engine/Plugins/PCG/Source/PCG/Public/PCGGraph.h
@@ -388,7 +388,7 @@ public:
 
 	/** Returns all the edges of the graph */
 	UFUNCTION(BlueprintCallable, Category = Graph)
-	TArray<UPCGEdge*> GetAllEdges() const;
+	UE_API TArray<UPCGEdge*> GetAllEdges() const;
 
 	/** Duplicate a given node by creating a new node with the same settings and properties, but without any edges and add it to the graph */
 	UE_API TObjectPtr<UPCGNode> ReconstructNewNode(const UPCGNode* InNode);
