[PCG] Don't cook PCG graph if the graph is marked editor only.

#rnx

diff --git a/Engine/Plugins/PCG/Source/PCG/Private/PCGGraph.cpp b/Engine/Plugins/PCG/Source/PCG/Private/PCGGraph.cpp
--- a/Engine/Plugins/PCG/Source/PCG/Private/PCGGraph.cpp
+++ b/Engine/Plugins/PCG/Source/PCG/Private/PCGGraph.cpp
@@ -731,7 +731,7 @@ void UPCGGraph::PreSave(FObjectPreSaveContext ObjectSaveContext)
 		PCGEditorModule->OnGraphPreSave(this, ObjectSaveContext);
 	}
 
-	if (ObjectSaveContext.IsCooking())
+	if (ObjectSaveContext.IsCooking() && !IsEditorOnly())
 	{
 		FPCGGraphCompiler GraphCompiler(/*bIsCooking=*/true);
 
