[PCG] Improved 'close to line' test in polygon creation and added axis to default to when this happens.

diff --git a/Engine/Plugins/PCG/Source/PCG/Private/Elements/Polygon/PCGCreatePolygon2D.cpp b/Engine/Plugins/PCG/Source/PCG/Private/Elements/Polygon/PCGCreatePolygon2D.cpp
--- a/Engine/Plugins/PCG/Source/PCG/Private/Elements/Polygon/PCGCreatePolygon2D.cpp
+++ b/Engine/Plugins/PCG/Source/PCG/Private/Elements/Polygon/PCGCreatePolygon2D.cpp
@@ -135,12 +135,22 @@ bool FPCGCreatePolygonElement::ExecuteInternal(FPCGContext* Context) const
 		// Project the positions so they lie on a plane. Note that we'll assume at this point that the polygon is always closed, regardless of whether that's true or not, without loss of generality.
 		FVector PlaneNormal = FVector::UpVector;
 		FVector PlaneCentroid = FVector::ZeroVector;
-		PolygonTriangulation::ComputePolygonPlane(Positions, PlaneNormal, PlaneCentroid);
+		double PolygonArea = PolygonTriangulation::ComputePolygonPlane(Positions, PlaneNormal, PlaneCentroid);
 
-		// If the polygon is a point/line or otherwise has a null area, assume normal
-		if (PlaneNormal.SquaredLength() < UE_KINDA_SMALL_NUMBER)
+		// Compute bounds for the input points so we have a comparison basis vs the polygon area
+		FBox PolygonBounds(EForceInit::ForceInit);
+
+		for (const FVector& Position : Positions)
+		{
+			PolygonBounds += Position;
+		}
+
+		const double LargestDimension = PolygonBounds.GetSize().GetMax();
+
+		// If the polygon is a point/line or otherwise has a very small area, use the provided default up axis.
+		if (PolygonArea <= UE_SMALL_NUMBER * FMath::Square(LargestDimension))
 		{
-			PlaneNormal = FVector::UpVector;
+			PlaneNormal = Settings->DefaultPolygonNormal.GetSafeNormal(UE_SMALL_NUMBER, FVector::ZAxisVector);
 		}
 
 		FTransform PlaneTransform;
diff --git a/Engine/Plugins/PCG/Source/PCG/Public/Elements/Polygon/PCGCreatePolygon2D.h b/Engine/Plugins/PCG/Source/PCG/Public/Elements/Polygon/PCGCreatePolygon2D.h
--- a/Engine/Plugins/PCG/Source/PCG/Public/Elements/Polygon/PCGCreatePolygon2D.h
+++ b/Engine/Plugins/PCG/Source/PCG/Public/Elements/Polygon/PCGCreatePolygon2D.h
@@ -76,4 +76,8 @@ public:
 	/** Maximum squared distance before we need to subdivide a segment again as part of the spline discretization to a polygon. */
 	UPROPERTY(EditAnywhere, BlueprintReadWrite, Category = Settings, AdvancedDisplay, meta = (ClampMin = 0.001, PCG_Overridable))
 	double SplineMaxDiscretizationError = 1.0;
+
+	/** Default axis to use when the input has an ambiguous shape (lines, collinear polygons, etc.) */
+	UPROPERTY(EditAnywhere, BlueprintReadWrite, Category = Settings, AdvancedDisplay, meta = (PCG_Overridable))
+	FVector DefaultPolygonNormal = FVector::ZAxisVector;
 };
\ No newline at end of file
