[PCG] More API exposure for graph construction

#rnx

diff --git a/Engine/Plugins/PCG/Source/PCG/Public/PCGEdge.h b/Engine/Plugins/PCG/Source/PCG/Public/PCGEdge.h
--- a/Engine/Plugins/PCG/Source/PCG/Public/PCGEdge.h
+++ b/Engine/Plugins/PCG/Source/PCG/Public/PCGEdge.h
@@ -46,10 +46,10 @@ public:
 	UE_API const UPCGPin* GetOtherPin(const UPCGPin* Pin) const;
 
 	UFUNCTION(BlueprintCallable, Category = Properties)
-	const UPCGNode* GetInputNode() const;
+	UE_API const UPCGNode* GetInputNode() const;
 
 	UFUNCTION(BlueprintCallable, Category = Properties)
-	const UPCGNode* GetOutputNode() const;
+	UE_API const UPCGNode* GetOutputNode() const;
 
 	UFUNCTION(BlueprintCallable, Category = Properties)
 	UE_API FName GetInputPinLabel() const;
