[PCG] PCGGraphExecutionRegistry
- Now possible to register execution source providers through a common registry
- Use this registry to gather execution sources to display in Graph debugger
- Implement a default provider for UPCGComponent & UPCGDefaultExecutionSource

diff --git a/Engine/Plugins/PCG/Source/PCG/Private/PCGGraphExecutionRegistry.cpp b/Engine/Plugins/PCG/Source/PCG/Private/PCGGraphExecutionRegistry.cpp
new file mode 100644
--- /dev/null
+++ b/Engine/Plugins/PCG/Source/PCG/Private/PCGGraphExecutionRegistry.cpp
@@ -0,0 +1,49 @@
+// Copyright Epic Games, Inc. All Rights Reserved.
+
+#include "PCGGraphExecutionRegistry.h"
+#include "PCGModule.h"
+
+int32 FPCGGraphExecutionRegistry::NextHandle = 0;
+
+void IPCGGraphExecutionSourceProvider::BroadcastExecutionSourcesChanged()
+{
+	if (FPCGModule* PCGModule = FPCGModule::GetPCGModule())
+	{
+		PCGModule->GetGraphExecutionRegistry().BroadcastGraphExecutionSourcesChanged();
+	}
+}
+
+TArray<FPCGGraphExecutionSourceDescriptor> FPCGGraphExecutionRegistry::GatherExecutionSources() const
+{
+	TArray<FPCGGraphExecutionSourceDescriptor> ExecutionSources;
+	for (TSharedRef<IPCGGraphExecutionSourceProvider> Provider : RegisteredExecutionSourceProviders)
+	{
+		ExecutionSources.Append(Provider->GatherExecutionSources());
+	}
+
+	return ExecutionSources;
+}
+
+void FPCGGraphExecutionRegistry::Shutdown()
+{
+	RegisteredExecutionSourceProviders.Empty();
+}
+
+FPCGGraphExecutionSourceProviderHandle FPCGGraphExecutionRegistry::RegisterExecutionSourceProvider(TSharedRef<IPCGGraphExecutionSourceProvider> Provider)
+{
+	check(Provider->Handle == INDEX_NONE);
+	const FPCGGraphExecutionSourceProviderHandle ProviderHandle = NextHandle++;
+	Provider->Handle = ProviderHandle;
+	RegisteredExecutionSourceProviders.Add(MoveTemp(Provider));
+	return ProviderHandle;
+}
+
+void FPCGGraphExecutionRegistry::UnregisterExecutionSourceProvider(FPCGGraphExecutionSourceProviderHandle Handle)
+{
+	const int32 ProviderIndex = RegisteredExecutionSourceProviders.IndexOfByPredicate([Handle](const TSharedRef<IPCGGraphExecutionSourceProvider>& Provider) { return Provider->Handle == Handle; });
+	if (ProviderIndex != INDEX_NONE)
+	{
+		RegisteredExecutionSourceProviders[ProviderIndex]->Handle = INDEX_NONE;
+		RegisteredExecutionSourceProviders.RemoveAtSwap(ProviderIndex);
+	}
+}
\ No newline at end of file
diff --git a/Engine/Plugins/PCG/Source/PCG/Private/PCGModule.cpp b/Engine/Plugins/PCG/Source/PCG/Private/PCGModule.cpp
--- a/Engine/Plugins/PCG/Source/PCG/Private/PCGModule.cpp
+++ b/Engine/Plugins/PCG/Source/PCG/Private/PCGModule.cpp
@@ -2,7 +2,9 @@
 
 #include "PCGModule.h"
 
+#include "PCGComponent.h"
 #include "PCGContext.h"
+#include "PCGDefaultExecutionSource.h"
 #include "PCGElement.h"
 #include "PCGGraph.h"
 #include "Compute/PCGComputeKernel.h"
@@ -12,6 +14,7 @@
 #include "Data/PCGSplineData.h"
 #include "Data/PCGPolygon2DData.h"
 
+#include "Engine/Engine.h"
 #include "Interfaces/IPluginManager.h"
 #include "Misc/CoreDelegates.h"
 #include "Modules/ModuleManager.h"
@@ -43,6 +46,79 @@ namespace PCGModule
 
 	static const TCHAR* DataTypeRegistryError = TEXT("Data type operations like compatibility or intersection can't run before the module is loaded."
 		"The logic will not function in a CDO and must be deferred until after the load.");
+
+	class FPCGGraphExecutionSourceProvider : public IPCGGraphExecutionSourceProvider
+	{
+	public:
+		FPCGGraphExecutionSourceProvider()
+		{
+#if WITH_EDITOR
+			if (GEngine)
+			{
+				GEngine->OnLevelActorDeleted().AddRaw(this, &FPCGGraphExecutionSourceProvider::OnLevelActorDeleted);
+			}
+#endif
+		}
+
+		~FPCGGraphExecutionSourceProvider()
+		{
+#if WITH_EDITOR
+			if (GEngine)
+			{
+				GEngine->OnLevelActorDeleted().RemoveAll(this);
+			}
+#endif
+		}
+
+		virtual TArray<FPCGGraphExecutionSourceDescriptor> GatherExecutionSources() const override
+		{
+			TArray<FPCGGraphExecutionSourceDescriptor> ExecutionSources;
+
+			TArray<UObject*> PCGSources;
+			GetObjectsOfClass(UPCGComponent::StaticClass(), PCGSources, /*bIncludeDerivedClasses=*/ true);
+
+			ExecutionSources.Reserve(PCGSources.Num());
+
+			for (UObject* PCGSource : PCGSources)
+			{
+				UPCGComponent* PCGComponent = CastChecked<UPCGComponent>(PCGSource, ECastCheckedType::NullAllowed);
+				if (IsValid(PCGComponent) && PCGComponent->IsRegistered() && PCGComponent->GetGraph())
+				{
+					FPCGGraphExecutionSourceDescriptor Descriptor = { /*ExecutionSource=*/ PCGComponent };
+#if WITH_EDITOR
+					Descriptor.bShowInDebugger = true;
+#endif
+					ExecutionSources.Add(MoveTemp(Descriptor));
+				}
+			}
+
+			PCGSources.Reset();
+			GetObjectsOfClass(UPCGDefaultExecutionSource::StaticClass(), PCGSources, /*bIncludeDerivedClasses=*/ true);
+
+			for (UObject* PCGSource : PCGSources)
+			{
+				UPCGDefaultExecutionSource* PCGDefaultExecutionSource = CastChecked<UPCGDefaultExecutionSource>(PCGSource, ECastCheckedType::NullAllowed);
+				if (IsValid(PCGDefaultExecutionSource) && PCGDefaultExecutionSource->GetExecutionState().GetGraph())
+				{
+					FPCGGraphExecutionSourceDescriptor Descriptor = { /*ExecutionSource=*/ PCGDefaultExecutionSource };
+#if WITH_EDITOR
+					Descriptor.bShowInDebugger = true;
+#endif
+					ExecutionSources.Add(MoveTemp(Descriptor));
+				}
+			}
+
+			return ExecutionSources;
+		}
+
+	private:
+#if WITH_EDITOR
+		void OnLevelActorDeleted(AActor* InActor)
+		{
+			BroadcastExecutionSourcesChanged();
+		}
+#endif
+	};
 }
 
 FPCGModule& FPCGModule::GetPCGModuleChecked()
@@ -51,6 +127,11 @@ FPCGModule& FPCGModule::GetPCGModuleChecked()
 	return *PCGModule::PCGModulePtr;
 }
 
+FPCGModule* FPCGModule::GetPCGModule()
+{
+	return PCGModule::PCGModulePtr;
+}
+
 const FPCGDataTypeRegistry& FPCGModule::GetConstDataTypeRegistry()
 {
 	checkf(IsPCGModuleLoaded(), TEXT("%s"), PCGModule::DataTypeRegistryError);
@@ -131,6 +212,8 @@ void FPCGModule::PreExit()
 
 	DataTypeRegistry.Shutdown();
 
+	GraphExecutionRegistry.Shutdown();
+
 #if WITH_EDITOR
 	DeregisterNativeElementDeterminismTests();
 
@@ -142,6 +225,8 @@ void FPCGModule::OnPostInitEngine()
 {
 	LLM_SCOPE_BYTAG(PCG);
 	DataTypeRegistry.RegisterKnownTypes();
+
+	GraphExecutionRegistry.RegisterExecutionSourceProvider(MakeShared<PCGModule::FPCGGraphExecutionSourceProvider>());
 }
 
 void FPCGModule::ExecuteNextTick(TFunction<void()> TickFunction)
diff --git a/Engine/Plugins/PCG/Source/PCG/Private/Subsystems/IPCGBaseSubsystem.cpp b/Engine/Plugins/PCG/Source/PCG/Private/Subsystems/IPCGBaseSubsystem.cpp
--- a/Engine/Plugins/PCG/Source/PCG/Private/Subsystems/IPCGBaseSubsystem.cpp
+++ b/Engine/Plugins/PCG/Source/PCG/Private/Subsystems/IPCGBaseSubsystem.cpp
@@ -250,12 +250,9 @@ bool IPCGBaseSubsystem::GetStackContext(UPCGGraph* InGraph, uint32 InGridSize, b
 
 bool IPCGBaseSubsystem::GetStackContext(const IPCGGraphExecutionSource* InSource, FPCGStackContext& OutStackContext)
 {
-	// @todo_pcg: Extend execution state.
-	const UPCGComponent* Component = Cast<UPCGComponent>(InSource);
-	const bool bIsPartitioned = Component && Component->IsPartitioned();
-	const uint32 GridSize = (!Component || bIsPartitioned) ?  PCGHiGenGrid::UnboundedGridSize() : Component->GetGenerationGridSize();
+	const uint32 GridSize = InSource && InSource->GetExecutionState().IsLocalSource() ? InSource->GetExecutionState().GetGenerationGridSize() : PCGHiGenGrid::UnboundedGridSize();
 	
-	return InSource && GetStackContext(InSource->GetExecutionState().GetGraph(),GridSize, bIsPartitioned, OutStackContext);
+	return InSource && GetStackContext(InSource->GetExecutionState().GetGraph(),GridSize, InSource->GetExecutionState().IsPartitioned(), OutStackContext);
 }
 
 uint32 IPCGBaseSubsystem::GetGraphCacheEntryCount(IPCGElement* InElement) const
diff --git a/Engine/Plugins/PCG/Source/PCG/Public/PCGGraphExecutionRegistry.h b/Engine/Plugins/PCG/Source/PCG/Public/PCGGraphExecutionRegistry.h
new file mode 100644
--- /dev/null
+++ b/Engine/Plugins/PCG/Source/PCG/Public/PCGGraphExecutionRegistry.h
@@ -0,0 +1,62 @@
+// Copyright Epic Games, Inc. All Rights Reserved.
+
+#pragma once
+
+#include "PCGCommon.h"
+
+DECLARE_MULTICAST_DELEGATE(FPCGGraphExecutionSourcesChanged);
+
+class FPCGGraphExecutionRegistry;
+
+/**
+ * Execution source descriptor which points to an execution source and some extra parameters, returned by the provider class.
+ */
+struct FPCGGraphExecutionSourceDescriptor
+{
+	IPCGGraphExecutionSource* ExecutionSource = nullptr;
+#if WITH_EDITOR
+	bool bShowInDebugger = true;
+#endif
+};
+
+using FPCGGraphExecutionSourceProviderHandle = int32;
+
+/** 
+ * Registered providers are responsible for providing a list of execution sources.
+ * Instead of having a catch all getting classes implementing IPCGGraphExecutionSource this gives a bit more control to implementers.
+ */
+class IPCGGraphExecutionSourceProvider
+{
+public:	
+	virtual ~IPCGGraphExecutionSourceProvider() = default;
+
+	virtual TArray<FPCGGraphExecutionSourceDescriptor> GatherExecutionSources() const = 0;
+	
+	void BroadcastExecutionSourcesChanged();
+
+private:
+	friend class FPCGGraphExecutionRegistry;
+	FPCGGraphExecutionSourceProviderHandle Handle = INDEX_NONE;
+};
+
+/**
+ * Registry for graph execution source providers.
+ */
+class FPCGGraphExecutionRegistry
+{
+public:
+	PCG_API TArray<FPCGGraphExecutionSourceDescriptor> GatherExecutionSources() const;
+	void Shutdown();
+
+	PCG_API FPCGGraphExecutionSourceProviderHandle RegisterExecutionSourceProvider(TSharedRef<IPCGGraphExecutionSourceProvider> Provider);
+	PCG_API void UnregisterExecutionSourceProvider(FPCGGraphExecutionSourceProviderHandle Handle);
+
+	FPCGGraphExecutionSourcesChanged& OnGraphExecutionSourcesChanged() { return GraphExecutionSourcesChanged; }
+
+	void BroadcastGraphExecutionSourcesChanged() { GraphExecutionSourcesChanged.Broadcast(); }
+private:
+	
+	TArray<TSharedRef<IPCGGraphExecutionSourceProvider>> RegisteredExecutionSourceProviders;
+	FPCGGraphExecutionSourcesChanged GraphExecutionSourcesChanged;
+	static FPCGGraphExecutionSourceProviderHandle NextHandle;
+};
\ No newline at end of file
diff --git a/Engine/Plugins/PCG/Source/PCG/Public/PCGModule.h b/Engine/Plugins/PCG/Source/PCG/Public/PCGModule.h
--- a/Engine/Plugins/PCG/Source/PCG/Public/PCGModule.h
+++ b/Engine/Plugins/PCG/Source/PCG/Public/PCGModule.h
@@ -3,6 +3,7 @@
 #pragma once
 
 #include "PCGDataVisualizationRegistry.h"
+#include "PCGGraphExecutionRegistry.h"
 #include "Data/Registry/PCGDataTypeRegistry.h"
 #include "Data/Registry/PCGGetDataFunctionRegistry.h"
 #include "Hash/PCGObjectHash.h"
@@ -45,6 +46,7 @@ public:
 	void OnPostInitEngine();
 
 	PCG_API static FPCGModule& GetPCGModuleChecked();
+	PCG_API static FPCGModule* GetPCGModule();
 	static const FPCGGetDataFunctionRegistry& ConstGetDataFunctionRegistry() { return GetPCGModuleChecked().GetDataFunctionRegistry; }
 	static FPCGGetDataFunctionRegistry& MutableGetDataFunctionRegistry() { return GetPCGModuleChecked().GetDataFunctionRegistry; }
 
@@ -54,6 +56,9 @@ public:
 	static const FPCGAttributeAccessorFactory& GetConstAttributeAccessorFactory() { return GetPCGModuleChecked().AttributeAccessorFactory; }
 	static FPCGAttributeAccessorFactory& GetMutableAttributeAccessorFactory() { return GetPCGModuleChecked().AttributeAccessorFactory; }
 	
+	static const FPCGGraphExecutionRegistry& GetConstGraphExecutionRegistry() { return GetPCGModuleChecked().GraphExecutionRegistry; }
+	static FPCGGraphExecutionRegistry& GetGraphExecutionRegistry() { return GetPCGModuleChecked().GraphExecutionRegistry; }
+
 #if WITH_EDITOR
 	static const FPCGObjectHashFactory& GetConstObjectHashFactory() { return GetPCGModuleChecked().ObjectHashFactory; }
 	static FPCGObjectHashFactory& GetMutableObjectHashFactory() { return GetPCGModuleChecked().ObjectHashFactory; }
@@ -78,6 +83,7 @@ private:
 	FPCGGetDataFunctionRegistry GetDataFunctionRegistry;
 	FPCGAttributeAccessorFactory AttributeAccessorFactory;
 	FPCGDataTypeRegistry DataTypeRegistry;
+	FPCGGraphExecutionRegistry GraphExecutionRegistry;
 
 #if WITH_EDITOR
 	FPCGObjectHashFactory ObjectHashFactory;
diff --git a/Engine/Plugins/PCG/Source/PCGEditor/Private/PCGEditor.cpp b/Engine/Plugins/PCG/Source/PCGEditor/Private/PCGEditor.cpp
--- a/Engine/Plugins/PCG/Source/PCGEditor/Private/PCGEditor.cpp
+++ b/Engine/Plugins/PCG/Source/PCGEditor/Private/PCGEditor.cpp
@@ -167,6 +167,8 @@ void FPCGEditor::Initialize(const EToolkitMode::Type InMode, const TSharedPtr<cl
 		PCGEditorModule->OnPermissionModeChanged().AddRaw(this, &FPCGEditor::OnPermissionChanged);
 	}
 
+	FPCGModule::GetGraphExecutionRegistry().OnGraphExecutionSourcesChanged().AddRaw(this, &FPCGEditor::OnGraphExecutionSourcesChanged);
+
 	PCGGraphBeingEdited = InPCGGraph;
 
 	UpdateDefaultExecutionSource();
@@ -230,10 +232,6 @@ void FPCGEditor::Initialize(const EToolkitMode::Type InMode, const TSharedPtr<cl
 	// Hook to map change / delete actor to refresh debug object selection list, to help prevent it going stale.
 	FLevelEditorModule& LevelEditor = FModuleManager::LoadModuleChecked<FLevelEditorModule>("LevelEditor");
 	LevelEditor.OnMapChanged().AddRaw(this, &FPCGEditor::OnMapChanged);
-	if (GEngine)
-	{
-		GEngine->OnLevelActorDeleted().AddRaw(this, &FPCGEditor::OnLevelActorDeleted);
-	}
 
 	// Hook to PIE start/end to keep callbacks up to date.
 	FEditorDelegates::PostPIEStarted.AddRaw(this, &FPCGEditor::OnPostPIEStarted);
@@ -2925,11 +2923,6 @@ void FPCGEditor::OnClose()
 		LevelEditor->OnMapChanged().RemoveAll(this);
 	}
 
-	if (GEngine)
-	{
-		GEngine->OnLevelActorDeleted().RemoveAll(this);
-	}
-
 	FEditorDelegates::PostPIEStarted.RemoveAll(this);
 	FEditorDelegates::EndPIE.RemoveAll(this);
 
@@ -2968,6 +2961,8 @@ void FPCGEditor::OnClose()
 	{
 		PCGEditorModule->OnPermissionModeChanged().RemoveAll(this);
 	}
+
+	FPCGModule::GetGraphExecutionRegistry().OnGraphExecutionSourcesChanged().RemoveAll(this);
 }
 
 void FPCGEditor::InitToolMenuContext(FToolMenuContext& MenuContext)
@@ -3343,9 +3338,8 @@ void FPCGEditor::OnEndPIE(bool bIsSimulating)
 	UnregisterDelegatesForWorld(GEditor ? GEditor->PlayWorld.Get() : nullptr);
 }
 
-void FPCGEditor::OnLevelActorDeleted(AActor* InActor)
+void FPCGEditor::OnGraphExecutionSourcesChanged()
 {
-	// Forward call as this makes an effort to retain the selection if the selected component has not been deleted.
 	if (DebugObjectTreeWidget.IsValid())
 	{
 		DebugObjectTreeWidget->RequestRefresh();
diff --git a/Engine/Plugins/PCG/Source/PCGEditor/Private/Widgets/SPCGEditorGraphDebugObjectTree.cpp b/Engine/Plugins/PCG/Source/PCGEditor/Private/Widgets/SPCGEditorGraphDebugObjectTree.cpp
--- a/Engine/Plugins/PCG/Source/PCGEditor/Private/Widgets/SPCGEditorGraphDebugObjectTree.cpp
+++ b/Engine/Plugins/PCG/Source/PCGEditor/Private/Widgets/SPCGEditorGraphDebugObjectTree.cpp
@@ -5,6 +5,8 @@
 #include "PCGComponent.h"
 #include "PCGDefaultExecutionSource.h"
 #include "PCGGraph.h"
+#include "PCGGraphExecutionRegistry.h"
+#include "PCGModule.h"
 #include "PCGSubsystem.h"
 #include "Editor/IPCGEditorModule.h"
 #include "Grid/PCGPartitionActor.h"
@@ -1192,16 +1194,10 @@ void SPCGEditorGraphDebugObjectTree::RefreshTree()
 	}
 
 
-	// @todo_pcg: Support other executions sources? GetObjectsOfClass(UPCGGraphExecutionSource::StaticClass()) doesn't work.
-	TArray<UObject*> PCGSources;
-	GetObjectsOfClass(UPCGComponent::StaticClass(), PCGSources, /*bIncludeDerivedClasses=*/ true);
-
-	if (UObject* DefaultExecutionSource = PCGEditor.Pin()->GetDefaultExecutionSource())
-	{
-		PCGSources.Add(DefaultExecutionSource);
-	}
-
-	if (!PCGSources.IsEmpty())
+	FPCGGraphExecutionRegistry& GraphExecutionRegistry = FPCGModule::GetGraphExecutionRegistry();
+	TArray<FPCGGraphExecutionSourceDescriptor> ExecutionSourceDescriptors = GraphExecutionRegistry.GatherExecutionSources();
+	
+	if (!ExecutionSourceDescriptors.IsEmpty())
 	{
 		TMap<UObject*, TSharedPtr<FPCGEditorGraphDebugObjectItem>> OwnerItems;
 		TMap<const FPCGStack, FPCGEditorGraphDebugObjectItemPtr> StackToItem;
@@ -1210,21 +1206,15 @@ void SPCGEditorGraphDebugObjectTree::RefreshTree()
 		using FGraphContextMapValue = TTuple<FPCGStackContext, bool /*bContainsGraph*/>;
 		TMap<FGraphContextMapKey, FGraphContextMapValue> GraphToContexts;
 
-		for (UObject* PCGSourceObject : PCGSources)
+		for (const FPCGGraphExecutionSourceDescriptor& ExecutionSourceDescriptor : ExecutionSourceDescriptors)
 		{
-			if (!IsValid(PCGSourceObject))
+			if (!ExecutionSourceDescriptor.ExecutionSource || !ExecutionSourceDescriptor.bShowInDebugger)
 			{
 				continue;
 			}
 
-			IPCGGraphExecutionSource* PCGSource = Cast<IPCGGraphExecutionSource>(PCGSourceObject);
-			UPCGComponent* PCGComponent = Cast<UPCGComponent>(PCGSource);
+			IPCGGraphExecutionSource* PCGSource = ExecutionSourceDescriptor.ExecutionSource;
 			
-			if (!PCGSource || (PCGComponent && !PCGComponent->IsRegistered()) || !PCGSource->GetExecutionState().GetGraph())
-			{
-				continue;
-			}
-
 			const UWorld* ExecutionStateWorld = PCGSource->GetExecutionState().GetWorld();
 			// Ignore Inactive and Preview worlds
 			if (ExecutionStateWorld && (ExecutionStateWorld->WorldType == EWorldType::Inactive || ExecutionStateWorld->WorldType == EWorldType::EditorPreview))
@@ -1232,8 +1222,8 @@ void SPCGEditorGraphDebugObjectTree::RefreshTree()
 				continue;
 			}
 
-			const bool bIsPartitioned = PCGComponent && PCGComponent->IsPartitioned();
-			const uint32 GridSize = (bIsPartitioned || !PCGComponent) ? PCGHiGenGrid::UnboundedGridSize() : PCGComponent->GetGenerationGridSize();
+			const bool bIsPartitioned = PCGSource->GetExecutionState().IsPartitioned();
+			const uint32 GridSize = PCGSource->GetExecutionState().IsLocalSource() ? PCGSource->GetExecutionState().GetGenerationGridSize() : PCGHiGenGrid::UnboundedGridSize();
 
 			// Process static stacks that can be read from the compiled graph.
 			{
@@ -1259,7 +1249,7 @@ void SPCGEditorGraphDebugObjectTree::RefreshTree()
 				if (MapValue->Get<1>())
 				{
 					FPCGStack SourceStack;
-					SourceStack.PushFrame(PCGSourceObject);
+					SourceStack.PushFrame(Cast<UObject>(PCGSource));
 					FPCGStackContext SourceContext(MapValue->Get<0>(), SourceStack);
 					AddStacksToTree(SourceContext.GetStacks(), OwnerItems, StackToItem);
 				}
diff --git a/Engine/Plugins/PCG/Source/PCGEditor/Public/PCGEditor.h b/Engine/Plugins/PCG/Source/PCGEditor/Public/PCGEditor.h
--- a/Engine/Plugins/PCG/Source/PCGEditor/Public/PCGEditor.h
+++ b/Engine/Plugins/PCG/Source/PCGEditor/Public/PCGEditor.h
@@ -482,7 +482,7 @@ private:
 	void OnMapChanged(UWorld* InWorld, EMapChangeType InMapChangedType);
 	void OnPostPIEStarted(bool bIsSimulating);
 	void OnEndPIE(bool bIsSimulating);
-	void OnLevelActorDeleted(AActor* InActor);
+	void OnGraphExecutionSourcesChanged();
 
 	void RefreshViews();
 
