[PCG] ExecutionSource API : Add missing exports

#rnx

diff --git a/Engine/Plugins/PCG/Source/PCG/Public/PCGGraphExecutionStateInterface.h b/Engine/Plugins/PCG/Source/PCG/Public/PCGGraphExecutionStateInterface.h
--- a/Engine/Plugins/PCG/Source/PCG/Public/PCGGraphExecutionStateInterface.h
+++ b/Engine/Plugins/PCG/Source/PCG/Public/PCGGraphExecutionStateInterface.h
@@ -127,7 +127,7 @@ public:
 	/** Lookup data using a resource key that identifies the pin. */
 	virtual const FPCGDataCollection* RetrieveOutputDataForPin(const FString& InResourceKey) const { return nullptr; }
 
-	virtual void AddToManagedResources(UPCGManagedResource* InResource);
+	PCG_API virtual void AddToManagedResources(UPCGManagedResource* InResource);
 
 	/** Returns the subsystem associated with this source. By default, it will return the UPCGSubsystem if the source has a world, otherwise, the EditorSubsystem (in editor only, obviously). */
 	PCG_API virtual IPCGBaseSubsystem* GetSubsystem() const;
@@ -181,7 +181,7 @@ public:
 	virtual bool Use2DGrid() const { return false; }
 
 	/** Returns a GridDescriptor based on this execution source for the specified grid size. */
-	virtual FPCGGridDescriptor GetGridDescriptor(uint32 InGridSize) const;
+	PCG_API virtual FPCGGridDescriptor GetGridDescriptor(uint32 InGridSize) const;
 
 	/** Get the HiGen grid this execution source executes on. */
 	virtual EPCGHiGenGrid GetGenerationGrid() const { return EPCGHiGenGrid::Uninitialized; }
