[PCG] Add a safer way to override structures with the override system

#rnx

diff --git a/Engine/Plugins/PCG/Source/PCG/Private/PCGContext.cpp b/Engine/Plugins/PCG/Source/PCG/Private/PCGContext.cpp
--- a/Engine/Plugins/PCG/Source/PCG/Private/PCGContext.cpp
+++ b/Engine/Plugins/PCG/Source/PCG/Private/PCGContext.cpp
@@ -339,14 +339,25 @@ void FPCGContext::OverrideSettings()
 		if (!SettingsWithOverride->GetClass()->IsChildOf(Param.PropertyClass))
 		{
 			UObject* ObjectPtr = GetExternalContainerForOverridableParam(Param);
+
 			if (ObjectPtr && ObjectPtr->IsA(Param.Properties[0]->GetOwnerClass()))
 			{
 				Container = ObjectPtr;
 			}
 			else if (!ObjectPtr)
 			{
-				// It's the responsability of the callee to make sure we have a valid memory space to read from.
-				Container = GetUnsafeExternalContainerForOverridableParam(Param);
+				FStructView StructView = GetExternalStructContainerForOverridableParam(Param);
+				if (StructView.IsValid() && StructView.GetScriptStruct() == Param.Properties[0]->GetOwnerStruct())
+				{
+					Container = StructView.GetMemory();
+				}
+				else if (!StructView.IsValid())
+				{
+PRAGMA_DISABLE_DEPRECATION_WARNINGS
+					// It's the responsability of the callee to make sure we have a valid memory space to read from.
+					Container = GetUnsafeExternalContainerForOverridableParam(Param);
+PRAGMA_ENABLE_DEPRECATION_WARNINGS
+				}
 			}
 		}
 		else
diff --git a/Engine/Plugins/PCG/Source/PCG/Private/PCGSubgraph.cpp b/Engine/Plugins/PCG/Source/PCG/Private/PCGSubgraph.cpp
--- a/Engine/Plugins/PCG/Source/PCG/Private/PCGSubgraph.cpp
+++ b/Engine/Plugins/PCG/Source/PCG/Private/PCGSubgraph.cpp
@@ -518,7 +518,7 @@ TObjectPtr<UPCGGraphInterface> UPCGSubgraphNode::GetSubgraphInterface() const
 	return Settings ? Settings->GetSubgraphInterface() : nullptr;
 }
 
-void* FPCGSubgraphContext::GetUnsafeExternalContainerForOverridableParam(const FPCGSettingsOverridableParam& InParam)
+FStructView FPCGSubgraphContext::GetExternalStructContainerForOverridableParam(const FPCGSettingsOverridableParam& InParam)
 {
 	const UPCGSubgraphSettings* Settings = GetInputSettings<UPCGSubgraphSettings>();
 	check(Settings);
@@ -527,7 +527,7 @@ void* FPCGSubgraphContext::GetUnsafeExternalContainerForOverridableParam(const F
 
 	if (UserParameters && !InParam.PropertiesNames.IsEmpty() && UserParameters->FindPropertyDescByName(InParam.PropertiesNames[0]) && GraphInstanceParametersOverride.IsValid())
 	{
-		return GraphInstanceParametersOverride.GetMutableMemory();
+		return GraphInstanceParametersOverride;
 	}
 	else
 	{
diff --git a/Engine/Plugins/PCG/Source/PCG/Public/PCGContext.h b/Engine/Plugins/PCG/Source/PCG/Public/PCGContext.h
--- a/Engine/Plugins/PCG/Source/PCG/Public/PCGContext.h
+++ b/Engine/Plugins/PCG/Source/PCG/Public/PCGContext.h
@@ -9,6 +9,7 @@
 #include "Templates/SubclassOf.h"
 #include "Utils/PCGExtraCapture.h"
 
+#include "StructUtils/StructView.h"
 #include "UObject/GCObject.h"
 #include "UObject/GarbageCollection.h"
 
@@ -105,9 +106,9 @@ PRAGMA_DISABLE_DEPRECATION_WARNINGS
 	FPCGContext(FPCGContext&&) = default;
 	FPCGContext& operator=(const FPCGContext&) = default;
 	FPCGContext& operator=(FPCGContext&&) = default;
-UE_API PRAGMA_ENABLE_DEPRECATION_WARNINGS
+PRAGMA_ENABLE_DEPRECATION_WARNINGS
 
-	virtual ~FPCGContext();
+	UE_API virtual ~FPCGContext();
 
 	UE_API void InitFromParams(const FPCGInitializeElementParams& InParams);
 
@@ -280,6 +281,9 @@ UE_API PRAGMA_ENABLE_DEPRECATION_WARNINGS
 
 protected:
 	virtual UObject* GetExternalContainerForOverridableParam(const FPCGSettingsOverridableParam& InParam) { return nullptr; }
+	virtual FStructView GetExternalStructContainerForOverridableParam(const FPCGSettingsOverridableParam& InParam) { return {}; }
+
+	UE_DEPRECATED(5.8, "You should use GetExternalStructContainerForOverridableParam for structs overrides.")
 	virtual void* GetUnsafeExternalContainerForOverridableParam(const FPCGSettingsOverridableParam& InParam) { return nullptr; }
 	
 	UE_API void InitializeGraphExecutor(FPCGContext* InContext);
diff --git a/Engine/Plugins/PCG/Source/PCG/Public/PCGSubgraph.h b/Engine/Plugins/PCG/Source/PCG/Public/PCGSubgraph.h
--- a/Engine/Plugins/PCG/Source/PCG/Public/PCGSubgraph.h
+++ b/Engine/Plugins/PCG/Source/PCG/Public/PCGSubgraph.h
@@ -192,7 +192,7 @@ struct FPCGSubgraphContext : public FPCGContext
 	void AddToReferencedObjects(const FPCGDataCollection& InDataCollection);
 
 protected:
-	virtual void* GetUnsafeExternalContainerForOverridableParam(const FPCGSettingsOverridableParam& InParam) override;
+	virtual FStructView GetExternalStructContainerForOverridableParam(const FPCGSettingsOverridableParam& InParam) override;
 	virtual void AddExtraStructReferencedObjects(FReferenceCollector& Collector);
 };
 
