[PCG] Avoid ticking RGS VT priming when there are no components that require it

Register original components to OnGraphParametersChanged to be notified when priming infos are changed, and update the priming info cache accordingly. Only tick priming when the cache has entries.

#rnx

diff --git a/Engine/Plugins/PCG/Source/PCG/Private/RuntimeGen/PCGRuntimeGenScheduler.cpp b/Engine/Plugins/PCG/Source/PCG/Private/RuntimeGen/PCGRuntimeGenScheduler.cpp
--- a/Engine/Plugins/PCG/Source/PCG/Private/RuntimeGen/PCGRuntimeGenScheduler.cpp
+++ b/Engine/Plugins/PCG/Source/PCG/Private/RuntimeGen/PCGRuntimeGenScheduler.cpp
@@ -280,7 +280,7 @@ void FPCGRuntimeGenScheduler::Tick(APCGWorldActor* InPCGWorldActor, double InEnd
 		}
 	}
 
-	if (!GenSources.IsEmpty() && PCGRuntimeGenSchedulerHelpers::CVarEnableVirtualTexturePriming.GetValueOnGameThread())
+	if (!CachedPrimingInfos.IsEmpty() && !GenSources.IsEmpty() && PCGRuntimeGenSchedulerHelpers::CVarEnableVirtualTexturePriming.GetValueOnGameThread())
 	{
 		// @todo_pcg: To support VT priming outside of RuntimeGen, this should probably move outside of the RGS tick, and be ticked directly by the PCGSubsystem.
 		// However, that would require also moving the GenSourceManager out of the RGS, 
@@ -1230,43 +1230,6 @@ void FPCGRuntimeGenScheduler::TickRequestVirtualTexturePriming(const TSet<IPCGGe
 
 	check(ActorAndComponentMapping);
 
-	TArray<const UPCGComponent*> OriginalComponents;
-
-	{
-		TRACE_CPUPROFILER_EVENT_SCOPE(GatherOriginalComponents);
-
-		TSet<UPCGComponent*> PartitionedComponents = ActorAndComponentMapping->GetAllRegisteredPartitionedComponents();
-		TSet<UPCGComponent*> NonPartitionedComponents = ActorAndComponentMapping->GetAllRegisteredNonPartitionedComponents();
-
-		OriginalComponents.Reserve(PartitionedComponents.Num() + NonPartitionedComponents.Num());
-
-		auto CanComponentRequestVirtualTexturePriming = [](const UPCGComponent* OriginalComponent)
-		{
-			return OriginalComponent && OriginalComponent->bActivated && OriginalComponent->IsManagedByRuntimeGenSystem() && OriginalComponent->GetGraphInstance();
-		};
-
-		for (UPCGComponent* OriginalComponent : PartitionedComponents)
-		{
-			if (CanComponentRequestVirtualTexturePriming(OriginalComponent))
-			{
-				OriginalComponents.Add(MoveTemp(OriginalComponent));
-			}
-		}
-
-		for (UPCGComponent* OriginalComponent : NonPartitionedComponents)
-		{
-			if (CanComponentRequestVirtualTexturePriming(OriginalComponent))
-			{
-				OriginalComponents.Add(MoveTemp(OriginalComponent));
-			}
-		}
-	}
-
-	if (OriginalComponents.IsEmpty())
-	{
-		return;
-	}
-
 	TMap<TSoftObjectPtr<URuntimeVirtualTexture>, TArray<URuntimeVirtualTextureComponent*>> VirtualTextureToComponents;
 
 	{
@@ -1297,53 +1260,32 @@ void FPCGRuntimeGenScheduler::TickRequestVirtualTexturePriming(const TSet<IPCGGe
 		return;
 	}
 
-	for (const UPCGComponent* OriginalComponent : OriginalComponents)
+	// Clear out any entries for components that no longer resolve.
+	for (auto It = CachedPrimingInfos.CreateIterator(); It; ++It)
 	{
-		TRACE_CPUPROFILER_EVENT_SCOPE(RequestVirtualTexturePrimingForComponent);
-		check(OriginalComponent);
-
-		// @todo_pcg: Instead of polling the PrimingInfos every tick, we could probably cache them and only update when the graph params change.
-		// Polling the PrimingInfos in this way is cheap for a single OC, but does not scale well if there are many OCs in the level.
-		TArray<FPCGVirtualTexturePrimingInfo*> PrimingInfos;
-
+		if (!It.Key().ResolveObjectPtr())
 		{
-			TRACE_CPUPROFILER_EVENT_SCOPE(PollVirtualTexturePrimingInfos);
-
-			const UPCGGraphInstance* GraphInstance = OriginalComponent->GetGraphInstance();
-			const FInstancedPropertyBag* UserParametersStruct = GraphInstance ? GraphInstance->GetUserParametersStruct() : nullptr;
-			const UPropertyBag* PropertyBag = UserParametersStruct ? UserParametersStruct->GetPropertyBagStruct() : nullptr;
-
-			if (!PropertyBag)
-			{
-				continue;
-			}
-
-			TConstArrayView<FPropertyBagPropertyDesc> PropertyDescs = PropertyBag->GetPropertyDescs();
+			It.RemoveCurrent();
+		}
+	}
 
-			for (const FPropertyBagPropertyDesc& PropertyDesc : PropertyDescs)
-			{
-				if (PropertyDesc.ValueType == EPropertyBagPropertyType::Struct && PropertyDesc.ValueTypeObject == TBaseStructure<FPCGVirtualTexturePrimingInfo>::Get())
-				{
-					TValueOrError<FPCGVirtualTexturePrimingInfo*, EPropertyBagResult> Property = UserParametersStruct->GetValueStruct<FPCGVirtualTexturePrimingInfo>(PropertyDesc);
+	for (const auto& [OriginalComponentKey, PrimingInfos] : CachedPrimingInfos)
+	{
+		TRACE_CPUPROFILER_EVENT_SCOPE(RequestVirtualTexturePrimingForComponent);
 
-					if (Property.HasValue() && Property.GetValue())
-					{
-						PrimingInfos.Add(Property.GetValue());
-					}
-				}
-			}
-		}
+		const UPCGComponent* OriginalComponent = OriginalComponentKey.ResolveObjectPtr();
+		check(OriginalComponent);
 
 		const FBox OriginalComponentBounds = OriginalComponent->GetGridBounds();
 
-		for (const FPCGVirtualTexturePrimingInfo* PrimingInfo : PrimingInfos)
+		for (const FPCGVirtualTexturePrimingInfo& PrimingInfo : PrimingInfos)
 		{
-			if (!PrimingInfo || !PrimingInfo->VirtualTexture || PrimingInfo->WorldTexelSize < PCGRuntimeGenSchedulerConstants::MinWorldVirtualTextureTexelSize)
+			if (!PrimingInfo.VirtualTexture || PrimingInfo.WorldTexelSize < PCGRuntimeGenSchedulerConstants::MinWorldVirtualTextureTexelSize)
 			{
 				continue;
 			}
 
-			TArray<URuntimeVirtualTextureComponent*>* VirtualTextureComponents = VirtualTextureToComponents.Find(PrimingInfo->VirtualTexture);
+			TArray<URuntimeVirtualTextureComponent*>* VirtualTextureComponents = VirtualTextureToComponents.Find(PrimingInfo.VirtualTexture);
 
 			if (!VirtualTextureComponents)
 			{
@@ -1354,7 +1296,7 @@ void FPCGRuntimeGenScheduler::TickRequestVirtualTexturePriming(const TSet<IPCGGe
 			{
 				check(VirtualTextureComponent);
 
-				const double PrimingRadius = OriginalComponent->GetGenerationRadiusFromGrid(PrimingInfo->Grid) + PCGHiGenGrid::GridToGridSize(PrimingInfo->Grid);
+				const double PrimingRadius = OriginalComponent->GetGenerationRadiusFromGrid(PrimingInfo.Grid) + PCGHiGenGrid::GridToGridSize(PrimingInfo.Grid);
 
 				for (const IPCGGenSourceBase* GenSource : InGenSources)
 				{
@@ -1379,9 +1321,9 @@ void FPCGRuntimeGenScheduler::TickRequestVirtualTexturePriming(const TSet<IPCGGe
 						TRACE_CPUPROFILER_EVENT_SCOPE(RequestVirtualTexturePreload);
 
 						const double BoundsMaxExtent = FMath::Max(VirtualTextureComponent->Bounds.BoxExtent.X, VirtualTextureComponent->Bounds.BoxExtent.Y);
-						const int32 VirtualTextureSizeTexels = FMath::Max(1, PrimingInfo->VirtualTexture->GetSize());
+						const int32 VirtualTextureSizeTexels = FMath::Max(1, PrimingInfo.VirtualTexture->GetSize());
 						const int32 SizeTexelsLog2 = FMath::FloorLog2(VirtualTextureSizeTexels);
-						const int32 RequestedNumTexels = FMath::Max(1, BoundsMaxExtent / PrimingInfo->WorldTexelSize);
+						const int32 RequestedNumTexels = FMath::Max(1, BoundsMaxExtent / PrimingInfo.WorldTexelSize);
 						const int32 RequestedTexelsLog2 = FMath::FloorLog2(RequestedNumTexels);
 						const int32 MipLevel = FMath::Max(0, SizeTexelsLog2 - RequestedTexelsLog2);
 
@@ -1470,6 +1412,12 @@ void FPCGRuntimeGenScheduler::OnOriginalComponentRegistered(UPCGComponent* InOri
 
 	// When an original/non-partitioned component is registered, we need to dirty the state.
 	bAnyRuntimeGenComponentsExistDirty = true;
+
+#if WITH_EDITOR
+	RegisterGraphParameterChangedEvent(InOriginalComponent);
+#endif
+
+	CacheVirtualTexturePrimingInfos(InOriginalComponent);
 }
 
 void FPCGRuntimeGenScheduler::OnOriginalComponentUnregistered(UPCGComponent* InOriginalComponent)
@@ -1484,6 +1432,12 @@ void FPCGRuntimeGenScheduler::OnOriginalComponentUnregistered(UPCGComponent* InO
 	// When an original/non-partitioned component is unregistered, we need to dirty the state.
 	bAnyRuntimeGenComponentsExistDirty = true;
 
+#if WITH_EDITOR
+	UnregisterGraphParameterChangedEvent(InOriginalComponent);
+#endif
+
+	CachedPrimingInfos.Remove(InOriginalComponent);
+
 	// Gather all generated components which originated from this original component.
 	TSet<FGridGenerationKey> KeysToCleanup;
 	for (FGridGenerationKey GenerationKey : GeneratedComponents)
@@ -1543,6 +1497,15 @@ void FPCGRuntimeGenScheduler::OnObjectsReplaced(const TMap<UObject*, UObject*>&
 					GenerationKey.OriginalComponent = NewComponent;
 				}
 			}
+
+			TArray<FPCGVirtualTexturePrimingInfo> PrimingInfos;
+			if (CachedPrimingInfos.RemoveAndCopyValue(OldComponent, PrimingInfos))
+			{
+				CachedPrimingInfos.Add(NewComponent, PrimingInfos);
+			}
+
+			UnregisterGraphParameterChangedEvent(OldComponent);
+			RegisterGraphParameterChangedEvent(NewComponent);
 		}
 	}
 }
@@ -2119,4 +2082,102 @@ void FPCGRuntimeGenScheduler::AddReferencedObjects(FReferenceCollector& Collecto
 	Collector.AddReferencedObjects(PartitionActorPool);
 }
 
+void FPCGRuntimeGenScheduler::CacheVirtualTexturePrimingInfos(UPCGComponent* InOriginalComponent)
+{
+	CachedPrimingInfos.Remove(InOriginalComponent);
+
+	const UPCGGraphInstance* Graph = InOriginalComponent ? InOriginalComponent->GetGraphInstance() : nullptr;
+	const FInstancedPropertyBag* UserParametersStruct = Graph ? Graph->GetUserParametersStruct() : nullptr;
+	const UPropertyBag* PropertyBag = UserParametersStruct ? UserParametersStruct->GetPropertyBagStruct() : nullptr;
+
+	if (!PropertyBag)
+	{
+		return;
+	}
+
+	TArray<FPCGVirtualTexturePrimingInfo> PrimingInfos;
+
+	for (const FPropertyBagPropertyDesc& PropertyDesc : PropertyBag->GetPropertyDescs())
+	{
+		if (PropertyDesc.ValueType == EPropertyBagPropertyType::Struct && PropertyDesc.ValueTypeObject == TBaseStructure<FPCGVirtualTexturePrimingInfo>::Get())
+		{
+			TValueOrError<FPCGVirtualTexturePrimingInfo*, EPropertyBagResult> Property = UserParametersStruct->GetValueStruct<FPCGVirtualTexturePrimingInfo>(PropertyDesc);
+
+			if (Property.HasValue() && Property.GetValue() && !Property.GetValue()->VirtualTexture.IsNull())
+			{
+				PrimingInfos.Add(*Property.GetValue());
+			}
+		}
+	}
+
+	if (!PrimingInfos.IsEmpty())
+	{
+		CachedPrimingInfos.Add(InOriginalComponent, MoveTemp(PrimingInfos));
+	}
+}
+
+#if WITH_EDITOR
+void FPCGRuntimeGenScheduler::RegisterGraphParameterChangedEvent(UPCGComponent* InOriginalComponent)
+{
+	if (UPCGGraph* Graph = InOriginalComponent ? InOriginalComponent->GetGraph() : nullptr)
+	{
+		FDelegateHandle* FoundDelegateHandle = GraphParamChangedEventHandles.Find(InOriginalComponent);
+
+		if (!FoundDelegateHandle)
+		{
+			GraphParamChangedEventHandles.Add(InOriginalComponent, Graph->OnGraphParametersChangedDelegate.AddRaw(this, &FPCGRuntimeGenScheduler::OnGraphParametersChanged));
+		}
+	}
+}
+
+void FPCGRuntimeGenScheduler::UnregisterGraphParameterChangedEvent(UPCGComponent* InOriginalComponent)
+{
+	FDelegateHandle DelegateHandle;
+
+	if (GraphParamChangedEventHandles.RemoveAndCopyValue(InOriginalComponent, DelegateHandle))
+	{
+		if (UPCGGraph* Graph = InOriginalComponent ? InOriginalComponent->GetGraph() : nullptr)
+		{
+			Graph->OnGraphParametersChangedDelegate.Remove(DelegateHandle);
+		}
+	}
+}
+
+void FPCGRuntimeGenScheduler::OnGraphParametersChanged(UPCGGraphInterface* InGraph, EPCGGraphParameterEvent InChangeType, FName InChangedPropertyName)
+{
+	TArray<UPCGComponent*> OriginalComponents;
+
+	TSet<UPCGComponent*> PartitionedComponents = ActorAndComponentMapping->GetAllRegisteredPartitionedComponents();
+	TSet<UPCGComponent*> NonPartitionedComponents = ActorAndComponentMapping->GetAllRegisteredNonPartitionedComponents();
+
+	OriginalComponents.Reserve(PartitionedComponents.Num() + NonPartitionedComponents.Num());
+
+	auto ShouldCachePrimingInfosForComponent = [InGraph](const UPCGComponent* OriginalComponent)
+	{
+		return OriginalComponent && OriginalComponent->bActivated && OriginalComponent->IsManagedByRuntimeGenSystem() && OriginalComponent->GetGraph() == InGraph;
+	};
+
+	for (UPCGComponent* OriginalComponent : PartitionedComponents)
+	{
+		if (ShouldCachePrimingInfosForComponent(OriginalComponent))
+		{
+			OriginalComponents.Add(MoveTemp(OriginalComponent));
+		}
+	}
+
+	for (UPCGComponent* OriginalComponent : NonPartitionedComponents)
+	{
+		if (ShouldCachePrimingInfosForComponent(OriginalComponent))
+		{
+			OriginalComponents.Add(MoveTemp(OriginalComponent));
+		}
+	}
+
+	for (UPCGComponent* Component : OriginalComponents)
+	{
+		CacheVirtualTexturePrimingInfos(Component);
+	}
+}
+#endif // WITH_EDITOR
+
 #undef PCG_RGS_ONSCREENDEBUGMESSAGES
diff --git a/Engine/Plugins/PCG/Source/PCG/Public/RuntimeGen/PCGRuntimeGenScheduler.h b/Engine/Plugins/PCG/Source/PCG/Public/RuntimeGen/PCGRuntimeGenScheduler.h
--- a/Engine/Plugins/PCG/Source/PCG/Public/RuntimeGen/PCGRuntimeGenScheduler.h
+++ b/Engine/Plugins/PCG/Source/PCG/Public/RuntimeGen/PCGRuntimeGenScheduler.h
@@ -22,6 +22,7 @@ class UPCGSubsystem;
 class UWorld;
 enum class ELevelStreamingState : uint8;
 enum class ELevelStreamingTargetState : uint8;
+enum class EPCGGraphParameterEvent;
 namespace PCGRuntimeGenSchedulerHelpers { struct FStatsOverlay; }
 
 namespace PCGRuntimeGenSchedulerHelpers
@@ -201,6 +202,15 @@ private:
 	/** Cleanup the remaining components that are not part of the GeneratedComponents */
 	void CleanupRemainingComponents(UPCGComponent* InOriginalComponent);
 
+	/** Cache VT priming infos found in graph params of the given component. */
+	void CacheVirtualTexturePrimingInfos(UPCGComponent* InOriginalComponent);
+
+#if WITH_EDITOR
+	void RegisterGraphParameterChangedEvent(UPCGComponent* InOriginalComponent);
+	void UnregisterGraphParameterChangedEvent(UPCGComponent* InOriginalComponent);
+	void OnGraphParametersChanged(UPCGGraphInterface* InGraph, EPCGGraphParameterEvent InChangeType, FName InChangedPropertyName);
+#endif
+
 	/** Tracks the generated components managed by the RuntimeGenScheduler. For local components, this generation key will hold the original component.
 	* For non-partitioned components, the generation key should have unbounded grid size and (0, 0, 0) grid coordinates.
 	*/
@@ -261,5 +271,12 @@ private:
 	/** Used to detect when world state (generation sources, cvars etc) have changed enough to warrant rescanning generation cells. */
 	PCGRuntimeGenChangeDetection::FDetector ChangeDetector;
 
+#if WITH_EDITOR
+	TMap<TObjectKey<UPCGComponent>, FDelegateHandle> GraphParamChangedEventHandles;
+#endif
+
+	/** Maps original component to a cache of VT priming infos. */
+	TMap<TObjectKey<UPCGComponent>, TArray<FPCGVirtualTexturePrimingInfo>> CachedPrimingInfos;
+
 	friend struct PCGRuntimeGenSchedulerHelpers::FStatsOverlay;
 };
