[PCG] Refactor of Instance Data packer logic to be able to re-use as much as possible code for the Custom Primitive Data

* Also a first step towards removing the use of UClasses for Instance Data Packer
* Attributes now use accessors. It's anyway only used for Regex packing (which is not used often)
* Add unit tests for new code, there are also already some unit tests for Static Mesh Spawner that test the attributes path.
* Add utility function to create accessor keys from attributes.

#rnx

diff --git a/Engine/Plugins/PCG/Source/PCG/Private/InstanceDataPackers/PCGInstanceDataPackerBase.cpp b/Engine/Plugins/PCG/Source/PCG/Private/InstanceDataPackers/PCGInstanceDataPackerBase.cpp
--- a/Engine/Plugins/PCG/Source/PCG/Private/InstanceDataPackers/PCGInstanceDataPackerBase.cpp
+++ b/Engine/Plugins/PCG/Source/PCG/Private/InstanceDataPackers/PCGInstanceDataPackerBase.cpp
@@ -2,196 +2,93 @@
 
 #include "InstanceDataPackers/PCGInstanceDataPackerBase.h"
 
+#include "PCGContext.h"
 #include "PCGElement.h"
 #include "Elements/Metadata/PCGMetadataElementCommon.h"
 #include "MeshSelectors/PCGMeshSelectorBase.h"
 #include "Metadata/PCGMetadata.h"
 #include "Metadata/Accessors/IPCGAttributeAccessor.h"
+#include "Metadata/Accessors/PCGAttributeAccessorHelpers.h"
 #include "Metadata/Accessors/PCGAttributeAccessorKeys.h"
 
 #include UE_INLINE_GENERATED_CPP_BY_NAME(PCGInstanceDataPackerBase)
 
-void UPCGInstanceDataPackerBase::PackInstances_Implementation(FPCGContext& Context, const UPCGSpatialData* InSpatialData, const FPCGMeshInstanceList& InstanceList, FPCGPackedCustomData& OutPackedCustomData) const
-{
-	PCGE_LOG_C(Error, GraphAndLog, &Context, NSLOCTEXT("PCGInstanceDataPackerBase", "InstanceDataPackerBaseFailed", "Unable to execute InstanceDataPacker pure virtual base function, override the PackInstances function or use a default implementation."));
-}
+#define LOCTEXT_NAMESPACE "PCGInstanceDataPackerBase"
 
-bool UPCGInstanceDataPackerBase::AddTypeToPacking(int TypeId, FPCGPackedCustomData& OutPackedCustomData) const
+bool PCGInstanceDataPackerBase::FPackedDataCommonParams::ValidateOffsets(int32 OffsetIndex, int32 NumOfCustomFloats) const
 {
-	switch (TypeId)
+	if (!OptionalOffsets.IsValidIndex(OffsetIndex))
 	{
-	case PCG::Private::MetadataTypes<bool>::Id: // fall-through
-	case PCG::Private::MetadataTypes<float>::Id: // fall-through
-	case PCG::Private::MetadataTypes<double>::Id: // fall-through
-	case PCG::Private::MetadataTypes<int32>::Id: // fall-through
-	case PCG::Private::MetadataTypes<int64>::Id:
-		OutPackedCustomData.NumCustomDataFloats += 1;
-		break;
-	case PCG::Private::MetadataTypes<FVector2D>::Id:
-		OutPackedCustomData.NumCustomDataFloats += 2;
-		break;
-	case PCG::Private::MetadataTypes<FRotator>::Id: // fall-through
-	case PCG::Private::MetadataTypes<FVector>::Id:
-		OutPackedCustomData.NumCustomDataFloats += 3;
-		break;
-	case PCG::Private::MetadataTypes<FVector4>::Id: // fall-through
-	case PCG::Private::MetadataTypes<FQuat>::Id:
-		OutPackedCustomData.NumCustomDataFloats += 4;
-		break;
-	default:
+		PCGLog::LogErrorOnGraph(FText::Format(LOCTEXT("InvalidOffsetIndex", "[PCGInstanceDataPackerBase::ValidateOffsets] Offsets are expected to be at least {0} elements."), OffsetIndex), OptionalContext);
 		return false;
 	}
-
-	return true;
+	else if (OptionalOffsets[OffsetIndex] < NumOfCustomFloats)
+	{
+		PCGLog::LogErrorOnGraph(FText::Format(LOCTEXT("InvalidOffset", "[PCGInstanceDataPackerBase::ValidateOffsets] Offset {0} at index {1} is too small for the number of expected floats ({2})"), OptionalOffsets[OffsetIndex], OffsetIndex, NumOfCustomFloats), OptionalContext);
+		return false;
+	}
+	else
+	{
+		return true;
+	}
 }
 
-void UPCGInstanceDataPackerBase::PackCustomDataFromAttributes(const FPCGMeshInstanceList& InstanceList, const UPCGMetadata* Metadata, const TArray<FName>& AttributeNames, FPCGPackedCustomData& OutPackedCustomData) const
+bool PCGInstanceDataPackerBase::FPackedDataCommonParams::ValidateNumOfFloats(int32 NumOfCustomFloats) const
 {
-	if (!Metadata)
+	check(OutPackedCustomData);
+	
+	if (OutPackedCustomData->NumCustomDataFloats != 0 && OutPackedCustomData->NumCustomDataFloats < NumOfCustomFloats)
 	{
-		UE_LOG(LogPCG, Error, TEXT("Invalid metadata"));
-		return;
+		PCGLog::LogErrorOnGraph(FText::Format(LOCTEXT("NotEnoughFloats", "[PCGInstanceDataPackerBase::ValidateNumOfFloats] Expected to have at least {0} custom floats, but was explicitly set to {1}, which is not enough."), NumOfCustomFloats, OutPackedCustomData->NumCustomDataFloats), OptionalContext);
+		return false;
 	}
-
-	TArray<const FPCGMetadataAttributeBase*> Attributes;
-	for (const FName& AttributeName : AttributeNames)
+	else
 	{
-		Attributes.Add(Metadata->GetConstAttribute(AttributeName));
+		return true;
 	}
-
-	PackCustomDataFromAttributes(InstanceList, Attributes, OutPackedCustomData);
 }
 
-void UPCGInstanceDataPackerBase::PackCustomDataFromAttributes(const FPCGMeshInstanceList& InstanceList, const TArray<const FPCGMetadataAttributeBase*>& Attributes, FPCGPackedCustomData& OutPackedCustomData) const
+void PCGInstanceDataPackerBase::FPackedDataCommonParams::ZeroOutCustomData()
 {
-	const UPCGBasePointData* PointData = InstanceList.PointData.Get();
-	const TConstPCGValueRange<int64> MetadataEntryRange = PointData ? PointData->GetConstMetadataEntryValueRange() : TConstPCGValueRange<int64>();
-
-	for (const int32 Index : InstanceList.InstancesIndices)
+	if (OutPackedCustomData)
 	{
-		const PCGMetadataEntryKey EntryKey = PointData ? MetadataEntryRange[Index] : Index;
-		for (const FPCGMetadataAttributeBase* AttributeBase : Attributes)
-		{
-			check(AttributeBase);
-
-			switch (AttributeBase->GetTypeId())
-			{
-			case PCG::Private::MetadataTypes<bool>::Id:
-			{
-				const FPCGMetadataAttribute<bool>* Attribute = static_cast<const FPCGMetadataAttribute<bool>*>(AttributeBase);
-				check(Attribute);
-
-				const bool Value = Attribute->GetValueFromItemKey(EntryKey);
-				OutPackedCustomData.CustomData.Add(Value);
-				break;
-			}
-			case PCG::Private::MetadataTypes<float>::Id:
-			{
-				const FPCGMetadataAttribute<float>* Attribute = static_cast<const FPCGMetadataAttribute<float>*>(AttributeBase);
-				check(Attribute);
+		OutPackedCustomData->CustomData.SetNumZeroed(OutPackedCustomData->NumCustomDataFloats * NumInstances);
+	}
+}
 
-				const float Value = Attribute->GetValueFromItemKey(EntryKey);
-				OutPackedCustomData.CustomData.Add(Value);
-				break;
-			}
-			case PCG::Private::MetadataTypes<double>::Id:
-			{
-				const FPCGMetadataAttribute<double>* Attribute = static_cast<const FPCGMetadataAttribute<double>*>(AttributeBase);
-				check(Attribute);
+void PCGInstanceDataPackerBase::PackCustomDataFromAccessors(FPackedDataFromAccessorParams& InParams)
+{
+	check(InParams.CommonParams.OutPackedCustomData);
+	FPCGPackedCustomData& OutPackedCustomData = *InParams.CommonParams.OutPackedCustomData;
 
-				const double Value = Attribute->GetValueFromItemKey(EntryKey);
-				OutPackedCustomData.CustomData.Add(Value);
-				break;
-			}
-			case PCG::Private::MetadataTypes<int32>::Id:
-			{
-				const FPCGMetadataAttribute<int32>* Attribute = static_cast<const FPCGMetadataAttribute<int32>*>(AttributeBase);
-				check(Attribute);
+	const int32 NumInstances = InParams.CommonParams.NumInstances;
+	const int32 TotalNumCustomDataFloats = OutPackedCustomData.NumCustomDataFloats * NumInstances;
 
-				const float Value = static_cast<float>(Attribute->GetValueFromItemKey(EntryKey));
-				OutPackedCustomData.CustomData.Add(Value);
-				break;
-			}
-			case PCG::Private::MetadataTypes<int64>::Id:
-			{
-				const FPCGMetadataAttribute<int64>* Attribute = static_cast<const FPCGMetadataAttribute<int64>*>(AttributeBase);
-				check(Attribute);
+	// Zero out the custom data to avoid sending garbage data to the GPU in case of error.
+	auto ZeroOut = [&InParams]()
+	{
+		InParams.CommonParams.ZeroOutCustomData();
+	};
 
-				const float Value = static_cast<float>(Attribute->GetValueFromItemKey(EntryKey));
-				OutPackedCustomData.CustomData.Add(Value);
-				break;
-			}
-			case PCG::Private::MetadataTypes<FRotator>::Id:
-			{
-				const FPCGMetadataAttribute<FRotator>* Attribute = static_cast<const FPCGMetadataAttribute<FRotator>*>(AttributeBase);
-				check(Attribute);
-
-				const FRotator Value = Attribute->GetValueFromItemKey(EntryKey);
-				OutPackedCustomData.CustomData.Add(Value.Roll);
-				OutPackedCustomData.CustomData.Add(Value.Pitch);
-				OutPackedCustomData.CustomData.Add(Value.Yaw);
-				break;
-			}
-			case PCG::Private::MetadataTypes<FVector2D>::Id:
-			{
-				const FPCGMetadataAttribute<FVector2D>* Attribute = static_cast<const FPCGMetadataAttribute<FVector2D>*>(AttributeBase);
-				check(Attribute);
+	if (!ensure(InParams.Accessors.Num() == InParams.AccessorKeys.Num() || (InParams.AccessorKeys.Num() == 1)))
+	{
+		ZeroOut();
+		return;
+	}
 
-				const FVector2D Value = Attribute->GetValueFromItemKey(EntryKey);
-				OutPackedCustomData.CustomData.Add(Value.X);
-				OutPackedCustomData.CustomData.Add(Value.Y);
-				break;
-			}
-			case PCG::Private::MetadataTypes<FVector>::Id:
-			{
-				const FPCGMetadataAttribute<FVector>* Attribute = static_cast<const FPCGMetadataAttribute<FVector>*>(AttributeBase);
-				check(Attribute);
-
-				const FVector Value = Attribute->GetValueFromItemKey(EntryKey);
-				OutPackedCustomData.CustomData.Add(Value.X);
-				OutPackedCustomData.CustomData.Add(Value.Y);
-				OutPackedCustomData.CustomData.Add(Value.Z);
-				break;
-			}
-			case PCG::Private::MetadataTypes<FVector4>::Id:
-			{
-				const FPCGMetadataAttribute<FVector4>* Attribute = static_cast<const FPCGMetadataAttribute<FVector4>*>(AttributeBase);
-				check(Attribute);
-
-				const FVector4 Value = Attribute->GetValueFromItemKey(EntryKey);
-				OutPackedCustomData.CustomData.Add(Value.X);
-				OutPackedCustomData.CustomData.Add(Value.Y);
-				OutPackedCustomData.CustomData.Add(Value.Z);
-				OutPackedCustomData.CustomData.Add(Value.W);
-				break;
-			}
-			case PCG::Private::MetadataTypes<FQuat>::Id:
-			{
-				const FPCGMetadataAttribute<FQuat>* Attribute = static_cast<const FPCGMetadataAttribute<FQuat>*>(AttributeBase);
-				check(Attribute);
-
-				const FQuat Value = Attribute->GetValueFromItemKey(EntryKey);
-				OutPackedCustomData.CustomData.Add(Value.X);
-				OutPackedCustomData.CustomData.Add(Value.Y);
-				OutPackedCustomData.CustomData.Add(Value.Z);
-				OutPackedCustomData.CustomData.Add(Value.W);
-				break;
-			}
-			default:
-				break;
-			}
-		}
+	if (!Algo::AllOf(InParams.Accessors) || !Algo::AllOf(InParams.AccessorKeys))
+	{
+		ZeroOut();
+		return;
 	}
-}
 
-void UPCGInstanceDataPackerBase::PackCustomDataFromAccessors(const FPCGMeshInstanceList& InstanceList, TArray<TUniquePtr<const IPCGAttributeAccessor>> Accessors, TArray<TUniquePtr<const IPCGAttributeAccessorKeys>> AccessorKeys, FPCGPackedCustomData& OutPackedCustomData) const
-{
-	const int TotalNumCustomDataFloats = OutPackedCustomData.NumCustomDataFloats * InstanceList.Instances.Num();
+	TConstArrayView<int32> OptionalOffsets = InParams.CommonParams.OptionalOffsets;
 
-	if (!ensure(Accessors.Num() == AccessorKeys.Num()))
+	// Validate that if we have custom offsets, they match the number of accessors
+	if (!ensure(OptionalOffsets.IsEmpty() || (InParams.Accessors.Num() == OptionalOffsets.Num())))
 	{
 		// Zero out the custom data to avoid sending garbage data to the GPU.
-		OutPackedCustomData.CustomData.SetNumZeroed(TotalNumCustomDataFloats);
+		ZeroOut();
 		return;
 	}
 
@@ -202,15 +99,36 @@ void UPCGInstanceDataPackerBase::PackCustomDataFromAccessors(const FPCGMeshInsta
 		return;
 	}
 
+	if (!OptionalOffsets.IsEmpty())
+	{
+		OutPackedCustomData.Mask.Init(false, OutPackedCustomData.NumCustomDataFloats);
+	}
+
 	// Index of the accessor's attribute/property for a single instance (e.g. if the
 	// accessors are [float, FVector, float], AccessorPackedIndex will be [0, 1, 4])
 	int AccessorPackedIndex = 0;
-	const EPCGAttributeAccessorFlags Flags = EPCGAttributeAccessorFlags::AllowBroadcast | EPCGAttributeAccessorFlags::AllowConstructible;
+	const EPCGAttributeAccessorFlags Flags = EPCGAttributeAccessorFlags::AllowBroadcastAndConstructible;
 
-	for (int AccessorIndex = 0; AccessorIndex < Accessors.Num(); ++AccessorIndex)
+	for (int AccessorIndex = 0; AccessorIndex < InParams.Accessors.Num(); ++AccessorIndex)
 	{
-		const IPCGAttributeAccessor* Accessor = Accessors[AccessorIndex].Get();
-		const IPCGAttributeAccessorKeys* Keys = AccessorKeys[AccessorIndex].Get();
+		const IPCGAttributeAccessor* Accessor = InParams.Accessors[AccessorIndex].Get();
+		const IPCGAttributeAccessorKeys* Keys = InParams.AccessorKeys.Num() == 1 ? InParams.AccessorKeys[0].Get() : InParams.AccessorKeys[AccessorIndex].Get();
+
+		if (!OptionalOffsets.IsEmpty())
+		{
+			// Make sure that the offset if not writing on existing data.
+			const int32 Offset = OptionalOffsets[AccessorIndex];
+			const int32 TypePackingSize = GetTypePackingSize(Accessor->GetUnderlyingType());
+			if (Offset < AccessorPackedIndex || (Offset + TypePackingSize) > OutPackedCustomData.NumCustomDataFloats)
+			{
+				PCGLog::LogErrorOnGraph(LOCTEXT("InvalidOffsets", "While packing custom data, offsets are overlapping or out of bounds."), InParams.CommonParams.OptionalContext);
+				ZeroOut();
+				return;
+			}
+			
+			AccessorPackedIndex = OptionalOffsets[AccessorIndex];
+			OutPackedCustomData.Mask.SetRange(AccessorPackedIndex, TypePackingSize, true);
+		}
 
 		check(Accessor && Keys);
 
@@ -228,7 +146,7 @@ void UPCGInstanceDataPackerBase::PackCustomDataFromAccessors(const FPCGMeshInsta
 				OutPackedCustomData.CustomData[PackedIndex] = Value;
 			};
 
-			ensure(PCGMetadataElementCommon::ApplyOnAccessor<float>(*Keys, *Accessor, PackFloats, Flags));
+			ensure(PCGMetadataElementCommon::ApplyOnAccessor<float>(*Keys, *Accessor, PackFloats, Flags, PCGMetadataElementCommon::DefaultChunkSize, NumInstances));
 
 			AccessorPackedIndex += 1;
 
@@ -245,7 +163,7 @@ void UPCGInstanceDataPackerBase::PackCustomDataFromAccessors(const FPCGMeshInsta
 				OutPackedCustomData.CustomData[PackedIndex + 2] = static_cast<float>(Value.Yaw);
 			};
 
-			ensure(PCGMetadataElementCommon::ApplyOnAccessor<FRotator>(*Keys, *Accessor, PackFloats, Flags));
+			ensure(PCGMetadataElementCommon::ApplyOnAccessor<FRotator>(*Keys, *Accessor, PackFloats, Flags, PCGMetadataElementCommon::DefaultChunkSize, NumInstances));
 
 			AccessorPackedIndex += 3;
 
@@ -261,7 +179,7 @@ void UPCGInstanceDataPackerBase::PackCustomDataFromAccessors(const FPCGMeshInsta
 				OutPackedCustomData.CustomData[PackedIndex + 1] = static_cast<float>(Value.Y);
 			};
 
-			ensure(PCGMetadataElementCommon::ApplyOnAccessor<FVector2D>(*Keys, *Accessor, PackFloats, Flags));
+			ensure(PCGMetadataElementCommon::ApplyOnAccessor<FVector2D>(*Keys, *Accessor, PackFloats, Flags, PCGMetadataElementCommon::DefaultChunkSize, NumInstances));
 
 			AccessorPackedIndex += 2;
 
@@ -278,7 +196,7 @@ void UPCGInstanceDataPackerBase::PackCustomDataFromAccessors(const FPCGMeshInsta
 				OutPackedCustomData.CustomData[PackedIndex + 2] = static_cast<float>(Value.Z);
 			};
 
-			ensure(PCGMetadataElementCommon::ApplyOnAccessor<FVector>(*Keys, *Accessor, PackFloats, Flags));
+			ensure(PCGMetadataElementCommon::ApplyOnAccessor<FVector>(*Keys, *Accessor, PackFloats, Flags, PCGMetadataElementCommon::DefaultChunkSize, NumInstances));
 
 			AccessorPackedIndex += 3;
 
@@ -296,7 +214,7 @@ void UPCGInstanceDataPackerBase::PackCustomDataFromAccessors(const FPCGMeshInsta
 				OutPackedCustomData.CustomData[PackedIndex + 3] = static_cast<float>(Value.W);
 			};
 
-			ensure(PCGMetadataElementCommon::ApplyOnAccessor<FVector4>(*Keys, *Accessor, PackFloats, Flags));
+			ensure(PCGMetadataElementCommon::ApplyOnAccessor<FVector4>(*Keys, *Accessor, PackFloats, Flags, PCGMetadataElementCommon::DefaultChunkSize, NumInstances));
 
 			AccessorPackedIndex += 4;
 
@@ -314,7 +232,7 @@ void UPCGInstanceDataPackerBase::PackCustomDataFromAccessors(const FPCGMeshInsta
 				OutPackedCustomData.CustomData[PackedIndex + 3] = static_cast<float>(Value.W);
 			};
 
-			ensure(PCGMetadataElementCommon::ApplyOnAccessor<FQuat>(*Keys, *Accessor, PackFloats, Flags));
+			ensure(PCGMetadataElementCommon::ApplyOnAccessor<FQuat>(*Keys, *Accessor, PackFloats, Flags, PCGMetadataElementCommon::DefaultChunkSize, NumInstances));
 
 			AccessorPackedIndex += 4;
 
@@ -325,3 +243,338 @@ void UPCGInstanceDataPackerBase::PackCustomDataFromAccessors(const FPCGMeshInsta
 		}
 	}
 }
+
+void PCGInstanceDataPackerBase::PackCustomDataFromAttributes(FPackedDataFromAttributesParams& InParams)
+{
+	check(InParams.CommonParams.OutPackedCustomData);
+
+	if (!InParams.InData || InParams.CommonParams.NumInstances <= 0 || !InParams.InData->ConstMetadata())
+	{
+		return;
+	}
+
+	if (InParams.Attributes.IsEmpty() && !InParams.AttributeIdentifiers.IsEmpty())
+	{
+		// Extract Attributes and start again
+		TArray<const FPCGMetadataAttributeBase*> Attributes;
+		Attributes.Reserve(InParams.AttributeIdentifiers.Num());
+		for (const FPCGAttributeIdentifier& AttributeIdentifier : InParams.AttributeIdentifiers)
+		{
+			const FPCGMetadataAttributeBase* Attribute = InParams.InData->ConstMetadata()->GetConstAttribute(AttributeIdentifier);
+			if (Attribute)
+			{
+				Attributes.Add(Attribute);
+			}
+		}
+
+		if (Attributes.IsEmpty())
+		{
+			// Nothing to do
+			InParams.CommonParams.ZeroOutCustomData();
+			return;
+		}
+
+		InParams.Attributes = Attributes;
+		return PackCustomDataFromAttributes(InParams);
+	}
+
+	if (!Algo::AllOf(InParams.Attributes) || InParams.Attributes.IsEmpty())
+	{
+		InParams.CommonParams.ZeroOutCustomData();
+		return;
+	}
+
+	if (!InParams.CommonParams.OptionalOffsets.IsEmpty() && (InParams.CommonParams.OptionalOffsets.Num() != InParams.Attributes.Num()))
+	{
+		PCGLog::LogErrorOnGraph(FText::Format(LOCTEXT("MismatchOffsetSelectors2", "[PCGInstanceDataPackerBase::PackCustomDataFromAttributes] There is a mismatch between the number of attributes ({0}) and offsets ({1})"),
+			InParams.Attributes.Num(), InParams.CommonParams.OptionalOffsets.Num()), InParams.CommonParams.OptionalContext);
+		
+		InParams.CommonParams.ZeroOutCustomData();
+		return;
+	}
+
+	// Verify if we have multi-domain
+	const bool bHasMultiDomain = Algo::AnyOf(InParams.Attributes, [FirstAttribute = InParams.Attributes[0]](const FPCGMetadataAttributeBase* Attribute) { return Attribute->GetMetadataDomain() != FirstAttribute->GetMetadataDomain(); });
+
+	if (bHasMultiDomain && !InParams.OptionalIndices.IsEmpty())
+	{
+		PCGLog::LogErrorOnGraph(LOCTEXT("MultiDomainWithIndices", "Provided attributes are on multiple domains, while providing custom indices. This is unsupported."), InParams.CommonParams.OptionalContext);
+		InParams.CommonParams.ZeroOutCustomData();
+		return;
+	}
+	
+	TArray<TUniquePtr<const IPCGAttributeAccessor>> Accessors;
+	TArray<TUniquePtr<const IPCGAttributeAccessorKeys>> AccessorKeys;
+	TArray<const PCGMetadataEntryKey> FilteredKeys;
+
+	int32 NumCustomDataFloats = 0;
+
+	// If we have optional indices, we need first to extract all the keys that will be read from.
+	if (!InParams.OptionalIndices.IsEmpty())
+	{
+		TUniquePtr<const IPCGAttributeAccessorKeys> TempKey = PCGAttributeAccessorHelpers::CreateConstKeys(InParams.Attributes[0]);
+		check(TempKey);
+		
+		TArray<const PCGMetadataEntryKey*> AllKeys;
+		FilteredKeys.Empty(InParams.OptionalIndices.Num());
+
+		if (TempKey->GetKeys<PCGMetadataEntryKey>(0, AllKeys))
+		{
+			Algo::Transform(InParams.OptionalIndices, FilteredKeys, [&AllKeys](const int32 Index) { return AllKeys.IsValidIndex(Index) ? *AllKeys[Index] : PCGInvalidEntryKey; });
+			AccessorKeys.Emplace(MakeUnique<FPCGAttributeAccessorKeysEntries>(FilteredKeys));
+		}
+	}
+
+	for (int32 i = 0; i < InParams.Attributes.Num(); ++i)
+	{
+		const FPCGMetadataAttributeBase* Attribute = InParams.Attributes[i];
+		if (!InParams.CommonParams.OptionalOffsets.IsEmpty())
+		{
+			if (!InParams.CommonParams.ValidateOffsets(i, NumCustomDataFloats))
+			{
+				return;
+			}
+			
+			NumCustomDataFloats = InParams.CommonParams.OptionalOffsets[i];
+		}
+		
+		const FPCGMetadataDomain* MetadataDomain = Attribute->GetMetadataDomain();
+		
+		if (Attribute && AddTypeToPacking(Attribute->GetTypeId(), NumCustomDataFloats))
+		{
+			Accessors.Emplace(PCGAttributeAccessorHelpers::CreateConstAccessor(Attribute, MetadataDomain));
+			
+			if (bHasMultiDomain)
+			{
+				AccessorKeys.Emplace(PCGAttributeAccessorHelpers::CreateConstKeys(Attribute));
+			}
+		}
+	}
+
+	if (Accessors.IsEmpty())
+	{
+		InParams.CommonParams.ZeroOutCustomData();
+		return;
+	}
+
+	// When all is done, verify that the number of custom data float is valid.
+	if (InParams.CommonParams.OutPackedCustomData->NumCustomDataFloats != 0)
+	{
+		if (!InParams.CommonParams.ValidateNumOfFloats(NumCustomDataFloats))
+		{
+			InParams.CommonParams.ZeroOutCustomData();
+			return;
+		}
+	}
+	else
+	{
+		InParams.CommonParams.OutPackedCustomData->NumCustomDataFloats = NumCustomDataFloats;
+	}
+
+	FPackedDataFromAccessorParams Params =
+	{
+		.CommonParams = InParams.CommonParams,
+		.Accessors = MakeConstArrayView(Accessors),
+		.AccessorKeys = MakeConstArrayView(AccessorKeys),
+	};
+
+	PackCustomDataFromAccessors(Params);
+}
+
+void PCGInstanceDataPackerBase::PackCustomData(FPackedDataParams& InParams)
+{
+	check(InParams.CommonParams.OutPackedCustomData);
+
+	if (!InParams.InData || InParams.CommonParams.NumInstances <= 0 || !InParams.InData->ConstMetadata())
+	{
+		InParams.CommonParams.ZeroOutCustomData();
+		return;
+	}
+
+	if (!InParams.CommonParams.OptionalOffsets.IsEmpty() && (InParams.CommonParams.OptionalOffsets.Num() != InParams.Selectors.Num()))
+	{
+		PCGLog::LogErrorOnGraph(FText::Format(LOCTEXT("MismatchOffsetSelectors", "[PCGInstanceDataPackerBase::PackCustomData] There is a mismatch between the number of selectors ({0}) and offsets ({1})"),
+			InParams.Selectors.Num(), InParams.CommonParams.OptionalOffsets.Num()), InParams.CommonParams.OptionalContext);
+		InParams.CommonParams.ZeroOutCustomData();
+		return;
+	}
+
+	TArray<TUniquePtr<const IPCGAttributeAccessor>> Accessors;
+	TArray<TUniquePtr<const IPCGAttributeAccessorKeys>> AccessorKeys;
+
+	int32 NumCustomDataFloats = 0;
+
+	for (int32 i = 0; i < InParams.Selectors.Num(); ++i)
+	{
+		const FPCGAttributePropertyInputSelector& Selector = InParams.Selectors[i];
+
+		if (!InParams.CommonParams.OptionalOffsets.IsEmpty())
+		{
+			if (!InParams.CommonParams.ValidateOffsets(i, NumCustomDataFloats))
+			{
+				InParams.CommonParams.ZeroOutCustomData();
+				return;
+			}
+			
+			NumCustomDataFloats = InParams.CommonParams.OptionalOffsets[i];
+		}
+		
+		FPCGAttributePropertyInputSelector FinalSelector = Selector.CopyAndFixLast(InParams.InData);
+
+		TUniquePtr<const IPCGAttributeAccessor> Accessor = PCGAttributeAccessorHelpers::CreateConstAccessor(InParams.InData, FinalSelector);
+		TUniquePtr<const IPCGAttributeAccessorKeys> Keys = PCGAttributeAccessorHelpers::CreateConstKeys(InParams.InData, FinalSelector);
+		
+		if (!Accessor || !Keys)
+		{
+			PCGLog::Metadata::LogFailToCreateAccessorError(FinalSelector, InParams.CommonParams.OptionalContext);
+			continue;
+		}
+		
+		if (AddTypeToPacking(Accessor->GetUnderlyingType(), NumCustomDataFloats))
+		{
+			Accessors.Emplace(MoveTemp(Accessor));
+			AccessorKeys.Emplace(MoveTemp(Keys));
+		}
+	}
+
+	if (Accessors.IsEmpty())
+	{
+		InParams.CommonParams.ZeroOutCustomData();
+		return;
+	}
+
+	// When all is done, verify that the number of custom data float is valid.
+	if (InParams.CommonParams.OutPackedCustomData->NumCustomDataFloats != 0)
+	{
+		if (!InParams.CommonParams.ValidateNumOfFloats(NumCustomDataFloats))
+		{
+			InParams.CommonParams.ZeroOutCustomData();
+			return;
+		}
+	}
+	else
+	{
+		InParams.CommonParams.OutPackedCustomData->NumCustomDataFloats = NumCustomDataFloats;
+	}
+
+	FPackedDataFromAccessorParams Params =
+	{
+		.CommonParams = InParams.CommonParams,
+		.Accessors = MakeConstArrayView(Accessors),
+		.AccessorKeys = MakeConstArrayView(AccessorKeys)
+	};
+
+	PackCustomDataFromAccessors(Params);
+}
+
+bool PCGInstanceDataPackerBase::AddTypeToPacking(int TypeId, int32& OutNumCustomDataFloats)
+{
+	const int32 TypePackingSize = GetTypePackingSize(TypeId);
+	if (TypePackingSize > 0)
+	{
+		OutNumCustomDataFloats += TypePackingSize;
+		return true;
+	}
+	else
+	{
+		return false;
+	}
+}
+
+int32 PCGInstanceDataPackerBase::GetTypePackingSize(int TypeId)
+{
+	switch (TypeId)
+	{
+	case PCG::Private::MetadataTypes<bool>::Id: // fall-through
+	case PCG::Private::MetadataTypes<float>::Id: // fall-through
+	case PCG::Private::MetadataTypes<double>::Id: // fall-through
+	case PCG::Private::MetadataTypes<int32>::Id: // fall-through
+	case PCG::Private::MetadataTypes<int64>::Id:
+		return 1;
+	case PCG::Private::MetadataTypes<FVector2D>::Id:
+		return 2;
+	case PCG::Private::MetadataTypes<FRotator>::Id: // fall-through
+	case PCG::Private::MetadataTypes<FVector>::Id:
+		return 3;
+	case PCG::Private::MetadataTypes<FVector4>::Id: // fall-through
+	case PCG::Private::MetadataTypes<FQuat>::Id:
+		return 4;
+	default:
+		return 0;
+	}
+}
+
+TConstArrayView<float> FPCGPackedCustomData::GetViewForIndex(int32 Index) const
+{
+	
+	if (NumCustomDataFloats == 0 || (Index + 1) * NumCustomDataFloats > CustomData.Num())
+	{
+		return {};
+	}
+	else
+	{
+		return MakeConstArrayView(CustomData.GetData() + (Index * NumCustomDataFloats), NumCustomDataFloats);
+	}
+}
+
+void UPCGInstanceDataPackerBase::PackInstances_Implementation(FPCGContext& Context, const UPCGSpatialData* InSpatialData, const FPCGMeshInstanceList& InstanceList, FPCGPackedCustomData& OutPackedCustomData) const
+{
+	PCGE_LOG_C(Error, GraphAndLog, &Context, LOCTEXT("InstanceDataPackerBaseFailed", "Unable to execute InstanceDataPacker pure virtual base function, override the PackInstances function or use a default implementation."));
+}
+
+void UPCGInstanceDataPackerBase::PackCustomDataFromAttributes(const FPCGMeshInstanceList& InstanceList, const UPCGMetadata* Metadata, const TArray<FName>& AttributeNames, FPCGPackedCustomData& OutPackedCustomData) const
+{
+	const TArray<FPCGAttributeIdentifier> AttributeIdentifiers = FPCGAttributeIdentifier::TransformNameArray(AttributeNames);
+	
+	PCGInstanceDataPackerBase::FPackedDataFromAttributesParams Params =
+	{
+		.CommonParams =
+		{
+			.NumInstances = InstanceList.InstancesIndices.Num(),
+			.OutPackedCustomData = &OutPackedCustomData,
+		},
+		.OptionalIndices = MakeConstArrayView(InstanceList.InstancesIndices),
+		.InData = InstanceList.PointData.Get(),
+		.AttributeIdentifiers = MakeConstArrayView(AttributeIdentifiers)
+	};
+
+	return PCGInstanceDataPackerBase::PackCustomDataFromAttributes(Params);
+}
+
+void UPCGInstanceDataPackerBase::PackCustomDataFromAttributes(const FPCGMeshInstanceList& InstanceList, const TArray<const FPCGMetadataAttributeBase*>& Attributes, FPCGPackedCustomData& OutPackedCustomData) const
+{
+	PCGInstanceDataPackerBase::FPackedDataFromAttributesParams Params =
+	{
+		.CommonParams =
+		{
+			.NumInstances = InstanceList.InstancesIndices.Num(),
+			.OutPackedCustomData = &OutPackedCustomData,
+		},
+		.OptionalIndices = MakeConstArrayView(InstanceList.InstancesIndices),
+		.InData = InstanceList.PointData.Get(),
+		.Attributes = Attributes
+	};
+
+	return PCGInstanceDataPackerBase::PackCustomDataFromAttributes(Params);
+}
+
+void UPCGInstanceDataPackerBase::PackCustomDataFromAccessors(const FPCGMeshInstanceList& InstanceList,
+	TArray<TUniquePtr<const IPCGAttributeAccessor>> Accessors,
+	TArray<TUniquePtr<const IPCGAttributeAccessorKeys>> AccessorKeys,
+	FPCGPackedCustomData& OutPackedCustomData) const
+{
+	PCGInstanceDataPackerBase::FPackedDataFromAccessorParams Params =
+	{
+		.CommonParams =
+		{
+			.NumInstances = InstanceList.InstancesIndices.Num(),
+			.OutPackedCustomData = &OutPackedCustomData,
+		},
+		.Accessors = MakeConstArrayView(Accessors),
+		.AccessorKeys = MakeConstArrayView(AccessorKeys),
+	};
+
+	return PCGInstanceDataPackerBase::PackCustomDataFromAccessors(Params);
+}
+
+#undef LOCTEXT_NAMESPACE
\ No newline at end of file
diff --git a/Engine/Plugins/PCG/Source/PCG/Private/Metadata/Accessors/PCGAttributeAccessorHelpers.cpp b/Engine/Plugins/PCG/Source/PCG/Private/Metadata/Accessors/PCGAttributeAccessorHelpers.cpp
--- a/Engine/Plugins/PCG/Source/PCG/Private/Metadata/Accessors/PCGAttributeAccessorHelpers.cpp
+++ b/Engine/Plugins/PCG/Source/PCG/Private/Metadata/Accessors/PCGAttributeAccessorHelpers.cpp
@@ -649,3 +649,33 @@ TUniquePtr<IPCGAttributeAccessorKeys> PCGAttributeAccessorHelpers::CreateKeys(UP
 {
 	return FPCGAttributeAccessorFactory::GetInstance().CreateSimpleKeys(InData, InSelector, /*bQuiet=*/ false);
 }
+
+TUniquePtr<const IPCGAttributeAccessorKeys> PCGAttributeAccessorHelpers::CreateConstKeys(const FPCGMetadataAttributeBase* InAttribute)
+{
+	const UPCGData* Data = InAttribute && InAttribute->GetMetadata() ? Cast<const UPCGData>(InAttribute->GetMetadata()->GetOuter()) : nullptr;
+	
+	if (!Data || !InAttribute->GetMetadataDomain())
+	{
+		return nullptr;
+	}
+
+	FPCGAttributePropertySelector Selector = FPCGAttributePropertySelector::CreateAttributeSelector(InAttribute->Name);
+	Data->SetDomainFromDomainID(InAttribute->GetMetadataDomain()->GetDomainID(), Selector);
+		
+	return PCGAttributeAccessorHelpers::CreateConstKeys(Data, Selector);
+}
+
+TUniquePtr<IPCGAttributeAccessorKeys> PCGAttributeAccessorHelpers::CreateKeys(FPCGMetadataAttributeBase* InAttribute)
+{
+	UPCGData* Data = InAttribute && InAttribute->GetMetadata() ? Cast<UPCGData>(InAttribute->GetMetadata()->GetOuter()) : nullptr;
+	
+	if (!Data || !InAttribute->GetMetadataDomain())
+	{
+		return nullptr;
+	}
+
+	FPCGAttributePropertySelector Selector = FPCGAttributePropertySelector::CreateAttributeSelector(InAttribute->Name);
+	Data->SetDomainFromDomainID(InAttribute->GetMetadataDomain()->GetDomainID(), Selector);
+		
+	return PCGAttributeAccessorHelpers::CreateKeys(Data, Selector);
+}
diff --git a/Engine/Plugins/PCG/Source/PCG/Private/Tests/Metadata/PCGInstancePackerTests.cpp b/Engine/Plugins/PCG/Source/PCG/Private/Tests/Metadata/PCGInstancePackerTests.cpp
new file mode 100644
--- /dev/null
+++ b/Engine/Plugins/PCG/Source/PCG/Private/Tests/Metadata/PCGInstancePackerTests.cpp
@@ -0,0 +1,386 @@
+// Copyright Epic Games, Inc. All Rights Reserved.
+
+#if WITH_EDITOR
+#include "PCGParamData.h"
+#include "Algo/Compare.h"
+
+#include "Tests/PCGTestsCommon.h"
+
+#include "Data/PCGPointData.h"
+#include "Metadata/PCGMetadata.h"
+#include "Metadata/PCGMetadataAttributeTpl.h"
+
+#include "Engine/Engine.h"
+#include "InstanceDataPackers/PCGInstanceDataPackerBase.h"
+
+class FPCGInstancePackerTestClass : public FPCGTestBaseClass
+{
+public:
+	using FPCGTestBaseClass::FPCGTestBaseClass;
+
+protected:
+	static constexpr int32 NumPoints = 10;
+	static constexpr FLazyName Int32AttributeName = "Int32";
+	static constexpr FLazyName Int64AttributeName = "Int64";
+	static constexpr FLazyName FloatAttributeName = "Float";
+	static constexpr FLazyName DoubleAttributeName = "Double";
+	static constexpr FLazyName BoolAttributeName = "Bool";
+	static constexpr FLazyName Vec2AttributeName = "Vec2";
+	static constexpr FLazyName Vec3AttributeName = "Vec3";
+	static constexpr FLazyName Vec4AttributeName = "Vec4";
+	static constexpr FLazyName RotatorAttributeName = "Rotator";
+	static constexpr FLazyName QuatAttributeName = "Quat";
+	static constexpr FLazyName StringAttributeName = "String";
+	
+	UPCGPointArrayData* CreatePointData()
+	{
+		UPCGPointArrayData* PointData = NewObject<UPCGPointArrayData>();
+		UPCGMetadata* Metadata = PointData->MutableMetadata();
+		
+		PointData->SetNumPoints(NumPoints);
+		PointData->AllocateProperties(EPCGPointNativeProperties::Transform | EPCGPointNativeProperties::Density | EPCGPointNativeProperties::MetadataEntry);
+
+		TPCGValueRange<float> DensityRange = PointData->GetDensityValueRange();
+		TPCGValueRange<FTransform> TransformRange = PointData->GetTransformValueRange();
+		TPCGValueRange<PCGMetadataEntryKey> MetadataEntryRange = PointData->GetMetadataEntryValueRange();
+
+		auto InitMetadata = [Metadata, &MetadataEntryRange](const FPCGMetadataDomainID DomainID)
+		{
+			FPCGMetadataDomain* Domain = Metadata->GetMetadataDomain(DomainID);
+
+			FPCGMetadataAttribute<int32>* Int32Attribute = Domain->CreateAttribute<int32>(Int32AttributeName, 1, false, false);
+			FPCGMetadataAttribute<int64>* Int64Attribute = Domain->CreateAttribute<int64>(Int64AttributeName, 1ll, false, false);
+			FPCGMetadataAttribute<float>* FloatAttribute = Domain->CreateAttribute<float>(FloatAttributeName, 1.0f, false, false);
+			FPCGMetadataAttribute<double>* DoubleAttribute = Domain->CreateAttribute<double>(DoubleAttributeName, 1.0, false, false);
+			FPCGMetadataAttribute<bool>* BoolAttribute = Domain->CreateAttribute<bool>(BoolAttributeName, true, false, false);
+			FPCGMetadataAttribute<FVector2D>* Vec2Attribute = Domain->CreateAttribute<FVector2D>(Vec2AttributeName, FVector2D::One(), false, false);
+			FPCGMetadataAttribute<FVector>* Vec3Attribute = Domain->CreateAttribute<FVector>(Vec3AttributeName, FVector::One(), false, false);
+			FPCGMetadataAttribute<FVector4>* Vec4Attribute = Domain->CreateAttribute<FVector4>(Vec4AttributeName, FVector4::One(), false, false);
+			FPCGMetadataAttribute<FRotator>* RotatorAttribute = Domain->CreateAttribute<FRotator>(RotatorAttributeName, FRotator(1.0, 1.0, 1.0), false, false);
+			FPCGMetadataAttribute<FQuat>* QuatAttribute = Domain->CreateAttribute<FQuat>(QuatAttributeName, FQuat(1.0, 1.0, 1.0, 1.0), false, false);
+			FPCGMetadataAttribute<FString>* StringAttribute = Domain->CreateAttribute<FString>(StringAttributeName, TEXT("1"), false, false);
+				
+			check(Int32Attribute);
+			check(Int64Attribute);
+			check(FloatAttribute);
+			check(DoubleAttribute);
+			check(BoolAttribute);
+			check(Vec2Attribute);
+			check(Vec3Attribute);
+			check(Vec4Attribute);
+			check(RotatorAttribute);
+			check(QuatAttribute);
+			check(StringAttribute);
+
+			if (DomainID == PCGMetadataDomainID::Elements)
+			{
+				for (int32 i = 0; i < NumPoints; ++i)
+				{
+					MetadataEntryRange[i] = Domain->AddEntry();
+					
+					Int32Attribute->SetValue(MetadataEntryRange[i], i);
+					Int64Attribute->SetValue(MetadataEntryRange[i], i);
+					FloatAttribute->SetValue(MetadataEntryRange[i], i);
+					DoubleAttribute->SetValue(MetadataEntryRange[i], i);
+					BoolAttribute->SetValue(MetadataEntryRange[i], i != 0);
+					Vec2Attribute->SetValue(MetadataEntryRange[i], FVector2D(i, i));
+					Vec3Attribute->SetValue(MetadataEntryRange[i], FVector(i, i, i));
+					Vec4Attribute->SetValue(MetadataEntryRange[i], FVector4(i, i, i, i));
+					RotatorAttribute->SetValue(MetadataEntryRange[i], FRotator(i, i, i));
+					QuatAttribute->SetValue(MetadataEntryRange[i], FQuat(i, i, i, i));
+					StringAttribute->SetValue(MetadataEntryRange[i], LexToString(i));
+				}
+			}
+		};
+
+		InitMetadata(PCGMetadataDomainID::Elements);
+		InitMetadata(PCGMetadataDomainID::Data);
+
+		for (int32 i = 0; i < NumPoints; ++i)
+		{
+			DensityRange[i] = i;
+			TransformRange[i].SetLocation(FVector(i, i, i));
+		}
+
+		return PointData;
+	}
+	
+	bool Validate()
+	{
+		UTEST_EQUAL("Custom Data has the right number of expected floats", ActualCustomData.NumCustomDataFloats, ExpectedCustomData.NumCustomDataFloats);
+		
+		if (!ExpectedCustomData.Mask.IsEmpty())
+		{
+			UTEST_EQUAL("Custom Data has the expected mask", ActualCustomData.Mask, ExpectedCustomData.Mask);
+		}
+
+		UTEST_EQUAL("Custom Data has the right number of values", ActualCustomData.CustomData.Num(), ExpectedCustomData.CustomData.Num());
+		UTEST_TRUE("Custom data has expected values", Algo::Compare(ActualCustomData.CustomData, ExpectedCustomData.CustomData, [](const float LHS, const float RHS) { return FMath::IsNearlyEqual(LHS, RHS); }))
+		
+		return true;
+	}
+
+	FPCGPackedCustomData ExpectedCustomData;
+	FPCGPackedCustomData ActualCustomData;
+};
+
+IMPLEMENT_CUSTOM_SIMPLE_AUTOMATION_TEST(FPCGInstancePackerTest_AllPoints, FPCGInstancePackerTestClass, "Plugins.PCG.InstancePacker.AllPoints", PCGTestsCommon::TestFlags)
+IMPLEMENT_CUSTOM_SIMPLE_AUTOMATION_TEST(FPCGInstancePackerTest_InvalidType, FPCGInstancePackerTestClass, "Plugins.PCG.InstancePacker.InvalidType", PCGTestsCommon::TestFlags)
+IMPLEMENT_CUSTOM_SIMPLE_AUTOMATION_TEST(FPCGInstancePackerTest_PointsOffset, FPCGInstancePackerTestClass, "Plugins.PCG.InstancePacker.PointsOffset", PCGTestsCommon::TestFlags)
+IMPLEMENT_CUSTOM_SIMPLE_AUTOMATION_TEST(FPCGInstancePackerTest_InvalidOffset_NotEnoughFloats, FPCGInstancePackerTestClass, "Plugins.PCG.InstancePacker.InvalidOffset.NotEnoughFloats", PCGTestsCommon::TestFlags)
+IMPLEMENT_CUSTOM_SIMPLE_AUTOMATION_TEST(FPCGInstancePackerTest_InvalidOffset_Overlapping, FPCGInstancePackerTestClass, "Plugins.PCG.InstancePacker.InvalidOffset.Overlapping", PCGTestsCommon::TestFlags)
+IMPLEMENT_CUSTOM_SIMPLE_AUTOMATION_TEST(FPCGInstancePackerTest_TooManyValuesToBeExtracted, FPCGInstancePackerTestClass, "Plugins.PCG.InstancePacker.TooManyValuesToBeExtracted", PCGTestsCommon::TestFlags)
+
+bool FPCGInstancePackerTest_AllPoints::RunTest(const FString& Parameters)
+{
+	const UPCGPointArrayData* PointData = CreatePointData();
+
+	// All of the attributes (21 floats) + $Density and $Position (4 floats)  -> 25 floats
+	constexpr int32 NumCustomDataFloats = 25;
+	
+	ExpectedCustomData.NumCustomDataFloats = NumCustomDataFloats;
+
+	// Values should be 25 zeroes, then 25 ones, then 25 twos, etc...
+	ExpectedCustomData.CustomData.SetNumZeroed(NumCustomDataFloats * NumPoints);
+	for (int32 i = NumCustomDataFloats; i < NumCustomDataFloats * NumPoints; ++i)
+	{
+		if (i % NumCustomDataFloats == 0)
+		{
+			// bool
+			ExpectedCustomData.CustomData[i] = 1.0f;
+		}
+		else
+		{
+			ExpectedCustomData.CustomData[i] = i / NumCustomDataFloats;
+		}
+	}
+
+	TArray<FPCGAttributePropertyInputSelector> Selectors =
+	{
+		FPCGAttributePropertyInputSelector::CreateAttributeSelector<FPCGAttributePropertyInputSelector>(BoolAttributeName),
+		FPCGAttributePropertyInputSelector::CreateAttributeSelector<FPCGAttributePropertyInputSelector>(Int32AttributeName),
+		FPCGAttributePropertyInputSelector::CreateAttributeSelector<FPCGAttributePropertyInputSelector>(Int64AttributeName),
+		FPCGAttributePropertyInputSelector::CreateAttributeSelector<FPCGAttributePropertyInputSelector>(FloatAttributeName),
+		FPCGAttributePropertyInputSelector::CreateAttributeSelector<FPCGAttributePropertyInputSelector>(DoubleAttributeName),
+		FPCGAttributePropertyInputSelector::CreateAttributeSelector<FPCGAttributePropertyInputSelector>(Vec2AttributeName),
+		FPCGAttributePropertyInputSelector::CreateAttributeSelector<FPCGAttributePropertyInputSelector>(Vec3AttributeName),
+		FPCGAttributePropertyInputSelector::CreateAttributeSelector<FPCGAttributePropertyInputSelector>(Vec4AttributeName),
+		FPCGAttributePropertyInputSelector::CreateAttributeSelector<FPCGAttributePropertyInputSelector>(RotatorAttributeName),
+		FPCGAttributePropertyInputSelector::CreateAttributeSelector<FPCGAttributePropertyInputSelector>(QuatAttributeName),
+		FPCGAttributePropertyInputSelector::CreatePointPropertySelector<FPCGAttributePropertyInputSelector>(EPCGPointProperties::Density),
+		FPCGAttributePropertyInputSelector::CreatePointPropertySelector<FPCGAttributePropertyInputSelector>(EPCGPointProperties::Position),
+	};
+
+	PCGInstanceDataPackerBase::FPackedDataParams Params =
+	{
+		.CommonParams =
+		{
+			.NumInstances = NumPoints,
+			.OutPackedCustomData = &ActualCustomData,
+		},
+		.InData = PointData,
+		.Selectors = MakeConstArrayView(Selectors)
+	};
+
+	PCGInstanceDataPackerBase::PackCustomData(Params);
+
+	return Validate();
+}
+
+/** String attribute is not supported for packing. We set the expected number of float to 1, so we have NumPoints zeroes in the custom data. */
+bool FPCGInstancePackerTest_InvalidType::RunTest(const FString& Parameters)
+{
+	const UPCGPointArrayData* PointData = CreatePointData();
+
+	ExpectedCustomData.NumCustomDataFloats = 1;
+	ExpectedCustomData.CustomData.SetNumZeroed(NumPoints);
+	
+	ActualCustomData.NumCustomDataFloats = 1;
+
+	TArray<FPCGAttributePropertyInputSelector> Selectors =
+	{
+		FPCGAttributePropertyInputSelector::CreateAttributeSelector<FPCGAttributePropertyInputSelector>(StringAttributeName),
+	};
+
+	PCGInstanceDataPackerBase::FPackedDataParams Params =
+	{
+		.CommonParams =
+		{
+			.NumInstances = NumPoints,
+			.OutPackedCustomData = &ActualCustomData,
+		},
+		.InData = PointData,
+		.Selectors = MakeConstArrayView(Selectors)
+	};
+
+	PCGInstanceDataPackerBase::PackCustomData(Params);
+
+	return Validate();
+}
+
+bool FPCGInstancePackerTest_PointsOffset::RunTest(const FString& Parameters)
+{
+	const UPCGPointArrayData* PointData = CreatePointData();
+
+	constexpr int32 NumCustomDataFloats = 5;
+	
+	ExpectedCustomData.NumCustomDataFloats = NumCustomDataFloats;
+	ActualCustomData.NumCustomDataFloats = NumCustomDataFloats;
+
+	ExpectedCustomData.Mask.Init(true, NumCustomDataFloats);
+	ExpectedCustomData.Mask[2] = false;
+	ExpectedCustomData.Mask[4] = false;
+
+	// Values [i, i, 50, i, -47] for each point, where i is point index.
+	// Values are initialize for both buffers, to check if the values not on the offset are kept.
+	ExpectedCustomData.CustomData.Reserve(NumCustomDataFloats * NumPoints);
+	ActualCustomData.CustomData.Reserve(NumCustomDataFloats * NumPoints);
+	for (int32 i = 0; i < NumCustomDataFloats * NumPoints; ++i)
+	{
+		if (i % 5 == 2)
+		{
+			ExpectedCustomData.CustomData.Add(50.f);
+			ActualCustomData.CustomData.Add(50.f);
+		}
+		else if (i % 5 == 4)
+		{
+			ExpectedCustomData.CustomData.Add(-47.f);
+			ActualCustomData.CustomData.Add(-47.f);
+		}
+		else
+		{
+			ExpectedCustomData.CustomData.Add(i / 5);
+			ActualCustomData.CustomData.Add(0.f);
+		}
+	}
+
+	TArray<FPCGAttributePropertyInputSelector> Selectors =
+	{
+		FPCGAttributePropertyInputSelector::CreateAttributeSelector<FPCGAttributePropertyInputSelector>(Vec2AttributeName),
+		FPCGAttributePropertyInputSelector::CreatePointPropertySelector<FPCGAttributePropertyInputSelector>(EPCGPointProperties::Position, NAME_None, {TEXT("X")}), // $Position.X
+	};
+
+	PCGInstanceDataPackerBase::FPackedDataParams Params =
+	{
+		.CommonParams =
+		{
+			.NumInstances = NumPoints,
+			.OutPackedCustomData = &ActualCustomData,
+			.OptionalOffsets = {0, 3} // Should write at offset 0 and at offset 3
+		},
+		.InData = PointData,
+		.Selectors = MakeConstArrayView(Selectors)
+	};
+
+	PCGInstanceDataPackerBase::PackCustomData(Params);
+
+	return Validate();
+}
+
+bool FPCGInstancePackerTest_InvalidOffset_NotEnoughFloats::RunTest(const FString& Parameters)
+{
+	const UPCGPointArrayData* PointData = CreatePointData();
+
+	constexpr int32 NumCustomDataFloats = 1;
+	
+	ExpectedCustomData.NumCustomDataFloats = NumCustomDataFloats;
+	ActualCustomData.NumCustomDataFloats = NumCustomDataFloats;
+
+	ExpectedCustomData.CustomData.SetNumZeroed(NumCustomDataFloats * NumPoints);
+	
+	TArray<FPCGAttributePropertyInputSelector> Selectors =
+	{
+		FPCGAttributePropertyInputSelector::CreateAttributeSelector<FPCGAttributePropertyInputSelector>(Int32AttributeName)
+	};
+
+	PCGInstanceDataPackerBase::FPackedDataParams Params =
+	{
+		.CommonParams =
+		{
+			.NumInstances = NumPoints,
+			.OutPackedCustomData = &ActualCustomData,
+			.OptionalOffsets = {5} // Offset too big
+		},
+		.InData = PointData,
+		.Selectors = MakeConstArrayView(Selectors)
+	};
+	
+	AddExpectedError(TEXT("[PCGInstanceDataPackerBase::ValidateNumOfFloats] Expected to have at least 6 custom floats, but was explicitly set to 1, which is not enough."), EAutomationExpectedErrorFlags::Contains, 1, /*bIsRegex=*/false);
+	
+	PCGInstanceDataPackerBase::PackCustomData(Params);
+
+	return Validate();
+}
+
+bool FPCGInstancePackerTest_InvalidOffset_Overlapping::RunTest(const FString& Parameters)
+{
+	const UPCGPointArrayData* PointData = CreatePointData();
+
+	constexpr int32 NumCustomDataFloats = 4;
+	
+	ExpectedCustomData.NumCustomDataFloats = NumCustomDataFloats;
+	ActualCustomData.NumCustomDataFloats = NumCustomDataFloats;
+
+	ExpectedCustomData.CustomData.SetNumZeroed(NumCustomDataFloats * NumPoints);
+	
+	TArray<FPCGAttributePropertyInputSelector> Selectors =
+	{
+		FPCGAttributePropertyInputSelector::CreateAttributeSelector<FPCGAttributePropertyInputSelector>(Vec3AttributeName),
+		FPCGAttributePropertyInputSelector::CreateAttributeSelector<FPCGAttributePropertyInputSelector>(FloatAttributeName)
+	};
+
+	PCGInstanceDataPackerBase::FPackedDataParams Params =
+	{
+		.CommonParams =
+		{
+			.NumInstances = NumPoints,
+			.OutPackedCustomData = &ActualCustomData,
+			.OptionalOffsets = {0, 1} // Overlapping offset, the second one should be higher or equal to 3
+		},
+		.InData = PointData,
+		.Selectors = MakeConstArrayView(Selectors)
+	};
+	
+	AddExpectedError(TEXT("[PCGInstanceDataPackerBase::ValidateOffsets] Offset 1 at index 1 is too small for the number of expected floats (3)"), EAutomationExpectedErrorFlags::Contains, 1, /*bIsRegex=*/false);
+	
+	PCGInstanceDataPackerBase::PackCustomData(Params);
+
+	return Validate();
+}
+
+bool FPCGInstancePackerTest_TooManyValuesToBeExtracted::RunTest(const FString& Parameters)
+{
+	const UPCGPointArrayData* PointData = CreatePointData();
+
+	constexpr int32 NumCustomDataFloats = 1;
+	
+	ExpectedCustomData.NumCustomDataFloats = NumCustomDataFloats;
+	ActualCustomData.NumCustomDataFloats = NumCustomDataFloats;
+
+	ExpectedCustomData.CustomData.SetNumZeroed(NumCustomDataFloats * NumPoints);
+	
+	TArray<FPCGAttributePropertyInputSelector> Selectors =
+	{
+		FPCGAttributePropertyInputSelector::CreateAttributeSelector<FPCGAttributePropertyInputSelector>(Vec3AttributeName),
+		FPCGAttributePropertyInputSelector::CreateAttributeSelector<FPCGAttributePropertyInputSelector>(FloatAttributeName)
+	};
+
+	PCGInstanceDataPackerBase::FPackedDataParams Params =
+	{
+		.CommonParams =
+		{
+			.NumInstances = NumPoints,
+			.OutPackedCustomData = &ActualCustomData,
+		},
+		.InData = PointData,
+		.Selectors = MakeConstArrayView(Selectors)
+	};
+	
+	AddExpectedError(TEXT("[PCGInstanceDataPackerBase::ValidateNumOfFloats] Expected to have at least 4 custom floats, but was explicitly set to 1, which is not enough."), EAutomationExpectedErrorFlags::Contains, 1, /*bIsRegex=*/false);
+	
+	PCGInstanceDataPackerBase::PackCustomData(Params);
+
+	return Validate();
+}
+
+#endif // WITH_EDITOR
\ No newline at end of file
diff --git a/Engine/Plugins/PCG/Source/PCG/Public/InstanceDataPackers/PCGInstanceDataPackerBase.h b/Engine/Plugins/PCG/Source/PCG/Public/InstanceDataPackers/PCGInstanceDataPackerBase.h
--- a/Engine/Plugins/PCG/Source/PCG/Public/InstanceDataPackers/PCGInstanceDataPackerBase.h
+++ b/Engine/Plugins/PCG/Source/PCG/Public/InstanceDataPackers/PCGInstanceDataPackerBase.h
@@ -2,18 +2,20 @@
 
 #pragma once
 
-#include "PCGInstanceDataPackerBase.generated.h"
+#include "Metadata/PCGMetadataCommon.h"
 
-#define UE_API PCG_API
+#include "PCGInstanceDataPackerBase.generated.h"
 
+struct FPCGAttributePropertyInputSelector;
+struct FPCGContext;
+struct FPCGMeshInstanceList;
 class FPCGMetadataAttributeBase;
+class FPCGMetadataDomain;
 class IPCGAttributeAccessor;
 class IPCGAttributeAccessorKeys;
+class UPCGData;
 class UPCGMetadata;
 class UPCGSpatialData;
-struct FPCGAttributePropertyInputSelector;
-struct FPCGContext;
-struct FPCGMeshInstanceList;
 
 USTRUCT(BlueprintType)
 struct FPCGPackedCustomData
@@ -25,8 +27,97 @@ struct FPCGPackedCustomData
 
 	UPROPERTY(BlueprintReadWrite, EditAnywhere, Category = Settings)
 	TArray<float> CustomData;
+
+	/** If the custom data was sparsely set, mark the custom data index corresponding bit. */
+	TBitArray<> Mask;
+
+	PCG_API TConstArrayView<float> GetViewForIndex(int32 Index) const;
 };
 
+namespace PCGInstanceDataPackerBase
+{
+	struct FPackedDataCommonParams
+	{
+		/** Number of different entries in the OutPackedCustomData. */
+		int32 NumInstances = 0;
+
+		/**
+		 * Out structure containing the packed data. Must be filled beforehand for its member NumCustomDataFloats for the accessor case or if OptionalOffsets is not empty. Otherwise it can be left at 0, it will be automatically filled.
+		 * If it is not 0, it'll be verified if the number is valid.
+		 */
+		FPCGPackedCustomData* OutPackedCustomData = nullptr;
+
+		/** If the packing is sparse, offsets can be provided. It is mandatory that those offsets are not overlapping, in ascending order and there must be as many offsets that there are different values. */
+		TConstArrayView<int32> OptionalOffsets;
+
+		/** Optional context for logging. */
+		FPCGContext* OptionalContext = nullptr;
+
+		/** For validation that the offset at OffsetIndex is not overlapping with NumOfCustomFloats. */
+		PCG_API bool ValidateOffsets(int32 OffsetIndex, int32 NumOfCustomFloats) const;
+
+		/** For validation that the computed number of custom floats is less or equal than the expected value set in OutPackedCustomData. */
+		PCG_API bool ValidateNumOfFloats(int32 NumOfCustomFloats) const;
+
+		/** If there is any kind of error, zero out the buffer in OutPackedCustomData. */
+		PCG_API void ZeroOutCustomData();
+	};
+	
+	struct FPackedDataFromAccessorParams
+	{
+		FPackedDataCommonParams CommonParams;
+		
+		/** List of accessors to read from. */
+		TConstArrayView<TUniquePtr<const IPCGAttributeAccessor>> Accessors;
+
+		/** List of accessor keys. Can provide just a single key if there is one key for all the accessors. */
+		TConstArrayView<TUniquePtr<const IPCGAttributeAccessorKeys>> AccessorKeys;
+	};
+
+	PCG_API void PackCustomDataFromAccessors(FPackedDataFromAccessorParams& InParams);
+
+	struct FPackedDataFromAttributesParams
+	{
+		FPackedDataCommonParams CommonParams;
+
+		/**
+		 * If not all values from the data are read from, can provide a list of indices. Note: It is mandatory that the attributes are all on
+		 * the same domain in that case, otherwise the mapping is ambiguous. The packing will fail if we have OptionalIndices and multi domain attributes.
+		 */
+		TConstArrayView<int32> OptionalIndices;
+
+		/** Input data to read attributes from. */
+		const UPCGData* InData = nullptr;
+
+		/** List of attributes to read from. Has precedence on AttributeIdentifiers. */
+		TConstArrayView<const FPCGMetadataAttributeBase*> Attributes;
+
+		/** List of attribute identifiers to read from. Ignored if Attributes is not empty. */
+		TConstArrayView<FPCGAttributeIdentifier> AttributeIdentifiers;
+	};
+
+	PCG_API void PackCustomDataFromAttributes(FPackedDataFromAttributesParams& InParams);
+
+	struct FPackedDataParams
+	{
+		FPackedDataCommonParams CommonParams;
+
+		/** Input data to read attributes from. */
+		const UPCGData* InData = nullptr;
+
+		/** List of selectors to read from. */
+		TConstArrayView<FPCGAttributePropertyInputSelector> Selectors;
+	};
+
+	PCG_API void PackCustomData(FPackedDataParams& InParams);
+
+	/** Will add the type packing size to OutPackedCustomData. Return false if the type is not supported. */
+	PCG_API bool AddTypeToPacking(int TypeId, int32& OutNumCustomDataFloats);
+
+	/** Return the type packing size. 0 if unsupported. */
+	PCG_API int32 GetTypePackingSize(int TypeId);
+}
+
 UCLASS(MinimalAPI, Abstract, BlueprintType, Blueprintable, ClassGroup = (Procedural))
 class UPCGInstanceDataPackerBase : public UObject 
 {
@@ -35,23 +126,23 @@ class UPCGInstanceDataPackerBase : public UObject
 public:
 	/** Defines the strategy for (H)ISM custom float data packing */
 	UFUNCTION(BlueprintNativeEvent, Category = InstancePacking)
-	UE_API void PackInstances(FPCGContext& Context, const UPCGSpatialData* InSpatialData, const FPCGMeshInstanceList& InstanceList, FPCGPackedCustomData& OutPackedCustomData) const;
+	PCG_API void PackInstances(FPCGContext& Context, const UPCGSpatialData* InSpatialData, const FPCGMeshInstanceList& InstanceList, FPCGPackedCustomData& OutPackedCustomData) const;
 
-	UE_API virtual void PackInstances_Implementation(FPCGContext& Context, const UPCGSpatialData* InSpatialData, const FPCGMeshInstanceList& InstanceList, FPCGPackedCustomData& OutPackedCustomData) const;
+	PCG_API virtual void PackInstances_Implementation(FPCGContext& Context, const UPCGSpatialData* InSpatialData, const FPCGMeshInstanceList& InstanceList, FPCGPackedCustomData& OutPackedCustomData) const;
 
 	/** Interprets Metadata TypeId and increments OutPackedCustomData.NumCustomDataFloats appropriately. Returns false if the type could not be interpreted. */
 	UFUNCTION(BlueprintCallable, Category = InstancePacking)
-	UE_API bool AddTypeToPacking(int TypeId, FPCGPackedCustomData& OutPackedCustomData) const;
+	bool AddTypeToPacking(int TypeId, FPCGPackedCustomData& OutPackedCustomData) const { return PCGInstanceDataPackerBase::AddTypeToPacking(TypeId, OutPackedCustomData.NumCustomDataFloats); }
 
 	/** Build a PackedCustomData by processing each attribute in order for each point in the InstanceList */
 	UFUNCTION(BlueprintCallable, Category = InstancePacking) 
-	UE_API void PackCustomDataFromAttributes(const FPCGMeshInstanceList& InstanceList, const UPCGMetadata* Metadata, const TArray<FName>& AttributeNames, FPCGPackedCustomData& OutPackedCustomData) const;
+	PCG_API void PackCustomDataFromAttributes(const FPCGMeshInstanceList& InstanceList, const UPCGMetadata* Metadata, const TArray<FName>& AttributeNames, FPCGPackedCustomData& OutPackedCustomData) const;
 
 	/** Build a PackedCustomData by processing each attribute in order for each point in the InstanceList */
-	UE_API void PackCustomDataFromAttributes(const FPCGMeshInstanceList& InstanceList, const TArray<const FPCGMetadataAttributeBase*>& Attributes, FPCGPackedCustomData& OutPackedCustomData) const;
+	PCG_API void PackCustomDataFromAttributes(const FPCGMeshInstanceList& InstanceList, const TArray<const FPCGMetadataAttributeBase*>& Attributes, FPCGPackedCustomData& OutPackedCustomData) const;
 
 	/** Build a PackedCustomData by processing each accessor in order for each point in the InstanceList */
-	UE_API void PackCustomDataFromAccessors(const FPCGMeshInstanceList& InstanceList, TArray<TUniquePtr<const IPCGAttributeAccessor>> Accessors, TArray<TUniquePtr<const IPCGAttributeAccessorKeys>> AccessorKeys, FPCGPackedCustomData& OutPackedCustomData) const;
+	PCG_API void PackCustomDataFromAccessors(const FPCGMeshInstanceList& InstanceList, TArray<TUniquePtr<const IPCGAttributeAccessor>> Accessors, TArray<TUniquePtr<const IPCGAttributeAccessorKeys>> AccessorKeys, FPCGPackedCustomData& OutPackedCustomData) const;
 
 	/** If OutNames is not null, returns a list of all attributes that will be packed. Returns true if this list can be statically determined (prior to execution). */
 	virtual bool GetAttributeNames(TArray<FName>* OutNames) { return false; }
@@ -60,4 +151,3 @@ public:
 	virtual TOptional<TConstArrayView<FPCGAttributePropertyInputSelector>> GetAttributeSelectors() const { return {}; }
 };
 
-#undef UE_API
diff --git a/Engine/Plugins/PCG/Source/PCG/Public/Metadata/Accessors/PCGAttributeAccessorHelpers.h b/Engine/Plugins/PCG/Source/PCG/Public/Metadata/Accessors/PCGAttributeAccessorHelpers.h
--- a/Engine/Plugins/PCG/Source/PCG/Public/Metadata/Accessors/PCGAttributeAccessorHelpers.h
+++ b/Engine/Plugins/PCG/Source/PCG/Public/Metadata/Accessors/PCGAttributeAccessorHelpers.h
@@ -120,6 +120,9 @@ namespace PCGAttributeAccessorHelpers
 	PCG_API TUniquePtr<const IPCGAttributeAccessorKeys> CreateConstKeys(const UPCGData* InData, const FPCGAttributePropertySelector& InSelector);
 	PCG_API TUniquePtr<IPCGAttributeAccessorKeys> CreateKeys(UPCGData* InData, const FPCGAttributePropertySelector& InSelector);
 
+	PCG_API TUniquePtr<const IPCGAttributeAccessorKeys> CreateConstKeys(const FPCGMetadataAttributeBase* InAttribute);
+	PCG_API TUniquePtr<IPCGAttributeAccessorKeys> CreateKeys(FPCGMetadataAttributeBase* InAttribute);
+
 	/** 
 	 * Utility function to extract all the values through the accessor for all the keys.
 	 *
