[PCG] Fix Mac compilation, static analysis and bughawk

#rnx

diff --git a/Engine/Plugins/PCG/Source/PCG/Private/InstanceDataPackers/PCGInstanceDataPackerBase.cpp b/Engine/Plugins/PCG/Source/PCG/Private/InstanceDataPackers/PCGInstanceDataPackerBase.cpp
--- a/Engine/Plugins/PCG/Source/PCG/Private/InstanceDataPackers/PCGInstanceDataPackerBase.cpp
+++ b/Engine/Plugins/PCG/Source/PCG/Private/InstanceDataPackers/PCGInstanceDataPackerBase.cpp
@@ -323,11 +323,16 @@ void PCGInstanceDataPackerBase::PackCustomDataFromAttributes(FPackedDataFromAttr
 			Algo::Transform(InParams.OptionalIndices, FilteredKeys, [&AllKeys](const int32 Index) { return AllKeys.IsValidIndex(Index) ? *AllKeys[Index] : PCGInvalidEntryKey; });
 			AccessorKeys.Emplace(MakeUnique<FPCGAttributeAccessorKeysEntries>(FilteredKeys));
 		}
+		else
+		{
+			AccessorKeys.Emplace(MoveTemp(TempKey));
+		}
 	}
 
 	for (int32 i = 0; i < InParams.Attributes.Num(); ++i)
 	{
 		const FPCGMetadataAttributeBase* Attribute = InParams.Attributes[i];
+		check(Attribute);
 		if (!InParams.CommonParams.OptionalOffsets.IsEmpty())
 		{
 			if (!InParams.CommonParams.ValidateOffsets(i, NumCustomDataFloats))
@@ -517,6 +522,13 @@ TConstArrayView<float> FPCGPackedCustomData::GetViewForIndex(int32 Index) const
 	}
 }
 
+void FPCGPackedCustomData::Reset()
+{
+	NumCustomDataFloats = 0;
+	CustomData.Reset();
+	Mask.Reset();
+}
+
 void UPCGInstanceDataPackerBase::PackInstances_Implementation(FPCGContext& Context, const UPCGSpatialData* InSpatialData, const FPCGMeshInstanceList& InstanceList, FPCGPackedCustomData& OutPackedCustomData) const
 {
 	PCGE_LOG_C(Error, GraphAndLog, &Context, LOCTEXT("InstanceDataPackerBaseFailed", "Unable to execute InstanceDataPacker pure virtual base function, override the PackInstances function or use a default implementation."));
diff --git a/Engine/Plugins/PCG/Source/PCG/Private/Tests/Metadata/PCGInstancePackerTests.cpp b/Engine/Plugins/PCG/Source/PCG/Private/Tests/Metadata/PCGInstancePackerTests.cpp
--- a/Engine/Plugins/PCG/Source/PCG/Private/Tests/Metadata/PCGInstancePackerTests.cpp
+++ b/Engine/Plugins/PCG/Source/PCG/Private/Tests/Metadata/PCGInstancePackerTests.cpp
@@ -260,13 +260,16 @@ bool FPCGInstancePackerTest_PointsOffset::RunTest(const FString& Parameters)
 		FPCGAttributePropertyInputSelector::CreatePointPropertySelector<FPCGAttributePropertyInputSelector>(EPCGPointProperties::Position, NAME_None, {TEXT("X")}), // $Position.X
 	};
 
+	// Should write at offset 0 and at offset 3
+	constexpr int32 Offsets[2] = {0, 3};
+
 	PCGInstanceDataPackerBase::FPackedDataParams Params =
 	{
 		.CommonParams =
 		{
 			.NumInstances = NumPoints,
 			.OutPackedCustomData = &ActualCustomData,
-			.OptionalOffsets = {0, 3} // Should write at offset 0 and at offset 3
+			.OptionalOffsets = MakeConstArrayView(Offsets)
 		},
 		.InData = PointData,
 		.Selectors = MakeConstArrayView(Selectors)
@@ -293,13 +296,16 @@ bool FPCGInstancePackerTest_InvalidOffset_NotEnoughFloats::RunTest(const FString
 		FPCGAttributePropertyInputSelector::CreateAttributeSelector<FPCGAttributePropertyInputSelector>(Int32AttributeName)
 	};
 
+	// Offset too big
+	constexpr int32 Offsets[1] = {5};
+	
 	PCGInstanceDataPackerBase::FPackedDataParams Params =
 	{
 		.CommonParams =
 		{
 			.NumInstances = NumPoints,
 			.OutPackedCustomData = &ActualCustomData,
-			.OptionalOffsets = {5} // Offset too big
+			.OptionalOffsets = MakeConstArrayView(Offsets)
 		},
 		.InData = PointData,
 		.Selectors = MakeConstArrayView(Selectors)
@@ -329,13 +335,16 @@ bool FPCGInstancePackerTest_InvalidOffset_Overlapping::RunTest(const FString& Pa
 		FPCGAttributePropertyInputSelector::CreateAttributeSelector<FPCGAttributePropertyInputSelector>(FloatAttributeName)
 	};
 
+	// Overlapping offset, the second one should be higher or equal to 3
+	constexpr int32 Offsets[2] = {0, 1};
+
 	PCGInstanceDataPackerBase::FPackedDataParams Params =
 	{
 		.CommonParams =
 		{
 			.NumInstances = NumPoints,
 			.OutPackedCustomData = &ActualCustomData,
-			.OptionalOffsets = {0, 1} // Overlapping offset, the second one should be higher or equal to 3
+			.OptionalOffsets = MakeConstArrayView(Offsets)
 		},
 		.InData = PointData,
 		.Selectors = MakeConstArrayView(Selectors)
diff --git a/Engine/Plugins/PCG/Source/PCG/Public/InstanceDataPackers/PCGInstanceDataPackerBase.h b/Engine/Plugins/PCG/Source/PCG/Public/InstanceDataPackers/PCGInstanceDataPackerBase.h
--- a/Engine/Plugins/PCG/Source/PCG/Public/InstanceDataPackers/PCGInstanceDataPackerBase.h
+++ b/Engine/Plugins/PCG/Source/PCG/Public/InstanceDataPackers/PCGInstanceDataPackerBase.h
@@ -32,6 +32,7 @@ struct FPCGPackedCustomData
 	TBitArray<> Mask;
 
 	PCG_API TConstArrayView<float> GetViewForIndex(int32 Index) const;
+	PCG_API void Reset();
 };
 
 namespace PCGInstanceDataPackerBase
