[PCG] Support of Custom Primitive Data on the Static Mesh Spawner for Primitive selector

* Added CPU path support for Primitive Selector
* Added CPU/GPU set of Custom Primitive Data (ISM, HISM, PISM and FastGeo)
* Made sure Primitive Table and custom primitive data are taken into account for ISMC reuse

#rnx

diff --git a/Engine/Plugins/Experimental/FastGeoStreaming/Source/FastGeoStreaming/Internal/FastGeoPrimitiveComponent.h b/Engine/Plugins/Experimental/FastGeoStreaming/Source/FastGeoStreaming/Internal/FastGeoPrimitiveComponent.h
--- a/Engine/Plugins/Experimental/FastGeoStreaming/Source/FastGeoStreaming/Internal/FastGeoPrimitiveComponent.h
+++ b/Engine/Plugins/Experimental/FastGeoStreaming/Source/FastGeoStreaming/Internal/FastGeoPrimitiveComponent.h
@@ -88,6 +88,7 @@ public:
 	bool CheckPSOPrecachingAndBoostPriority(EPSOPrecachePriority NewPSOPrecachePriority = EPSOPrecachePriority::High);
 	bool IsPSOPrecaching() const;
 	void SetCollisionEnabled(bool bInCollisionEnabled);
+	void SetCustomPrimitiveData(TConstArrayView<float> InCustomPrimitiveData);
 	void UpdateVisibility();
 	EComponentMobility::Type GetMobility() const;
 
diff --git a/Engine/Plugins/Experimental/FastGeoStreaming/Source/FastGeoStreaming/Private/FastGeoPrimitiveComponent.cpp b/Engine/Plugins/Experimental/FastGeoStreaming/Source/FastGeoStreaming/Private/FastGeoPrimitiveComponent.cpp
--- a/Engine/Plugins/Experimental/FastGeoStreaming/Source/FastGeoStreaming/Private/FastGeoPrimitiveComponent.cpp
+++ b/Engine/Plugins/Experimental/FastGeoStreaming/Source/FastGeoStreaming/Private/FastGeoPrimitiveComponent.cpp
@@ -608,6 +608,13 @@ void FFastGeoPrimitiveComponent::SetCollisionEnabled(bool bInCollisionEnabled)
 	SceneProxyDesc.bCollisionEnabled = bInCollisionEnabled;
 }
 
+void FFastGeoPrimitiveComponent::SetCustomPrimitiveData(TConstArrayView<float> InCustomPrimitiveData)
+{
+	// Make sure to clamp to the maximum number of values supported by custom primitive data.
+	CustomPrimitiveData.Data = InCustomPrimitiveData.Left(FMath::Min(InCustomPrimitiveData.Num(), FCustomPrimitiveData::NumCustomPrimitiveDataFloats));
+	MarkRenderStateDirty();
+}
+
 void FFastGeoPrimitiveComponent::InitializeDynamicProperties()
 {
 	Super::InitializeDynamicProperties();
diff --git a/Engine/Plugins/Experimental/PCGInterops/PCGFastGeoInterop/Source/PCGFastGeoInterop/Private/Compute/PrimitiveFactories/PCGPrimitiveFactoryFastGeoPISMC.cpp b/Engine/Plugins/Experimental/PCGInterops/PCGFastGeoInterop/Source/PCGFastGeoInterop/Private/Compute/PrimitiveFactories/PCGPrimitiveFactoryFastGeoPISMC.cpp
--- a/Engine/Plugins/Experimental/PCGInterops/PCGFastGeoInterop/Source/PCGFastGeoInterop/Private/Compute/PrimitiveFactories/PCGPrimitiveFactoryFastGeoPISMC.cpp
+++ b/Engine/Plugins/Experimental/PCGInterops/PCGFastGeoInterop/Source/PCGFastGeoInterop/Private/Compute/PrimitiveFactories/PCGPrimitiveFactoryFastGeoPISMC.cpp
@@ -71,7 +71,8 @@ int32 FPCGPrimitiveFactoryFastGeoPISMC::GetNumInstances(int32 InPrimitiveIndex)
 
 void FPCGPrimitiveFactoryFastGeoPISMC::Initialize(FParameters&& InParameters)
 {
-	Descriptors = InParameters.Descriptors;
+	Descriptors = MoveTemp(InParameters.Descriptors);
+	CustomPrimitiveData = MoveTemp(InParameters.CustomPrimitiveData);
 }
 
 bool FPCGPrimitiveFactoryFastGeoPISMC::Create(FPCGContext* InContext)
@@ -113,6 +114,7 @@ bool FPCGPrimitiveFactoryFastGeoPISMC::Create(FPCGContext* InContext)
 
 		FFastGeoProceduralISMComponent& FastGeoComponent = static_cast<FFastGeoProceduralISMComponent&>(ComponentCluster->AddComponent(FFastGeoProceduralISMComponent::Type));
 		FastGeoComponent.InitializeFromComponentDescriptor(Desc);
+		FastGeoComponent.SetCustomPrimitiveData(CustomPrimitiveData.GetViewForIndex(DescIndex));
 
 		MeshBounds[DescIndex] = Desc.StaticMesh->GetBoundingBox();
 	}
diff --git a/Engine/Plugins/Experimental/PCGInterops/PCGFastGeoInterop/Source/PCGFastGeoInterop/Private/Compute/PrimitiveFactories/PCGPrimitiveFactoryFastGeoPISMC.h b/Engine/Plugins/Experimental/PCGInterops/PCGFastGeoInterop/Source/PCGFastGeoInterop/Private/Compute/PrimitiveFactories/PCGPrimitiveFactoryFastGeoPISMC.h
--- a/Engine/Plugins/Experimental/PCGInterops/PCGFastGeoInterop/Source/PCGFastGeoInterop/Private/Compute/PrimitiveFactories/PCGPrimitiveFactoryFastGeoPISMC.h
+++ b/Engine/Plugins/Experimental/PCGInterops/PCGFastGeoInterop/Source/PCGFastGeoInterop/Private/Compute/PrimitiveFactories/PCGPrimitiveFactoryFastGeoPISMC.h
@@ -33,6 +33,7 @@ protected:
 
 protected:
 	TArray<FPCGProceduralISMComponentDescriptor> Descriptors;
+	FPCGPackedCustomData CustomPrimitiveData;
 
 	TArray<FWeakFastGeoComponent> Components;
 	TArray<int32> InstanceCounts;
diff --git a/Engine/Plugins/PCG/Source/PCG/Private/Components/PCGProceduralISMComponent.cpp b/Engine/Plugins/PCG/Source/PCG/Private/Components/PCGProceduralISMComponent.cpp
--- a/Engine/Plugins/PCG/Source/PCG/Private/Components/PCGProceduralISMComponent.cpp
+++ b/Engine/Plugins/PCG/Source/PCG/Private/Components/PCGProceduralISMComponent.cpp
@@ -17,11 +17,12 @@
 #include "PrimitiveSceneProxyDesc.h"
 #include "SceneInterface.h"
 #include "StaticMeshSceneProxyDesc.h"
+#include "Algo/Compare.h"
 #include "Engine/StaticMesh.h"
 
 #include UE_INLINE_GENERATED_CPP_BY_NAME(PCGProceduralISMComponent)
 
-UPCGManagedProceduralISMComponent* PCGManagedProceduralISMComponent::GetOrCreateManagedProceduralISMC(AActor* InTargetActor, UPCGComponent* InSourceComponent, uint64 InSettingsUID, const FPCGProceduralISMCBuilderParameters& InParams, FPCGContext* OptionalContext)
+UPCGManagedProceduralISMComponent* PCGManagedProceduralISMComponent::GetOrCreateManagedProceduralISMC(AActor* InTargetActor, UPCGComponent* InSourceComponent, uint64 InSettingsUID, FPCGProceduralISMCBuilderParameters&& InParams, FPCGContext* OptionalContext)
 {
 	check(InTargetActor && InSourceComponent);
 
@@ -33,7 +34,7 @@ UPCGManagedProceduralISMComponent* PCGManagedProceduralISMComponent::GetOrCreate
 
 	TRACE_CPUPROFILER_EVENT_SCOPE(UPCGActorHelpers::GetOrCreateManagedISMC);
 
-	FPCGProceduralISMComponentDescriptor Descriptor = InParams.Descriptor;
+	FPCGProceduralISMComponentDescriptor& Descriptor = InParams.Descriptor;
 
 	{
 		TRACE_CPUPROFILER_EVENT_SCOPE(UPCGActorHelpers::GetOrCreateManagedISMC::FindMatchingMISMC);
@@ -58,7 +59,8 @@ UPCGManagedProceduralISMComponent* PCGManagedProceduralISMComponent::GetOrCreate
 				{
 					if (IsValid(ISMC) &&
 						ISMC->GetOwner() == InTargetActor &&
-						Resource->GetDescriptor() == Descriptor)
+						Resource->GetDescriptor() == Descriptor &&
+						Algo::Compare(ISMC->GetDefaultCustomPrimitiveData().Data, InParams.CustomPrimitiveData, [](const float LHS, const float RHS) { return FMath::IsNearlyEqual(LHS, RHS); }))
 					{
 						MatchingResource = Resource;
 					}
@@ -97,10 +99,12 @@ UPCGManagedProceduralISMComponent* PCGManagedProceduralISMComponent::GetOrCreate
 		PCGLog::Component::LogComponentAttachmentFailedWarning(OptionalContext);
 	}
 
+	ISMC->SetCustomPrimitiveDataFloatArray(0, InParams.CustomPrimitiveData);
+
 	// Create managed resource on source component
 	UPCGManagedProceduralISMComponent* Resource = NewObject<UPCGManagedProceduralISMComponent>(InSourceComponent);
 	Resource->SetComponent(ISMC);
-	Resource->SetDescriptor(Descriptor);
+	Resource->SetDescriptor(MoveTemp(Descriptor));
 	if (InTargetActor->GetRootComponent())
 	{
 		Resource->SetRootLocation(InTargetActor->GetRootComponent()->GetComponentLocation());
@@ -514,6 +518,11 @@ void UPCGManagedProceduralISMComponent::SetDescriptor(const FPCGProceduralISMCom
 	Descriptor = InDescriptor;
 }
 
+void UPCGManagedProceduralISMComponent::SetDescriptor(FPCGProceduralISMComponentDescriptor&& InDescriptor)
+{
+	Descriptor = MoveTemp(InDescriptor);
+}
+
 bool UPCGManagedProceduralISMComponent::ReleaseIfUnused(TSet<TSoftObjectPtr<AActor>>& OutActorsToDelete)
 {
 	if (Super::ReleaseIfUnused(OutActorsToDelete) || !GetComponent())
diff --git a/Engine/Plugins/PCG/Source/PCG/Private/Components/PCGProceduralISMComponent.h b/Engine/Plugins/PCG/Source/PCG/Private/Components/PCGProceduralISMComponent.h
--- a/Engine/Plugins/PCG/Source/PCG/Private/Components/PCGProceduralISMComponent.h
+++ b/Engine/Plugins/PCG/Source/PCG/Private/Components/PCGProceduralISMComponent.h
@@ -12,6 +12,7 @@
 #include "InstanceDataSceneProxy.h"
 #include "StaticMeshResources.h"
 #include "StaticMeshSceneProxy.h"
+#include "Algo/Compare.h"
 #include "Misc/Guid.h"
 #include "Stats/Stats.h"
 #include "UObject/ObjectMacros.h"
@@ -184,12 +185,21 @@ struct FPCGProceduralISMCBuilderParameters
 {
 	friend inline uint32 GetTypeHash(const FPCGProceduralISMCBuilderParameters& Key)
 	{
-		return GetTypeHash(Key.Descriptor);
+		uint32 Hash = GetTypeHash(Key.Descriptor);
+		Hash = HashCombine(Hash, Key.bAllowDescriptorChanges ? 2 : 1);
+		Hash = GetArrayHash(Key.CustomPrimitiveData.GetData(), Key.CustomPrimitiveData.Num(), Hash);
+		return Hash;
 	}
 
-	inline bool operator==(const FPCGProceduralISMCBuilderParameters& Other) const { return Descriptor == Other.Descriptor && bAllowDescriptorChanges == Other.bAllowDescriptorChanges; }
+	inline bool operator==(const FPCGProceduralISMCBuilderParameters& Other) const
+	{
+		return Descriptor == Other.Descriptor
+		&& bAllowDescriptorChanges == Other.bAllowDescriptorChanges
+		&& Algo::Compare(CustomPrimitiveData, Other.CustomPrimitiveData);
+	}
 
 	FPCGProceduralISMComponentDescriptor Descriptor;
+	TConstArrayView<float> CustomPrimitiveData;
 
 	bool bAllowDescriptorChanges = true;
 };
@@ -219,6 +229,8 @@ public:
 	UE_API void SetComponent(UPCGProceduralISMComponent* InComponent);
 
 	UE_API void SetDescriptor(const FPCGProceduralISMComponentDescriptor& InDescriptor);
+	UE_API void SetDescriptor(FPCGProceduralISMComponentDescriptor&& InDescriptor);
+	
 	const FPCGProceduralISMComponentDescriptor& GetDescriptor() const { return Descriptor; }
 
 	UE_API void SetRootLocation(const FVector& InRootLocation);
@@ -242,7 +254,7 @@ protected:
 
 namespace PCGManagedProceduralISMComponent
 {
-	UPCGManagedProceduralISMComponent* GetOrCreateManagedProceduralISMC(AActor* InTargetActor, UPCGComponent* InSourceComponent, uint64 InSettingsUID, const FPCGProceduralISMCBuilderParameters& InParams, FPCGContext* OptionalContext = nullptr);
+	UPCGManagedProceduralISMComponent* GetOrCreateManagedProceduralISMC(AActor* InTargetActor, UPCGComponent* InSourceComponent, uint64 InSettingsUID, FPCGProceduralISMCBuilderParameters&& InParams, FPCGContext* OptionalContext = nullptr);
 }
 
 #undef UE_API
diff --git a/Engine/Plugins/PCG/Source/PCG/Private/Compute/DataInterfaces/Elements/PCGStaticMeshSpawnerDataInterface.cpp b/Engine/Plugins/PCG/Source/PCG/Private/Compute/DataInterfaces/Elements/PCGStaticMeshSpawnerDataInterface.cpp
--- a/Engine/Plugins/PCG/Source/PCG/Private/Compute/DataInterfaces/Elements/PCGStaticMeshSpawnerDataInterface.cpp
+++ b/Engine/Plugins/PCG/Source/PCG/Private/Compute/DataInterfaces/Elements/PCGStaticMeshSpawnerDataInterface.cpp
@@ -554,6 +554,7 @@ void UPCGStaticMeshSpawnerDataProvider::Reset()
 	bRegisteredPrimitives = false;
 	bStaticMeshesLoaded = false;
 	LoadHandle.Reset();
+	CustomPrimitiveData.Reset();
 }
 
 void UPCGStaticMeshSpawnerDataProvider::CreatePrimitiveDescriptors(FPCGContext* InContext, UPCGDataBinding* InBinding)
@@ -763,12 +764,14 @@ void UPCGStaticMeshSpawnerDataProvider::CreatePrimitiveDescriptors(FPCGContext*
 		FPCGMeshMaterialOverrideHelper MaterialOverrideHelper;
 		MaterialOverrideHelper.Initialize(*InContext, /*bInByAttributeOverride=*/true, SelectorPrimitiveData->MaterialOverrideAttributes, PrimitiveMetadata);
 
+		SelectorPrimitiveData->PackCustomPrimitiveData(PrimitiveData, PrimitiveIdToInstanceCount.Num(), CustomPrimitiveData, InContext);
+
 		const int32 NumPrimitiveDataElements = PrimitiveMetadata->GetItemCountForChild();
 
 		PrimitiveDescriptors.Reserve(PrimitiveIdToInstanceCount.Num());
 		PrimitiveIds.Reserve(PrimitiveIdToInstanceCount.Num());
 		PrimitiveStringKeys.Reserve(PrimitiveIdToInstanceCount.Num());
-
+		
 		for (const TPair<int32, uint32>& PrimitiveIdAndInstanceCount : PrimitiveIdToInstanceCount)
 		{
 			if (PrimitiveIdAndInstanceCount.Value == 0)
@@ -972,7 +975,8 @@ bool UPCGStaticMeshSpawnerDataProvider::SetupPrimitives(FPCGContext* InContext,
 		Primitives->PrimitiveFactory = Factory;
 
 		IPCGPrimitiveFactoryISMBase::FParameters Params;
-		Params.Descriptors = PrimitiveDescriptors;
+		Params.Descriptors = MoveTemp(PrimitiveDescriptors);
+		Params.CustomPrimitiveData = MoveTemp(CustomPrimitiveData);
 		Params.TargetActor = Settings->TargetActor.Get() ? Settings->TargetActor.Get() : InContext->GetTargetActor(nullptr);
 
 		Factory->Initialize(MoveTemp(Params));
diff --git a/Engine/Plugins/PCG/Source/PCG/Private/Compute/DataInterfaces/Elements/PCGStaticMeshSpawnerDataInterface.h b/Engine/Plugins/PCG/Source/PCG/Private/Compute/DataInterfaces/Elements/PCGStaticMeshSpawnerDataInterface.h
--- a/Engine/Plugins/PCG/Source/PCG/Private/Compute/DataInterfaces/Elements/PCGStaticMeshSpawnerDataInterface.h
+++ b/Engine/Plugins/PCG/Source/PCG/Private/Compute/DataInterfaces/Elements/PCGStaticMeshSpawnerDataInterface.h
@@ -6,6 +6,7 @@
 #include "ComputeFramework/ComputeDataProvider.h"
 
 #include "Components/PCGProceduralISMComponentDescriptor.h"
+#include "InstanceDataPackers/PCGInstanceDataPackerBase.h"
 
 #include "PCGStaticMeshSpawnerDataInterface.generated.h"
 
@@ -89,6 +90,8 @@ public:
 
 	TArray<FPCGProceduralISMComponentDescriptor> PrimitiveDescriptors;
 
+	FPCGPackedCustomData CustomPrimitiveData;
+
 	uint32 CustomFloatCount = 0;
 
 	bool bRegisteredPrimitives = false;
diff --git a/Engine/Plugins/PCG/Source/PCG/Private/Compute/PrimitiveFactories/PCGPrimitiveFactoryPISMC.cpp b/Engine/Plugins/PCG/Source/PCG/Private/Compute/PrimitiveFactories/PCGPrimitiveFactoryPISMC.cpp
--- a/Engine/Plugins/PCG/Source/PCG/Private/Compute/PrimitiveFactories/PCGPrimitiveFactoryPISMC.cpp
+++ b/Engine/Plugins/PCG/Source/PCG/Private/Compute/PrimitiveFactories/PCGPrimitiveFactoryPISMC.cpp
@@ -45,7 +45,8 @@ bool FPCGPrimitiveFactoryPISMC::IsRenderStateCreated() const
 
 void FPCGPrimitiveFactoryPISMC::Initialize(FParameters&& InParameters)
 {
-	Descriptors = InParameters.Descriptors;
+	Descriptors = MoveTemp(InParameters.Descriptors);
+	CustomPrimitiveData = MoveTemp(InParameters.CustomPrimitiveData);
 	TargetActor = InParameters.TargetActor;
 }
 
@@ -62,24 +63,29 @@ bool FPCGPrimitiveFactoryPISMC::Create(FPCGContext* InContext)
 		return true;
 	}
 
-	FPCGProceduralISMCBuilderParameters Params;
-	Params.bAllowDescriptorChanges = false;
-
 	const int32 NumToCreate = FMath::Min(Descriptors.Num(), PCGComputeConstants::MAX_PRIMITIVE_COMPONENTS_PER_SPAWNER);
 
 	while (NumPrimitivesProcessed < NumToCreate)
 	{
-		const FPCGProceduralISMComponentDescriptor& Desc = Descriptors[NumPrimitivesProcessed];
+		FPCGProceduralISMComponentDescriptor& Desc = Descriptors[NumPrimitivesProcessed];
 
-		Params.Descriptor = Desc;
+		const int32 PrimNumInstances = Desc.NumInstances;
+		FBox BoundingBox = Desc.StaticMesh.IsValid() ? Desc.StaticMesh->GetBoundingBox() : FBox{};
 
-		UPCGManagedProceduralISMComponent* ManagedComponent = PCGManagedProceduralISMComponent::GetOrCreateManagedProceduralISMC(TargetActor, SourceComponent, /*InSettingsUID=*/0, Params);
+		FPCGProceduralISMCBuilderParameters Params =
+		{
+			.Descriptor = MoveTemp(Desc),
+			.CustomPrimitiveData = CustomPrimitiveData.GetViewForIndex(NumPrimitivesProcessed),
+			.bAllowDescriptorChanges = false,
+		};
+		
+		UPCGManagedProceduralISMComponent* ManagedComponent = PCGManagedProceduralISMComponent::GetOrCreateManagedProceduralISMC(TargetActor, SourceComponent, /*InSettingsUID=*/0, MoveTemp(Params));
 
 		if (ManagedComponent)
 		{
 			Components.Add(ManagedComponent->GetComponent());
-			NumInstances.Add(Desc.NumInstances);
-			MeshBounds.Add(Desc.StaticMesh->GetBoundingBox());
+			NumInstances.Add(PrimNumInstances);
+			MeshBounds.Add(MoveTemp(BoundingBox));
 		}
 
 		++NumPrimitivesProcessed;
diff --git a/Engine/Plugins/PCG/Source/PCG/Private/Compute/PrimitiveFactories/PCGPrimitiveFactoryPISMC.h b/Engine/Plugins/PCG/Source/PCG/Private/Compute/PrimitiveFactories/PCGPrimitiveFactoryPISMC.h
--- a/Engine/Plugins/PCG/Source/PCG/Private/Compute/PrimitiveFactories/PCGPrimitiveFactoryPISMC.h
+++ b/Engine/Plugins/PCG/Source/PCG/Private/Compute/PrimitiveFactories/PCGPrimitiveFactoryPISMC.h
@@ -27,6 +27,7 @@ public:
 
 protected:
 	TArray<FPCGProceduralISMComponentDescriptor> Descriptors;
+	FPCGPackedCustomData CustomPrimitiveData;
 	AActor* TargetActor = nullptr;
 
 	int32 NumPrimitivesProcessed = 0;
diff --git a/Engine/Plugins/PCG/Source/PCG/Private/Elements/PCGStaticMeshSpawner.cpp b/Engine/Plugins/PCG/Source/PCG/Private/Elements/PCGStaticMeshSpawner.cpp
--- a/Engine/Plugins/PCG/Source/PCG/Private/Elements/PCGStaticMeshSpawner.cpp
+++ b/Engine/Plugins/PCG/Source/PCG/Private/Elements/PCGStaticMeshSpawner.cpp
@@ -43,11 +43,6 @@ static TAutoConsoleVariable<bool> CVarUseRawDataForAnalysis(
 	ECVF_ReadOnly // Structural change, could flush graph compilation cache, but trickier to flush compute graph instance pool in UPCGGraph currently.
 	);
 
-namespace PCGStaticMeshSpawner
-{
-	static const FName PrimitiveTablePinLabel = TEXT("PrimitiveTable");
-}
-
 UPCGStaticMeshSpawnerSettings::UPCGStaticMeshSpawnerSettings(const FObjectInitializer &ObjectInitializer)
 {
 	MeshSelectorType = UPCGMeshSelectorWeighted::StaticClass();
@@ -254,6 +249,15 @@ bool FPCGStaticMeshSpawnerElement::PrepareDataInternal(FPCGContext* InContext) c
 
 	// perform mesh selection
 	TArray<FPCGTaggedData> Inputs = Context->InputData.GetInputsByPin(PCGPinConstants::DefaultInputLabel);
+	TArray<FPCGTaggedData> PrimitiveTableInputs = Context->InputData.GetInputsByPin(PCGStaticMeshSpawner::PrimitiveTablePinLabel);
+
+	if (PrimitiveTableInputs.Num() > 1)
+	{
+		PCGLog::InputOutput::LogFirstInputOnlyWarning(PCGStaticMeshSpawner::PrimitiveTablePinLabel, Context);
+	}
+
+	const FPCGTaggedData* PrimitiveTableInput = !PrimitiveTableInputs.IsEmpty() ? &PrimitiveTableInputs[0] : nullptr;
+
 	TArray<FPCGTaggedData>& Outputs = Context->OutputData.TaggedData;
 
 	while(Context->CurrentInputIndex < Inputs.Num())
@@ -300,6 +304,11 @@ bool FPCGStaticMeshSpawnerElement::PrepareDataInternal(FPCGContext* InContext) c
 				{
 					FPCGDataCollection SubCollection;
 					SubCollection.TaggedData.Add(Input);
+					if (PrimitiveTableInput)
+					{
+						SubCollection.TaggedData.Add(*PrimitiveTableInput);
+					}
+
 					SubCollection.ComputeCrcs(ShouldComputeFullOutputDataCrc(Context));
 
 					GetDependenciesCrc(FPCGGetDependenciesCrcParams(&SubCollection, Settings, Context->ExecutionSource.Get()), DataCrc);
@@ -607,6 +616,7 @@ void FPCGStaticMeshSpawnerElement::SpawnStaticMeshInstances(FPCGStaticMeshSpawne
 	FPCGISMComponentBuilderParams Params;
 	Params.Descriptor = InstanceList.Descriptor;
 	Params.NumCustomDataFloats = PackedCustomData.NumCustomDataFloats;
+	Params.CustomPrimitiveData = MakeConstArrayView(InstanceList.CustomPrimitiveData);
 
 	// If the root actor we're binding to is movable, then the ISMC should be movable by default
 	if (USceneComponent* SceneComponent = TargetActor->GetRootComponent())
diff --git a/Engine/Plugins/PCG/Source/PCG/Private/Helpers/PCGActorHelpers.cpp b/Engine/Plugins/PCG/Source/PCG/Private/Helpers/PCGActorHelpers.cpp
--- a/Engine/Plugins/PCG/Source/PCG/Private/Helpers/PCGActorHelpers.cpp
+++ b/Engine/Plugins/PCG/Source/PCG/Private/Helpers/PCGActorHelpers.cpp
@@ -170,7 +170,8 @@ UPCGManagedISMComponent* UPCGActorHelpers::GetOrCreateManagedISMC(AActor* InTarg
 							ISMC->GetOwner() == InTargetActor &&
 							ISMC->NumCustomDataFloats == InParams.NumCustomDataFloats &&
 							Resource->GetDescriptor() == Descriptor &&
-							ISMC->HasAnyFlags(RF_Transient) == InParams.bTransient)
+							ISMC->HasAnyFlags(RF_Transient) == InParams.bTransient &&
+							Algo::Compare(ISMC->GetDefaultCustomPrimitiveData().Data, InParams.CustomPrimitiveData, [](const float LHS, const float RHS) { return FMath::IsNearlyEqual(LHS, RHS); }))
 						{
 							MatchingResource = Resource;
 						}
@@ -230,7 +231,15 @@ UPCGManagedISMComponent* UPCGActorHelpers::GetOrCreateManagedISMC(AActor* InTarg
 	{
 		PCGLog::Component::LogComponentAttachmentFailedWarning(OptionalContext);
 	}
-	
+
+	if (InParams.CustomPrimitiveData.Num() > FCustomPrimitiveData::NumCustomPrimitiveDataFloats)
+	{
+		PCGLog::LogWarningOnGraph(FText::Format(NSLOCTEXT("PCGActorHelpers", "TooManyPrimitiveDataFloats", "Too many custom primitive data {0}, max is {1}. Will be truncated."),
+			InParams.CustomPrimitiveData.Num(), FCustomPrimitiveData::NumCustomPrimitiveDataFloats), OptionalContext);
+	}
+
+	ISMC->SetDefaultCustomPrimitiveDataFloatArray(0, InParams.CustomPrimitiveData);
+
 	// Implementation note: Because we've used the FISMComponentDescriptor here (Descriptor vs InParams.Descriptor) which takes care of the loading, we still need to apply tags manually.
 	AddTagsToComponent(ISMC);
 
diff --git a/Engine/Plugins/PCG/Source/PCG/Private/MeshSelectors/PCGMeshSelectorPrimitiveData.cpp b/Engine/Plugins/PCG/Source/PCG/Private/MeshSelectors/PCGMeshSelectorPrimitiveData.cpp
--- a/Engine/Plugins/PCG/Source/PCG/Private/MeshSelectors/PCGMeshSelectorPrimitiveData.cpp
+++ b/Engine/Plugins/PCG/Source/PCG/Private/MeshSelectors/PCGMeshSelectorPrimitiveData.cpp
@@ -2,12 +2,55 @@
 
 #include "MeshSelectors/PCGMeshSelectorPrimitiveData.h"
 
+#include "PCGParamData.h"
+#include "Elements/PCGStaticMeshSpawner.h"
 #include "Elements/PCGStaticMeshSpawnerContext.h"
 
+#include "Algo/Compare.h"
+
 #include UE_INLINE_GENERATED_CPP_BY_NAME(PCGMeshSelectorPrimitiveData)
 
 #define LOCTEXT_NAMESPACE "PCGMeshSelectorPrimitiveData"
 
+namespace PCGMeshSelectorPrimitiveData
+{
+	// Returns variation based on mesh, material overrides and reverse culling
+	FPCGMeshInstanceList& GetInstanceList(
+		TArray<FPCGMeshInstanceList>& InstanceLists,
+		const FPCGSoftISMComponentDescriptor& TemplateDescriptor,
+		TSoftObjectPtr<UStaticMesh> Mesh,
+		const TArray<TSoftObjectPtr<UMaterialInterface>>& MaterialOverrides,
+		bool bReverseCulling,
+		const UPCGBasePointData* InPointData,
+		FPCGPackedCustomData InCustomData,
+		const int AttributePartitionIndex)
+	{
+		TConstArrayView<float> PrimitiveCustomData = InCustomData.GetViewForIndex(AttributePartitionIndex);
+		
+		for (FPCGMeshInstanceList& InstanceList : InstanceLists)
+		{
+			if (InstanceList.Descriptor.StaticMesh == Mesh &&
+				InstanceList.Descriptor.bReverseCulling == bReverseCulling &&
+				InstanceList.Descriptor.OverrideMaterials == MaterialOverrides &&
+				InstanceList.AttributePartitionIndex == AttributePartitionIndex &&
+				Algo::Compare(InstanceList.CustomPrimitiveData, PrimitiveCustomData, [](const float LHS, const float RHS) { return FMath::IsNearlyEqual(LHS, RHS);}))
+			{
+				return InstanceList;
+			}
+		}
+
+		FPCGMeshInstanceList& NewInstanceList = InstanceLists.Emplace_GetRef(TemplateDescriptor);
+		NewInstanceList.Descriptor.StaticMesh = Mesh;
+		NewInstanceList.Descriptor.OverrideMaterials = MaterialOverrides;
+		NewInstanceList.Descriptor.bReverseCulling = bReverseCulling;
+		NewInstanceList.AttributePartitionIndex = AttributePartitionIndex;
+		NewInstanceList.PointData = InPointData;
+		NewInstanceList.CustomPrimitiveData.Append(PrimitiveCustomData);
+
+		return NewInstanceList;
+	}
+}
+
 bool UPCGMeshSelectorPrimitiveData::SelectMeshInstances(
 	FPCGStaticMeshSpawnerContext& Context,
 	const UPCGStaticMeshSpawnerSettings* Settings,
@@ -15,9 +58,145 @@ bool UPCGMeshSelectorPrimitiveData::SelectMeshInstances(
 	TArray<FPCGMeshInstanceList>& OutMeshInstances,
 	UPCGBasePointData* OutPointData) const
 {
-	// todo_pcg: Currently a stub until CPU support is implemented (possibly implemented via a match and set).
-	PCGE_LOG_C(Error, GraphAndLog, &Context, LOCTEXT("SelectorGPUOnly", "Primitive selection via primitive data is currently only supported for GPU execution."));
+	TRACE_CPUPROFILER_EVENT_SCOPE(UPCGMeshSelectorPrimitiveData::SelectInstances);
+
+	if (!InPointData || !InPointData->ConstMetadata())
+	{
+		PCGLog::InputOutput::LogInvalidInputDataError(&Context);
+		return true;
+	}
+
+	if (!InPointData->Metadata->HasAttribute(PrimitiveIndexAttribute))
+	{
+		PCGE_LOG_C(Error, GraphAndLog, &Context, FText::Format(LOCTEXT("AttributeNotInMetadata", "Attribute '{0}' is not in the metadata"), FText::FromName(PrimitiveIndexAttribute)));
+		return true;
+	}
+
+	// Read the data table and create MeshInstanceList out of it
+	const TArray<FPCGTaggedData> PrimitiveTableInput = Context.InputData.GetInputsByPin(PCGStaticMeshSpawner::PrimitiveTablePinLabel);
+	if (PrimitiveTableInput.IsEmpty())
+	{
+		return true;
+	}
+
+	if (PrimitiveTableInput.Num() > 1)
+	{
+		PCGLog::InputOutput::LogFirstInputOnlyWarning(PCGStaticMeshSpawner::PrimitiveTablePinLabel, &Context);
+	}
+
+	const UPCGParamData* PrimitiveTableData = Cast<const UPCGParamData>(PrimitiveTableInput[0].Data);
+	if (!PrimitiveTableData || !PrimitiveTableData->ConstMetadata())
+	{
+		PCGLog::InputOutput::LogInvalidInputDataError(&Context);
+		return true;
+	}
+
+	const bool bUseAttributeMaterialOverrides = !MaterialOverrideAttributes.IsEmpty();
+
+	const FPCGAttributePropertyInputSelector MeshSelector = FPCGAttributePropertyInputSelector::CreateAttributeSelector<FPCGAttributePropertyInputSelector>(MeshAttribute);
+	TArray<FSoftObjectPath> Meshes;
+
+	if (!PCGAttributeAccessorHelpers::ExtractAllValues<FSoftObjectPath>(PrimitiveTableData, MeshSelector, Meshes, &Context, EPCGAttributeAccessorFlags::AllowBroadcastAndConstructible))
+	{
+		return true;
+	}
+
+	if (Meshes.IsEmpty())
+	{
+		return true;
+	}
+
+	TArray<FPCGSoftISMComponentDescriptor> Descriptors;
+	TArray<TArray<float>> CustomPrimitiveData;
+	Descriptors.Reserve(Meshes.Num());
+	for (int32 i = 0; i < Meshes.Num(); ++i)
+	{
+		FPCGSoftISMComponentDescriptor& Descriptor = Descriptors.Emplace_GetRef();
+		Descriptor.StaticMesh = Meshes[i];
+
+		FPCGMeshMaterialOverrideHelper MaterialOverrideHelper;
+		MaterialOverrideHelper.Initialize(Context, bUseAttributeMaterialOverrides, Descriptor.OverrideMaterials, MaterialOverrideAttributes, PrimitiveTableData->ConstMetadata());
+		if (MaterialOverrideHelper.IsValid())
+		{
+			// Param data keys are between 0 and Num() - 1;
+			Descriptor.OverrideMaterials = MaterialOverrideHelper.GetMaterialOverrides(PCGMetadataEntryKey{i});	
+		}
+
+		FPCGObjectOverrides Overrides(&Descriptor);
+		Overrides.Initialize(PrimitiveOverrideAttributes, &Descriptor, PrimitiveTableData, &Context);
+		if (Overrides.IsValid())
+		{
+			Overrides.Apply(i);
+		}
+
+		Context.MeshToOutPoints.FindOrAdd(Descriptor.StaticMesh).FindOrAdd(OutPointData);
+	}
+
+	FPCGPackedCustomData CustomData;
+	PackCustomPrimitiveData(PrimitiveTableData, Meshes.Num(), CustomData, &Context);
+
+	// @todo_pcg: Time Slice, maybe partition is better?
+	const FPCGAttributePropertyInputSelector PrimitiveIndexSelector = FPCGAttributePropertyInputSelector::CreateAttributeSelector<FPCGAttributePropertyInputSelector>(PrimitiveIndexAttribute);
+	const TUniquePtr<const IPCGAttributeAccessor> PrimitiveIndexAccessor = PCGAttributeAccessorHelpers::CreateConstAccessor(InPointData, PrimitiveIndexSelector);
+	const TUniquePtr<const IPCGAttributeAccessorKeys> PrimitiveIndexKeys = PCGAttributeAccessorHelpers::CreateConstKeys(InPointData, PrimitiveIndexSelector);
+
+	if (!PrimitiveIndexAccessor || !PrimitiveIndexKeys)
+	{
+		PCGLog::Metadata::LogFailToCreateAccessorError(PrimitiveIndexSelector, &Context);
+		return true;
+	}
+
+	const TConstPCGValueRange<FTransform> InTransformRange = InPointData->GetConstTransformValueRange();
+
+	TArray<TPair</*MeshInstanceIndex*/int32, /*ReverseMeshInstanceIndex*/int32>> PrimitiveToMeshInstanceMapping;
+	PrimitiveToMeshInstanceMapping.Init({INDEX_NONE, INDEX_NONE}, Descriptors.Num());
+
+	PCGMetadataElementCommon::ApplyOnAccessor<int32>(*PrimitiveIndexKeys, *PrimitiveIndexAccessor, [&InTransformRange, &OutMeshInstances, &Descriptors, &PrimitiveToMeshInstanceMapping, &Context, bUseAttributeMaterialOverrides, InPointData, OutPointData, &CustomData](int32 PrimitiveIndex, int32 PointIndex)
+	{
+		if (!Descriptors.IsValidIndex(PrimitiveIndex))
+		{
+			return;
+		}
+
+		const FPCGSoftISMComponentDescriptor& Descriptor = Descriptors[PrimitiveIndex];
+		const FTransform& InTransform = InTransformRange[PointIndex];
+		const bool bReverseInstance = InTransform.GetDeterminant() < 0;
+		int32& MeshInstanceIndex = bReverseInstance ? PrimitiveToMeshInstanceMapping[PrimitiveIndex].Value : PrimitiveToMeshInstanceMapping[PrimitiveIndex].Key;
+		if (MeshInstanceIndex == INDEX_NONE)
+		{
+			FPCGMeshInstanceList& InstanceList = PCGMeshSelectorPrimitiveData::GetInstanceList(OutMeshInstances, Descriptor, Descriptor.StaticMesh, Descriptor.OverrideMaterials, bReverseInstance, InPointData, CustomData, PrimitiveIndex);
+			MeshInstanceIndex = static_cast<uint32>(&InstanceList - OutMeshInstances.GetData());
+		}
+
+		check(MeshInstanceIndex != INDEX_NONE);
+		FPCGMeshInstanceList& InstanceList = OutMeshInstances[MeshInstanceIndex];
+		Context.MeshToOutPoints[Descriptor.StaticMesh][OutPointData].Add(PointIndex);
+		InstanceList.Instances.Emplace(InTransform);
+		InstanceList.InstancesIndices.Emplace(PointIndex);
+
+	}, EPCGAttributeAccessorFlags::AllowBroadcastAndConstructible);
+
 	return true;
 }
 
+void UPCGMeshSelectorPrimitiveData::PackCustomPrimitiveData(const UPCGData* InputData, int32 NumInstances, FPCGPackedCustomData& OutCustomPrimitiveData, FPCGContext* OptionalContext) const
+{
+	if (!CustomPrimitiveDataAttributes.IsEmpty())
+	{
+		PCGInstanceDataPackerBase::FPackedDataParams Params =
+		{
+			.CommonParams =
+			{
+				.NumInstances = NumInstances,
+				.OutPackedCustomData = &OutCustomPrimitiveData,
+				.OptionalContext = OptionalContext
+			},
+			.InData = InputData,
+			.Selectors = CustomPrimitiveDataAttributes
+		};
+
+		PCGInstanceDataPackerBase::PackCustomData(Params);
+	}
+}
+
 #undef LOCTEXT_NAMESPACE
diff --git a/Engine/Plugins/PCG/Source/PCG/Public/Compute/PrimitiveFactories/PCGPrimitiveFactoryISMBase.h b/Engine/Plugins/PCG/Source/PCG/Public/Compute/PrimitiveFactories/PCGPrimitiveFactoryISMBase.h
--- a/Engine/Plugins/PCG/Source/PCG/Public/Compute/PrimitiveFactories/PCGPrimitiveFactoryISMBase.h
+++ b/Engine/Plugins/PCG/Source/PCG/Public/Compute/PrimitiveFactories/PCGPrimitiveFactoryISMBase.h
@@ -4,7 +4,7 @@
 
 #include "Compute/PrimitiveFactories/IPCGRuntimePrimitiveFactory.h"
 
-#include "UObject/WeakObjectPtr.h"
+#include "InstanceDataPackers/PCGInstanceDataPackerBase.h"
 
 class AActor;
 struct FPCGContext;
@@ -16,6 +16,7 @@ public:
 	struct FParameters
 	{
 		TArray<FPCGProceduralISMComponentDescriptor> Descriptors;
+		FPCGPackedCustomData CustomPrimitiveData;
 		AActor* TargetActor = nullptr; // Optional, unless creating actor primitive components.
 	};
 
diff --git a/Engine/Plugins/PCG/Source/PCG/Public/Elements/PCGStaticMeshSpawner.h b/Engine/Plugins/PCG/Source/PCG/Public/Elements/PCGStaticMeshSpawner.h
--- a/Engine/Plugins/PCG/Source/PCG/Public/Elements/PCGStaticMeshSpawner.h
+++ b/Engine/Plugins/PCG/Source/PCG/Public/Elements/PCGStaticMeshSpawner.h
@@ -18,6 +18,11 @@ class UPCGSpatialData;
 struct FPCGContext;
 class UStaticMesh;
 
+namespace PCGStaticMeshSpawner
+{
+	static constexpr FLazyName PrimitiveTablePinLabel = "PrimitiveTable";
+}
+
 UCLASS(MinimalAPI, BlueprintType, ClassGroup = (Procedural), meta = (PrioritizeCategories = "Settings"))
 class UPCGStaticMeshSpawnerSettings : public UPCGSettings
 {
diff --git a/Engine/Plugins/PCG/Source/PCG/Public/Helpers/PCGActorHelpers.h b/Engine/Plugins/PCG/Source/PCG/Public/Helpers/PCGActorHelpers.h
--- a/Engine/Plugins/PCG/Source/PCG/Public/Helpers/PCGActorHelpers.h
+++ b/Engine/Plugins/PCG/Source/PCG/Public/Helpers/PCGActorHelpers.h
@@ -3,6 +3,7 @@
 #pragma once
 
 #include "PCGCrc.h"
+#include "Algo/Compare.h"
 #include "Elements/PCGSplineMeshParams.h"
 #include "MeshSelectors/PCGISMDescriptor.h"
 #include "MeshSelectors/PCGSkinnedMeshDescriptor.h"
@@ -56,6 +57,7 @@ struct FPCGISMComponentBuilderParams
 
 	FPCGSoftISMComponentDescriptor Descriptor;
 	int32 NumCustomDataFloats = 0;
+	TConstArrayView<float> CustomPrimitiveData;
 	FPCGCrc SettingsCrc;
 	FPCGCrc DataCrc;
 	bool bAllowDescriptorChanges = true;
@@ -70,10 +72,21 @@ PRAGMA_ENABLE_DEPRECATION_WARNINGS
 
 	friend inline uint32 GetTypeHash(const FPCGISMComponentBuilderParams& Key)
 	{
-		return HashCombine(HashCombine(GetTypeHash(Key.Descriptor), 1 + Key.NumCustomDataFloats), (Key.bAllowDescriptorChanges ? 2 : 1));
+		uint32 Result = GetTypeHash(Key.Descriptor);
+		Result = HashCombine(Result, 1 + Key.NumCustomDataFloats);
+		Result = HashCombine(Result, (Key.bAllowDescriptorChanges ? 2 : 1));
+		Result = GetArrayHash(Key.CustomPrimitiveData.GetData(), Key.CustomPrimitiveData.Num(), Result);
+
+		return Result;
 	}
 
-	inline bool operator==(const FPCGISMComponentBuilderParams& Other) const { return Descriptor == Other.Descriptor && NumCustomDataFloats == Other.NumCustomDataFloats && bAllowDescriptorChanges == Other.bAllowDescriptorChanges; }
+	inline bool operator==(const FPCGISMComponentBuilderParams& Other) const
+	{
+		return Descriptor == Other.Descriptor
+			&& NumCustomDataFloats == Other.NumCustomDataFloats
+			&& bAllowDescriptorChanges == Other.bAllowDescriptorChanges
+			&& Algo::Compare(CustomPrimitiveData, Other.CustomPrimitiveData);
+	}
 };
 
 struct FPCGSkinnedMeshComponentBuilderParams
diff --git a/Engine/Plugins/PCG/Source/PCG/Public/MeshSelectors/PCGMeshSelectorBase.h b/Engine/Plugins/PCG/Source/PCG/Public/MeshSelectors/PCGMeshSelectorBase.h
--- a/Engine/Plugins/PCG/Source/PCG/Public/MeshSelectors/PCGMeshSelectorBase.h
+++ b/Engine/Plugins/PCG/Source/PCG/Public/MeshSelectors/PCGMeshSelectorBase.h
@@ -56,6 +56,9 @@ struct FPCGMeshInstanceList
 	UPROPERTY(BlueprintReadWrite, EditAnywhere, Category = Settings)
 	TArray<int32> InstancesIndices;
 
+	UPROPERTY()
+	TArray<float> CustomPrimitiveData;
+
 #if WITH_EDITORONLY_DATA
 	UE_DEPRECATED(5.5, "Use PointData + InstanceIndices instead.")
 	UPROPERTY(meta = (DeprecatedProperty))
diff --git a/Engine/Plugins/PCG/Source/PCG/Public/MeshSelectors/PCGMeshSelectorPrimitiveData.h b/Engine/Plugins/PCG/Source/PCG/Public/MeshSelectors/PCGMeshSelectorPrimitiveData.h
--- a/Engine/Plugins/PCG/Source/PCG/Public/MeshSelectors/PCGMeshSelectorPrimitiveData.h
+++ b/Engine/Plugins/PCG/Source/PCG/Public/MeshSelectors/PCGMeshSelectorPrimitiveData.h
@@ -13,6 +13,7 @@
 class UPCGBasePointData;
 class UPCGStaticMeshSpawnerSettings;
 struct FPCGMeshInstanceList;
+struct FPCGPackedCustomData;
 struct FPCGStaticMeshSpawnerContext;
 
 /** Selector type that uses a table of primitive data (one row per primitive) and expects a table row index on the input points. */
@@ -29,6 +30,8 @@ public:
 		TArray<FPCGMeshInstanceList>& OutMeshInstances,
 		UPCGBasePointData* OutPointData) const override;
 
+	void PackCustomPrimitiveData(const UPCGData* InputData, int32 NumInstances, FPCGPackedCustomData& OutCustomPrimitiveData, FPCGContext* OptionalContext = nullptr) const;
+
 public:
 	/** Base/default settings for the spawned primitive before overrides are applied. */
 	UPROPERTY(EditAnywhere, Category = "Primitive Setup")
@@ -49,4 +52,7 @@ public:
 	/** Array of attributes in primitive data that drive properties on the spawned primitives. */
 	UPROPERTY(EditAnywhere, Category = "Primitive Setup")
 	TArray<FPCGObjectPropertyOverrideDescription> PrimitiveOverrideAttributes;
+
+	UPROPERTY(EditAnywhere, Category = "Primitive Setup")
+	TArray<FPCGAttributePropertyInputSelector> CustomPrimitiveDataAttributes;
 };
