[PCG] Pack data collection task.

This can be a very heavy operation and is generally unsuitable for GT. It was also responsible for consisently over-shooting the frame budget in our high performance test case.

[RN] [PCG] Optimized data packing for upload to GPU by moving it onto an async task.

diff --git a/Engine/Plugins/PCG/Source/PCG/Private/Compute/PCGDataUploadAdaptor.cpp b/Engine/Plugins/PCG/Source/PCG/Private/Compute/PCGDataUploadAdaptor.cpp
--- a/Engine/Plugins/PCG/Source/PCG/Private/Compute/PCGDataUploadAdaptor.cpp
+++ b/Engine/Plugins/PCG/Source/PCG/Private/Compute/PCGDataUploadAdaptor.cpp
@@ -2,12 +2,14 @@
 
 #include "Compute/PCGDataUploadAdaptor.h"
 
+#include "PCGContext.h"
 #include "PCGModule.h"
 #include "Compute/PCGComputeCommon.h"
 #include "Compute/PCGComputeGraph.h"
 #include "Compute/PCGDataBinding.h"
 #include "Compute/PCGKernelHelpers.h"
 #include "Compute/Data/PCGProxyForGPUData.h"
+#include "Compute/Elements/PCGComputeGraphElement.h"
 #include "Compute/Packing/PCGDataCollectionPacking.h"
 
 #include "RenderCaptureInterface.h"
@@ -151,51 +153,75 @@ bool FPCGDataUploadAdaptor::PrepareData_GameThread(const UPCGComputeKernel* InPr
 	}
 	else
 	{
-		const uint64 BufferSize = PCGDataCollectionPackingHelpers::ComputePackedSizeBytes(TargetDataCollectionDesc);
-		if (PCGComputeHelpers::IsBufferSizeTooLarge(BufferSize))
+		if (!PackDataCollectionTask.IsValid())
 		{
-			return true;
-		}
-
-		bool bAllDataReady = true;
-
-		for (FPCGTaggedData& TaggedData : Binding->GetInputDataCollectionMutable().TaggedData)
-		{
-			if (TaggedData.Pin != InputPinLabel)
+			const uint64 BufferSize = PCGDataCollectionPackingHelpers::ComputePackedSizeBytes(TargetDataCollectionDesc);
+			if (PCGComputeHelpers::IsBufferSizeTooLarge(BufferSize))
 			{
-				continue;
+				return true;
 			}
 
-			if (const UPCGProxyForGPUData* Proxy = Cast<UPCGProxyForGPUData>(TaggedData.Data))
+			bool bAllDataReady = true;
+
+			for (FPCGTaggedData& TaggedData : Binding->GetInputDataCollectionMutable().TaggedData)
 			{
+				if (TaggedData.Pin != InputPinLabel)
+				{
+					continue;
+				}
+
+				if (const UPCGProxyForGPUData* Proxy = Cast<UPCGProxyForGPUData>(TaggedData.Data))
+				{
 #if WITH_EDITOR
-				PCGComputeHelpers::NotifyGPUToCPUReadback(Binding, InProducerKernel, InProducerSettings);
+					PCGComputeHelpers::NotifyGPUToCPUReadback(Binding, InProducerKernel, InProducerSettings);
 #endif
 
-				UPCGProxyForGPUData::FReadbackResult Result = Proxy->GetCPUData(/*InContext=*/nullptr);
-
-				if (!Result.bComplete)
-				{
-					bAllDataReady = false;
+					UPCGProxyForGPUData::FReadbackResult Result = Proxy->GetCPUData(/*InContext=*/nullptr);
+
+					if (!Result.bComplete)
+					{
+						bAllDataReady = false;
+					}
+					else
+					{
+						TaggedData.Data = Result.TaggedData.Data;
+						TaggedData.Tags = MoveTemp(Result.TaggedData.Tags);
+					}
 				}
-				else
+			}
+
+			if (!bAllDataReady)
+			{
+				return false;
+			}
+
+			// Kick off async task to pack data collection as this can be a heavy operation.
+			PackDataCollectionTask = UE::Tasks::Launch(UE_SOURCE_LOCATION, [ContextHandle=Binding->GetContextHandle(), This=AsShared()]()
+			{
+				TRACE_CPUPROFILER_EVENT_SCOPE(PackDataCollectionTask);
+				LLM_SCOPE_BYTAG(PCG);
+
+				FPCGContext::FSharedContext<FPCGComputeGraphContext> SharedContext(ContextHandle);
+				FPCGComputeGraphContext* Context = SharedContext.Get();
+				if (!Context)
 				{
-					TaggedData.Data = Result.TaggedData.Data;
-					TaggedData.Tags = MoveTemp(Result.TaggedData.Tags);
+					return;
 				}
-			}
-		}
 
-		if (!bAllDataReady)
-		{
-			return false;
-		}
+				PCGDataCollectionPackingHelpers::PackDataCollection(Context->DataBinding->GetInputDataCollection(), This->TargetDataCollectionDesc, This->InputPinLabel, Context->DataBinding, This->PackedDataCollection);
+			});
+
+			if (!ensure(PackDataCollectionTask.IsValid()))
+			{
+				return true;
+			}
 
 #if WITH_EDITOR
-		PCGComputeHelpers::NotifyCPUToGPUUpload(Binding, InProducerKernel, InProducerSettings);
+			PCGComputeHelpers::NotifyCPUToGPUUpload(Binding, InProducerKernel, InProducerSettings);
 #endif
+		}
 
-		PCGDataCollectionPackingHelpers::PackDataCollection(Binding->GetInputDataCollection(), TargetDataCollectionDesc, InputPinLabel, Binding, PackedDataCollection);
+		return PackDataCollectionTask.IsCompleted();
 	}
 
 	return true;
diff --git a/Engine/Plugins/PCG/Source/PCG/Private/Compute/PCGDataUploadAdaptor.h b/Engine/Plugins/PCG/Source/PCG/Private/Compute/PCGDataUploadAdaptor.h
--- a/Engine/Plugins/PCG/Source/PCG/Private/Compute/PCGDataUploadAdaptor.h
+++ b/Engine/Plugins/PCG/Source/PCG/Private/Compute/PCGDataUploadAdaptor.h
@@ -14,7 +14,7 @@ class UPCGSettings;
 enum class EPCGExportMode : uint8;
 
 /** A helper to ingest incoming data objects and pass the required data to the data interface. */
-class FPCGDataUploadAdaptor
+class FPCGDataUploadAdaptor : public TSharedFromThis<FPCGDataUploadAdaptor, ESPMode::ThreadSafe>
 {
 public:
 	FPCGDataUploadAdaptor(UPCGDataBinding* InDataBinding, const TSharedPtr<const FPCGDataCollectionDesc> InTargetDataCollectionDesc, FName InInputPinLabel);
@@ -59,4 +59,6 @@ private:
 	FName InputPinLabel;
 
 	uint32 ExternalBufferSizeBytes = 0;
+
+	UE::Tasks::TTask<void> PackDataCollectionTask;
 };
