[PCG] IPCGGraphExecutionState::GetTarget() /  GetTypedTarget()
- Abstract execution target to allow non-actor/component targets
- Replace calls to GetTargetActor(nullptr) by FPCGContext::GetTypedExecutionTarget<AActor>()

#rnx

diff --git a/Engine/Plugins/PCG/Source/PCG/Private/Compute/DataInterfaces/Elements/PCGSkinnedMeshSpawnerDataInterface.cpp b/Engine/Plugins/PCG/Source/PCG/Private/Compute/DataInterfaces/Elements/PCGSkinnedMeshSpawnerDataInterface.cpp
--- a/Engine/Plugins/PCG/Source/PCG/Private/Compute/DataInterfaces/Elements/PCGSkinnedMeshSpawnerDataInterface.cpp
+++ b/Engine/Plugins/PCG/Source/PCG/Private/Compute/DataInterfaces/Elements/PCGSkinnedMeshSpawnerDataInterface.cpp
@@ -458,7 +458,7 @@ void UPCGSkinnedMeshSpawnerDataProvider::CreatePrimitiveDescriptors(FPCGContext*
 
 	const UPCGSkinnedMeshSpawnerSettings* Settings = CastChecked<UPCGSkinnedMeshSpawnerSettings>(GetProducerKernel()->GetSettings());
 
-	AActor* TargetActor = Settings->TargetActor.Get() ? Settings->TargetActor.Get() : InContext->GetTargetActor(nullptr);
+	AActor* TargetActor = Settings->TargetActor.Get() ? Settings->TargetActor.Get() : InContext->GetTypedExecutionTarget<AActor>();
 	if (!ensure(Settings->MeshSelectorParameters) || !ensure(TargetActor))
 	{
 		return;
@@ -639,7 +639,7 @@ bool UPCGSkinnedMeshSpawnerDataProvider::SetupPrimitives(FPCGContext* InContext,
 
 	const UPCGSkinnedMeshSpawnerSettings* Settings = CastChecked<UPCGSkinnedMeshSpawnerSettings>(GetProducerKernel()->GetSettings());
 
-	AActor* TargetActor = Settings->TargetActor.Get() ? Settings->TargetActor.Get() : InContext->GetTargetActor(nullptr);
+	AActor* TargetActor = Settings->TargetActor.Get() ? Settings->TargetActor.Get() : InContext->GetTypedExecutionTarget<AActor>();
 	if (!ensure(TargetActor))
 	{
 		return true;
diff --git a/Engine/Plugins/PCG/Source/PCG/Private/Compute/DataInterfaces/Elements/PCGStaticMeshSpawnerDataInterface.cpp b/Engine/Plugins/PCG/Source/PCG/Private/Compute/DataInterfaces/Elements/PCGStaticMeshSpawnerDataInterface.cpp
--- a/Engine/Plugins/PCG/Source/PCG/Private/Compute/DataInterfaces/Elements/PCGStaticMeshSpawnerDataInterface.cpp
+++ b/Engine/Plugins/PCG/Source/PCG/Private/Compute/DataInterfaces/Elements/PCGStaticMeshSpawnerDataInterface.cpp
@@ -977,7 +977,7 @@ bool UPCGStaticMeshSpawnerDataProvider::SetupPrimitives(FPCGContext* InContext,
 		IPCGPrimitiveFactoryISMBase::FParameters Params;
 		Params.Descriptors = MoveTemp(PrimitiveDescriptors);
 		Params.CustomPrimitiveData = MoveTemp(CustomPrimitiveData);
-		Params.TargetActor = Settings->TargetActor.Get() ? Settings->TargetActor.Get() : InContext->GetTargetActor(nullptr);
+		Params.TargetActor = Settings->TargetActor.Get() ? Settings->TargetActor.Get() : InContext->GetTypedExecutionTarget<AActor>();
 
 		Factory->Initialize(MoveTemp(Params));
 	}
diff --git a/Engine/Plugins/PCG/Source/PCG/Private/Compute/Elements/PCGComputeGraphElement.cpp b/Engine/Plugins/PCG/Source/PCG/Private/Compute/Elements/PCGComputeGraphElement.cpp
--- a/Engine/Plugins/PCG/Source/PCG/Private/Compute/Elements/PCGComputeGraphElement.cpp
+++ b/Engine/Plugins/PCG/Source/PCG/Private/Compute/Elements/PCGComputeGraphElement.cpp
@@ -803,7 +803,7 @@ void FPCGComputeGraphElement::ExecuteDebugDraw(FPCGComputeGraphContext* InContex
 		{
 			if (const IPCGDataVisualization* DataVis = DataVisRegistry.GetDataVisualization(DataToDisplay->GetClass()))
 			{
-				DataVis->ExecuteDebugDisplay(InContext, ProducerSettings, DataToDisplay, InContext->GetTargetActor(nullptr));
+				DataVis->ExecuteDebugDisplay(InContext, ProducerSettings, DataToDisplay, InContext->GetTypedExecutionTarget<AActor>());
 			}
 		}
 	}
diff --git a/Engine/Plugins/PCG/Source/PCG/Private/Elements/PCGCreateSpline.cpp b/Engine/Plugins/PCG/Source/PCG/Private/Elements/PCGCreateSpline.cpp
--- a/Engine/Plugins/PCG/Source/PCG/Private/Elements/PCGCreateSpline.cpp
+++ b/Engine/Plugins/PCG/Source/PCG/Private/Elements/PCGCreateSpline.cpp
@@ -87,7 +87,7 @@ bool FPCGCreateSplineElement::ExecuteInternal(FPCGContext* Context) const
 			continue;
 		}
 
-		AActor* TargetActor = Settings->TargetActor.Get() ? Settings->TargetActor.Get() : Context->GetTargetActor(nullptr);
+		AActor* TargetActor = Settings->TargetActor.Get() ? Settings->TargetActor.Get() : Context->GetTypedExecutionTarget<AActor>();
 		if (!TargetActor)
 		{
 			PCGE_LOG(Error, GraphAndLog, LOCTEXT("InvalidTargetActor", "Invalid target actor. Ensure TargetActor member is initialized when creating SpatialData."));
diff --git a/Engine/Plugins/PCG/Source/PCG/Private/Elements/PCGCreateTargetActor.cpp b/Engine/Plugins/PCG/Source/PCG/Private/Elements/PCGCreateTargetActor.cpp
--- a/Engine/Plugins/PCG/Source/PCG/Private/Elements/PCGCreateTargetActor.cpp
+++ b/Engine/Plugins/PCG/Source/PCG/Private/Elements/PCGCreateTargetActor.cpp
@@ -356,7 +356,7 @@ bool FPCGCreateTargetActorElement::ExecuteInternal(FPCGContext* Context) const
 		return true;
 	}
 
-	AActor* TargetActor = Settings->RootActor.Get() ? Settings->RootActor.Get() : Context->GetTargetActor(nullptr);
+	AActor* TargetActor = Settings->RootActor.Get() ? Settings->RootActor.Get() : Context->GetTypedExecutionTarget<AActor>();
 	if (!TargetActor)
 	{
 		PCGE_LOG(Error, GraphAndLog, LOCTEXT("InvalidTargetActor", "Invalid target actor"));
diff --git a/Engine/Plugins/PCG/Source/PCG/Private/Elements/PCGDebugElement.cpp b/Engine/Plugins/PCG/Source/PCG/Private/Elements/PCGDebugElement.cpp
--- a/Engine/Plugins/PCG/Source/PCG/Private/Elements/PCGDebugElement.cpp
+++ b/Engine/Plugins/PCG/Source/PCG/Private/Elements/PCGDebugElement.cpp
@@ -59,7 +59,7 @@ bool FPCGDebugElement::ExecuteInternal(FPCGContext* Context) const
 
 		if (!TargetActor)
 		{
-			TargetActor = Context->GetTargetActor(nullptr);
+			TargetActor = Context->GetTypedExecutionTarget<AActor>();
 		}
 
 		if (const IPCGDataVisualization* DataVis = DataVisRegistry.GetDataVisualization(Input.Data->GetClass()))
diff --git a/Engine/Plugins/PCG/Source/PCG/Private/Elements/PCGSkinnedMeshSpawner.cpp b/Engine/Plugins/PCG/Source/PCG/Private/Elements/PCGSkinnedMeshSpawner.cpp
--- a/Engine/Plugins/PCG/Source/PCG/Private/Elements/PCGSkinnedMeshSpawner.cpp
+++ b/Engine/Plugins/PCG/Source/PCG/Private/Elements/PCGSkinnedMeshSpawner.cpp
@@ -198,7 +198,7 @@ bool FPCGSkinnedMeshSpawnerElement::PrepareDataInternal(FPCGContext* InContext)
 				continue;
 			}
 
-			AActor* TargetActor = Settings->TargetActor.Get() ? Settings->TargetActor.Get() : Context->GetTargetActor(nullptr);
+			AActor* TargetActor = Settings->TargetActor.Get() ? Settings->TargetActor.Get() : Context->GetTypedExecutionTarget<AActor>();
 			if (!TargetActor)
 			{
 				PCGE_LOG(Error, GraphAndLog, LOCTEXT("InvalidTargetActor", "Invalid target actor. Ensure TargetActor member is initialized when creating SpatialData."));
@@ -370,7 +370,7 @@ bool FPCGSkinnedMeshSpawnerElement::ExecuteInternal(FPCGContext* InContext) cons
 	const bool bFinishedExecution = Context->MeshInstancesData.IsEmpty();
 	if (bFinishedExecution)
 	{
-		if (AActor* TargetActor = Settings->TargetActor.Get() ? Settings->TargetActor.Get() : Context->GetTargetActor(nullptr))
+		if (AActor* TargetActor = Settings->TargetActor.Get() ? Settings->TargetActor.Get() : Context->GetTypedExecutionTarget<AActor>())
 		{
 			for (UFunction* Function : PCGHelpers::FindUserFunctions(TargetActor->GetClass(), Settings->PostProcessFunctionNames, { UPCGFunctionPrototypes::GetPrototypeWithNoParams() }, Context))
 			{
diff --git a/Engine/Plugins/PCG/Source/PCG/Private/Elements/PCGSpawnActor.cpp b/Engine/Plugins/PCG/Source/PCG/Private/Elements/PCGSpawnActor.cpp
--- a/Engine/Plugins/PCG/Source/PCG/Private/Elements/PCGSpawnActor.cpp
+++ b/Engine/Plugins/PCG/Source/PCG/Private/Elements/PCGSpawnActor.cpp
@@ -553,7 +553,7 @@ bool FPCGSpawnActorElement::SpawnAndPrepareSubgraphs(FPCGSubgraphContext* Contex
 			continue;
 		}
 
-		AActor* TargetActor = Settings->RootActor.Get() ? Settings->RootActor.Get() : Context->GetTargetActor(nullptr);
+		AActor* TargetActor = Settings->RootActor.Get() ? Settings->RootActor.Get() : Context->GetTypedExecutionTarget<AActor>();
 
 		if (!TargetActor)
 		{
diff --git a/Engine/Plugins/PCG/Source/PCG/Private/Elements/PCGSpawnSpline.cpp b/Engine/Plugins/PCG/Source/PCG/Private/Elements/PCGSpawnSpline.cpp
--- a/Engine/Plugins/PCG/Source/PCG/Private/Elements/PCGSpawnSpline.cpp
+++ b/Engine/Plugins/PCG/Source/PCG/Private/Elements/PCGSpawnSpline.cpp
@@ -59,7 +59,7 @@ bool FPCGSpawnSplineElement::ExecuteInternal(FPCGContext* InContext) const
 		return true;
 	}
 
-	AActor* TargetActor = Settings->TargetActor.Get() ? Settings->TargetActor.Get() : InContext->GetTargetActor(nullptr);
+	AActor* TargetActor = Settings->TargetActor.Get() ? Settings->TargetActor.Get() : InContext->GetTypedExecutionTarget<AActor>();
 	if (!TargetActor)
 	{
 		PCGLog::LogErrorOnGraph(LOCTEXT("InvalidTargetActor", "Invalid target actor. Ensure TargetActor member is initialized when creating SpatialData."), InContext);
diff --git a/Engine/Plugins/PCG/Source/PCG/Private/Elements/PCGSpawnSplineMesh.cpp b/Engine/Plugins/PCG/Source/PCG/Private/Elements/PCGSpawnSplineMesh.cpp
--- a/Engine/Plugins/PCG/Source/PCG/Private/Elements/PCGSpawnSplineMesh.cpp
+++ b/Engine/Plugins/PCG/Source/PCG/Private/Elements/PCGSpawnSplineMesh.cpp
@@ -101,7 +101,7 @@ bool FPCGSpawnSplineMeshElement::PrepareDataInternal(FPCGContext* InContext) con
 
 		EPCGTimeSliceInitResult ExecResult = Context->InitializePerExecutionState([Settings](ContextType* Context, FPCGSpawnSplineMeshPerExecutionState& OutState)
 		{
-			OutState.TargetActor = Settings->TargetActor.Get() ? Settings->TargetActor.Get() : Context->GetTargetActor(nullptr);
+			OutState.TargetActor = Settings->TargetActor.Get() ? Settings->TargetActor.Get() : Context->GetTypedExecutionTarget<AActor>();
 
 			if (!OutState.TargetActor)
 			{
diff --git a/Engine/Plugins/PCG/Source/PCG/Private/Elements/PCGStaticMeshSpawner.cpp b/Engine/Plugins/PCG/Source/PCG/Private/Elements/PCGStaticMeshSpawner.cpp
--- a/Engine/Plugins/PCG/Source/PCG/Private/Elements/PCGStaticMeshSpawner.cpp
+++ b/Engine/Plugins/PCG/Source/PCG/Private/Elements/PCGStaticMeshSpawner.cpp
@@ -282,7 +282,7 @@ bool FPCGStaticMeshSpawnerElement::PrepareDataInternal(FPCGContext* InContext) c
 				continue;
 			}
 
-			AActor* TargetActor = Settings->TargetActor.Get() ? Settings->TargetActor.Get() : Context->GetTargetActor(nullptr);
+			AActor* TargetActor = Settings->TargetActor.Get() ? Settings->TargetActor.Get() : Context->GetTypedExecutionTarget<AActor>();
 			if (!TargetActor)
 			{
 				PCGE_LOG(Error, GraphAndLog, LOCTEXT("InvalidTargetActor", "Invalid target actor. Ensure TargetActor member is initialized when creating SpatialData."));
@@ -535,7 +535,7 @@ bool FPCGStaticMeshSpawnerElement::ExecuteInternal(FPCGContext* InContext) const
 	const bool bFinishedExecution = Context->MeshInstancesData.IsEmpty();
 	if (bFinishedExecution)
 	{
-		if (AActor* TargetActor = Settings->TargetActor.Get() ? Settings->TargetActor.Get() : Context->GetTargetActor(nullptr))
+		if (AActor* TargetActor = Settings->TargetActor.Get() ? Settings->TargetActor.Get() : Context->GetTypedExecutionTarget<AActor>())
 		{
 			for (UFunction* Function : PCGHelpers::FindUserFunctions(TargetActor->GetClass(), Settings->PostProcessFunctionNames, { UPCGFunctionPrototypes::GetPrototypeWithNoParams() }, Context))
 			{
diff --git a/Engine/Plugins/PCG/Source/PCG/Private/Elements/PCGVisualizeAttribute.cpp b/Engine/Plugins/PCG/Source/PCG/Private/Elements/PCGVisualizeAttribute.cpp
--- a/Engine/Plugins/PCG/Source/PCG/Private/Elements/PCGVisualizeAttribute.cpp
+++ b/Engine/Plugins/PCG/Source/PCG/Private/Elements/PCGVisualizeAttribute.cpp
@@ -93,7 +93,7 @@ bool FPCGVisualizeAttribute::ExecuteInternal(FPCGContext* Context) const
 
 		UPCGDebugDrawComponent* DebugDrawComponent = nullptr;
 
-		AActor* TargetActor = Context->GetTargetActor(PointData);
+		AActor* TargetActor = Context->GetTypedExecutionTarget<AActor>();
 		if (TargetActor)
 		{
 			// Grab it if it exists already
diff --git a/Engine/Plugins/PCG/Source/PCG/Private/PCGComponentExecutionState.cpp b/Engine/Plugins/PCG/Source/PCG/Private/PCGComponentExecutionState.cpp
--- a/Engine/Plugins/PCG/Source/PCG/Private/PCGComponentExecutionState.cpp
+++ b/Engine/Plugins/PCG/Source/PCG/Private/PCGComponentExecutionState.cpp
@@ -49,6 +49,12 @@ UWorld* FPCGComponentExecutionState::GetWorld() const
 	return Component->GetWorld();
 }
 
+UObject* FPCGComponentExecutionState::GetTarget() const 
+{
+	check(Component);
+	return Component->GetOwner();
+}
+
 bool FPCGComponentExecutionState::HasAuthority() const
 {
 	check(Component && Component->GetOwner());
diff --git a/Engine/Plugins/PCG/Source/PCG/Private/PCGContext.cpp b/Engine/Plugins/PCG/Source/PCG/Private/PCGContext.cpp
--- a/Engine/Plugins/PCG/Source/PCG/Private/PCGContext.cpp
+++ b/Engine/Plugins/PCG/Source/PCG/Private/PCGContext.cpp
@@ -157,9 +157,17 @@ AActor* FPCGContext::GetTargetActor(const UPCGSpatialData* InSpatialData) const
 	{
 		return InSpatialData->TargetActor.Get();
 	}
-	else if (UPCGComponent* PCGComponent = Cast<UPCGComponent>(ExecutionSource.Get()))
+	else
+	{
+		return GetTypedExecutionTarget<AActor>();
+	}
+}
+
+UObject* FPCGContext::GetExecutionTarget() const
+{
+	if (IPCGGraphExecutionSource* ExecutionSourcePtr = ExecutionSource.Get())
 	{
-		return PCGComponent->GetOwner();
+		return ExecutionSourcePtr->GetExecutionState().GetTarget();
 	}
 	else
 	{
diff --git a/Engine/Plugins/PCG/Source/PCG/Private/PCGElement.cpp b/Engine/Plugins/PCG/Source/PCG/Private/PCGElement.cpp
--- a/Engine/Plugins/PCG/Source/PCG/Private/PCGElement.cpp
+++ b/Engine/Plugins/PCG/Source/PCG/Private/PCGElement.cpp
@@ -995,7 +995,7 @@ bool IPCGElement::DebugDisplay(FPCGContext* Context) const
 	{
 		if (const IPCGDataVisualization* DataVis = DataVisRegistry.GetDataVisualization(Data.Data->GetClass()))
 		{
-			DataVis->ExecuteDebugDisplay(Context, SettingsInterface, Data.Data, Context->GetTargetActor(nullptr));
+			DataVis->ExecuteDebugDisplay(Context, SettingsInterface, Data.Data, Context->GetTypedExecutionTarget<AActor>());
 		}
 	}
 
diff --git a/Engine/Plugins/PCG/Source/PCG/Public/PCGComponentExecutionState.h b/Engine/Plugins/PCG/Source/PCG/Public/PCGComponentExecutionState.h
--- a/Engine/Plugins/PCG/Source/PCG/Public/PCGComponentExecutionState.h
+++ b/Engine/Plugins/PCG/Source/PCG/Public/PCGComponentExecutionState.h
@@ -17,6 +17,7 @@ public:
 	virtual FTransform GetTransform() const override;
 	virtual FTransform GetOriginalTransform() const override;
 	virtual UWorld* GetWorld() const override;
+	virtual UObject* GetTarget() const override;
 	virtual bool HasAuthority() const override;
 	virtual FBox GetBounds() const override;
 	virtual FBox GetOriginalBounds() const override;
diff --git a/Engine/Plugins/PCG/Source/PCG/Public/PCGContext.h b/Engine/Plugins/PCG/Source/PCG/Public/PCGContext.h
--- a/Engine/Plugins/PCG/Source/PCG/Public/PCGContext.h
+++ b/Engine/Plugins/PCG/Source/PCG/Public/PCGContext.h
@@ -197,8 +197,17 @@ PRAGMA_ENABLE_DEPRECATION_WARNINGS
 
 	bool ShouldStop() const { return AsyncState.ShouldStop(); }
 
+	// Will eventually be deprecated in favor of GetExecutionTarget/GetTypedExecutionTarget
 	UE_API AActor* GetTargetActor(const UPCGSpatialData* InSpatialData) const;
 
+	UE_API UObject* GetExecutionTarget() const;
+
+	template<class TargetType>
+	TargetType* GetTypedExecutionTarget() const
+	{
+		return Cast<TargetType>(GetExecutionTarget());
+	}
+
 	/** Time slicing is not enabled by default. */
 	virtual bool TimeSliceIsEnabled() const { return false; }
 
diff --git a/Engine/Plugins/PCG/Source/PCG/Public/PCGGraphExecutionStateInterface.h b/Engine/Plugins/PCG/Source/PCG/Public/PCGGraphExecutionStateInterface.h
--- a/Engine/Plugins/PCG/Source/PCG/Public/PCGGraphExecutionStateInterface.h
+++ b/Engine/Plugins/PCG/Source/PCG/Public/PCGGraphExecutionStateInterface.h
@@ -71,6 +71,15 @@ public:
 	/** Returns a World, can be null */
 	virtual UWorld* GetWorld() const = 0;
 
+	/** Returns the execution target for managed resources */
+	virtual UObject* GetTarget() const { return nullptr; }
+
+	template<class TargetType>
+	TargetType* GetTypedTarget() const
+	{
+		return Cast<TargetType>(GetTarget());
+	}
+
 	/** Returns true if the ExecutionState has network authority */
 	virtual bool HasAuthority() const = 0;
 
diff --git a/Engine/Plugins/PCGInterops/PCGGeometryScriptInterop/Source/PCGGeometryScriptInterop/Private/Elements/PCGSpawnDynamicMesh.cpp b/Engine/Plugins/PCGInterops/PCGGeometryScriptInterop/Source/PCGGeometryScriptInterop/Private/Elements/PCGSpawnDynamicMesh.cpp
--- a/Engine/Plugins/PCGInterops/PCGGeometryScriptInterop/Source/PCGGeometryScriptInterop/Private/Elements/PCGSpawnDynamicMesh.cpp
+++ b/Engine/Plugins/PCGInterops/PCGGeometryScriptInterop/Source/PCGGeometryScriptInterop/Private/Elements/PCGSpawnDynamicMesh.cpp
@@ -46,7 +46,7 @@ bool FPCGSpawnDynamicMeshElement::ExecuteInternal(FPCGContext* InContext) const
 	const UPCGSpawnDynamicMeshSettings* Settings = InContext->GetInputSettings<UPCGSpawnDynamicMeshSettings>();
 	check(Settings);
 
-	AActor* TargetActor = Settings->TargetActor.IsValid() ? Settings->TargetActor.Get() : InContext->GetTargetActor(nullptr);
+	AActor* TargetActor = Settings->TargetActor.IsValid() ? Settings->TargetActor.Get() : InContext->GetTypedExecutionTarget<AActor>();
 	if (!TargetActor)
 	{
 		PCGLog::LogErrorOnGraph(LOCTEXT("InvalidTargetActor", "Invalid target actor."), InContext);
