[PCG] [Bughawk] Fix badly scoped RF_Transient flag
- This was not causing issues because UPCGDebugDrawComponent class was already marked with transient class flag

#rnx

diff --git a/Engine/Plugins/PCG/Source/PCG/Private/Elements/PCGVisualizeAttribute.cpp b/Engine/Plugins/PCG/Source/PCG/Private/Elements/PCGVisualizeAttribute.cpp
--- a/Engine/Plugins/PCG/Source/PCG/Private/Elements/PCGVisualizeAttribute.cpp
+++ b/Engine/Plugins/PCG/Source/PCG/Private/Elements/PCGVisualizeAttribute.cpp
@@ -108,7 +108,7 @@ bool FPCGVisualizeAttribute::ExecuteInternal(FPCGContext* Context) const
 		// If it doesn't exist already, create a new one
 		if (!DebugDrawComponent)
 		{
-			DebugDrawComponent = NewObject<UPCGDebugDrawComponent>(TargetActor, FName(LOCTEXT("PCGDebugDrawComponent", "PCG Debug Draw Component").ToString(), RF_Transient));
+			DebugDrawComponent = NewObject<UPCGDebugDrawComponent>(TargetActor, FName(LOCTEXT("PCGDebugDrawComponent", "PCG Debug Draw Component").ToString()), RF_Transient);
 
 			TargetActor->Modify(/*bAlwaysMarkDirty=*/false);
 			DebugDrawComponent->RegisterComponent();
