[PCG] Switch to capturing data upload adaptor using weak pointer in data pack task.

If task execution is delayed the adaptor could be destroyed and the task would be aborted.

#rnx

diff --git a/Engine/Plugins/PCG/Source/PCG/Private/Compute/PCGDataUploadAdaptor.cpp b/Engine/Plugins/PCG/Source/PCG/Private/Compute/PCGDataUploadAdaptor.cpp
--- a/Engine/Plugins/PCG/Source/PCG/Private/Compute/PCGDataUploadAdaptor.cpp
+++ b/Engine/Plugins/PCG/Source/PCG/Private/Compute/PCGDataUploadAdaptor.cpp
@@ -196,7 +196,7 @@ bool FPCGDataUploadAdaptor::PrepareData_GameThread(const UPCGComputeKernel* InPr
 			}
 
 			// Kick off async task to pack data collection as this can be a heavy operation.
-			PackDataCollectionTask = UE::Tasks::Launch(UE_SOURCE_LOCATION, [ContextHandle=Binding->GetContextHandle(), This=AsShared()]()
+			PackDataCollectionTask = UE::Tasks::Launch(UE_SOURCE_LOCATION, [ContextHandle=Binding->GetContextHandle(), ThisWeak=AsWeak()]()
 			{
 				TRACE_CPUPROFILER_EVENT_SCOPE(PackDataCollectionTask);
 				LLM_SCOPE_BYTAG(PCG);
@@ -208,6 +208,12 @@ bool FPCGDataUploadAdaptor::PrepareData_GameThread(const UPCGComputeKernel* InPr
 					return;
 				}
 
+				TSharedPtr<FPCGDataUploadAdaptor> This = ThisWeak.Pin();
+				if (!This)
+				{
+					return;
+				}
+
 				PCGDataCollectionPackingHelpers::PackDataCollection(Context->DataBinding->GetInputDataCollection(), This->TargetDataCollectionDesc, This->InputPinLabel, Context->DataBinding, This->PackedDataCollection);
 			});
 
