[PCG] Revisited spline implementations so we are not reliant on having uniform input keys & added proper last segment key length (key offset) for loops.

#rnx

diff --git a/Engine/Plugins/PCG/Source/PCG/Private/Data/PCGSplineData.cpp b/Engine/Plugins/PCG/Source/PCG/Private/Data/PCGSplineData.cpp
--- a/Engine/Plugins/PCG/Source/PCG/Private/Data/PCGSplineData.cpp
+++ b/Engine/Plugins/PCG/Source/PCG/Private/Data/PCGSplineData.cpp
@@ -236,6 +236,19 @@ float UPCGSplineData::GetInputKeyAtDistance(int SegmentIndex, FVector::FReal Dis
 	return SplineStruct.GetInputKeyAtDistanceAlongSpline(FullDistance);
 }
 
+float UPCGSplineData::GetInputKeyAtAlpha(float InAlpha) const
+{
+	// Alpha is a measure of progression along the spline, where segments have an equal portion.
+	const float Alpha = FMath::Clamp(InAlpha, 0.0f, 1.0f);
+	const int NumSegments = GetNumSegments();
+
+	const float UnnormalizedAlpha = Alpha * NumSegments;
+	const int SegmentIndex = static_cast<int>(UnnormalizedAlpha);
+	const float Remainder = UnnormalizedAlpha - static_cast<float>(SegmentIndex);
+
+	return GetInputKeyAtDistance(SegmentIndex, Remainder * GetSegmentLength(SegmentIndex));
+}
+
 void UPCGSplineData::GetTangentsAtSegmentStart(int SegmentIndex, FVector& OutArriveTangent, FVector& OutLeaveTangent) const
 {
 	check(SplineStruct.GetSplinePointsPosition().Points.IsValidIndex(SegmentIndex));
@@ -258,6 +271,11 @@ FTransform UPCGSplineData::GetTransformAtAlpha(float Alpha) const
 	return SplineStruct.GetTransformAtSplineInputKey(GetInputKeyAtAlpha(Alpha), ESplineCoordinateSpace::World);
 }
 
+float UPCGSplineData::GetLoopKeyOffset() const
+{
+	return SplineStruct.GetLoopKeyOffset();
+}
+
 void UPCGSplineData::WriteMetadataToPoint(float InputKey, FPCGPoint& OutPoint, UPCGMetadata* OutMetadata) const
 {
 	return WriteMetadataToEntry(InputKey, OutPoint.MetadataEntry, OutMetadata);
@@ -551,7 +569,7 @@ TArray<FSplinePoint> UPCGSplineData::GetSplinePoints() const
 
 	for (int i = 0; i < ControlPointsPosition.Points.Num(); ++i)
 	{
-		ControlPoints.Emplace(static_cast<float>(ControlPoints.Num()),
+		ControlPoints.Emplace(ControlPointsPosition.Points[i].InVal,
 			ControlPointsPosition.Points[i].OutVal,
 			ControlPointsPosition.Points[i].ArriveTangent,
 			ControlPointsPosition.Points[i].LeaveTangent,
diff --git a/Engine/Plugins/PCG/Source/PCG/Private/Data/PCGSplineStruct.cpp b/Engine/Plugins/PCG/Source/PCG/Private/Data/PCGSplineStruct.cpp
--- a/Engine/Plugins/PCG/Source/PCG/Private/Data/PCGSplineStruct.cpp
+++ b/Engine/Plugins/PCG/Source/PCG/Private/Data/PCGSplineStruct.cpp
@@ -454,6 +454,27 @@ int FPCGSplineStruct::GetNumberOfPoints() const
 		: Spline.GetNumControlPoints();
 }
 
+float FPCGSplineStruct::GetLoopKeyOffset() const
+{
+	return ShouldUseSplineCurves()
+		? WarninglessSplineCurves().Position.LoopKeyOffset
+		: Spline.GetSplinePointsPosition().LoopKeyOffset;
+}
+
+void FPCGSplineStruct::SetLoopKeyOffset(float LoopKeyOffset)
+{
+	//@todo_pcg - have a proper way to set this during initialization.
+	if (bClosedLoop)
+	{
+		const bool bStationaryEndpoints = false;
+		const bool bLoopPositionOverride = true;
+
+		FSpline::FUpdateSplineParams Params{ bClosedLoop, bStationaryEndpoints, ReparamStepsPerSegment, bLoopPositionOverride, LoopKeyOffset };
+		WarninglessSplineCurves().UpdateSpline(Params.bClosedLoop, Params.bStationaryEndpoints, Params.ReparamStepsPerSegment, Params.bLoopPositionOverride, Params.LoopPosition, Params.Scale3D);
+		Spline.UpdateSpline(Params);
+	}
+}
+
 FVector::FReal FPCGSplineStruct::GetSplineLength() const
 {
 	return ShouldUseSplineCurves()
diff --git a/Engine/Plugins/PCG/Source/PCG/Private/Elements/PCGSplineDirection.cpp b/Engine/Plugins/PCG/Source/PCG/Private/Elements/PCGSplineDirection.cpp
--- a/Engine/Plugins/PCG/Source/PCG/Private/Elements/PCGSplineDirection.cpp
+++ b/Engine/Plugins/PCG/Source/PCG/Private/Elements/PCGSplineDirection.cpp
@@ -64,6 +64,9 @@ namespace PCGSplineDirection
 		TArray<FSplinePoint> NewControlPoints;
 		NewControlPoints.Reserve(ControlPointsPosition.Points.Num());
 
+		const float FirstKeyValue = ControlPointsPosition.Points.IsEmpty() ? 0.0f : ControlPointsPosition.Points[0].InVal;
+		const float LastKeyValue = ControlPointsPosition.Points.IsEmpty() ? 0.0f : ControlPointsPosition.Points.Last().InVal;
+
 		for (int i = ControlPointsPosition.Points.Num() - 1; i >= 0; --i)
 		{
 			/* Implementation Note: Segment interpolation is determined by the interpolation mode of the preceding
@@ -73,19 +76,32 @@ namespace PCGSplineDirection
 			 * noting that since each spline segment is calculated from [0..1) there is a slight inconsistency when
 			 * evaluated in reverse order, as effectively the reverse is [1..0) per segment.
 			 */
-			NewControlPoints.Emplace(static_cast<float>(NewControlPoints.Num()),
+			ESplinePointType::Type NewControlPointType = ConvertInterpCurveModeToSplinePointType(ControlPointsPosition.Points[(i > 0 ? i - 1 : i)].InterpMode);
+
+			if ((NewControlPointType == ESplinePointType::Curve || NewControlPointType == ESplinePointType::CurveClamped || NewControlPointType == ESplinePointType::CurveCustomTangent))
+			{
+				NewControlPointType = ESplinePointType::CurveCustomTangent;
+			}
+
+			NewControlPoints.Emplace(FirstKeyValue + LastKeyValue - ControlPointsPosition.Points[i].InVal,
 				ControlPointsPosition.Points[i].OutVal,
 				-ControlPointsPosition.Points[i].LeaveTangent, // Tangents are inverted and swapped
 				-ControlPointsPosition.Points[i].ArriveTangent,
 				ControlPointsRotation.Points[i].OutVal.Rotator(),
 				ControlPointsScale.Points[i].OutVal,
-				ESplinePointType::CurveCustomTangent);
+				NewControlPointType);
 		}
 
 		UPCGSplineData* NewSplineData = FPCGContext::NewObject_AnyThread<UPCGSplineData>(Context);
 		NewSplineData->InitializeFromData(InputSplineData);
 		NewSplineData->Initialize(NewControlPoints, InputSplineData->IsClosed(), InputSplineData->GetTransform(), std::move(ControlPointKeys));
 
+		// Set the key offset too based on the original spline.
+		if (InputSplineData->IsClosed())
+		{
+			NewSplineData->SplineStruct.SetLoopKeyOffset(InputSplineData->GetLoopKeyOffset());
+		}
+		
 		return NewSplineData;
 	}
 }
diff --git a/Engine/Plugins/PCG/Source/PCG/Private/Elements/PCGSplitSplines.cpp b/Engine/Plugins/PCG/Source/PCG/Private/Elements/PCGSplitSplines.cpp
--- a/Engine/Plugins/PCG/Source/PCG/Private/Elements/PCGSplitSplines.cpp
+++ b/Engine/Plugins/PCG/Source/PCG/Private/Elements/PCGSplitSplines.cpp
@@ -45,8 +45,10 @@ namespace PCGSplitSpline
 			}
 		};
 
-		auto AddSplitPoint = [&OutSplitSpline, &SplitSplinePoints, &SplitSplinePointsEntryKeys, &Spline, &SplinePointsEntryKeys](const double Key)
+		auto AddSplitPoint = [&OutSplitSpline, &SplitSplinePoints, &SplitSplinePointsEntryKeys, &Spline, &SplinePointsEntryKeys](const double Key, TEnumAsByte<ESplinePointType::Type> PointType, const double PreviousKey, const double NextKey)
 		{
+			check(Key >= PreviousKey && Key <= NextKey && PreviousKey <= NextKey);
+
 			FSplinePoint SplitPoint;
 			SplitPoint.InputKey = Key;
 			SplitPoint.Position = Spline->SplineStruct.GetLocationAtSplineInputKey(Key, ESplineCoordinateSpace::Local);
@@ -54,11 +56,12 @@ namespace PCGSplitSpline
 			SplitPoint.LeaveTangent = SplitPoint.ArriveTangent;
 			SplitPoint.Rotation = Spline->SplineStruct.GetQuaternionAtSplineInputKey(Key, ESplineCoordinateSpace::Local).Rotator();
 			SplitPoint.Scale = Spline->SplineStruct.GetScaleAtSplineInputKey(Key);
-			SplitPoint.Type = ESplinePointType::CurveCustomTangent;
+			SplitPoint.Type = PointType;
 
 			// The tangent is always computed from the original spline's perspective so we will assume the key range is always one.
-			SplitPoint.ArriveTangent *= FMath::Frac(Key);
-			SplitPoint.LeaveTangent *= 1.0 - FMath::Frac(Key);
+			const double Ratio = (NextKey - PreviousKey) > UE_SMALL_NUMBER ? (Key - PreviousKey) / (NextKey - PreviousKey) : 0.5;
+			SplitPoint.ArriveTangent *= Ratio;
+			SplitPoint.LeaveTangent *= 1.0 - Ratio;
 
 			SplitSplinePoints.Add(SplitPoint);
 
@@ -70,56 +73,109 @@ namespace PCGSplitSpline
 			}
 		};
 
-		int FirstControlPointInSegment = FMath::CeilToInt(StartKey);
-		int LastControlPointInSegment = FMath::FloorToInt(EndKey);
+		auto GetPointTypeForTouchedPoint = [](TEnumAsByte<ESplinePointType::Type> FirstType) -> TEnumAsByte<ESplinePointType::Type>
+		{
+			if (FirstType == ESplinePointType::Curve ||
+				FirstType == ESplinePointType::CurveClamped ||
+				FirstType == ESplinePointType::CurveCustomTangent)
+			{
+				return ESplinePointType::CurveCustomTangent;
+			}
+			else
+			{
+				return FirstType;
+			}
+		};
+
+		int FirstControlPointIndexInSegment = INDEX_NONE;
+		int LastControlPointIndexInSegment = INDEX_NONE;
+
+		for (int SplinePointIndex = 0; SplinePointIndex < SplinePoints.Num(); ++SplinePointIndex)
+		{
+			if (SplinePoints[SplinePointIndex].InputKey >= StartKey)
+			{
+				FirstControlPointIndexInSegment = SplinePointIndex;
+				break;
+			}
+		}
+
+		for (int SplinePointIndex = SplinePoints.Num() - 1; SplinePointIndex >= 0; --SplinePointIndex)
+		{
+			if (SplinePoints[SplinePointIndex].InputKey <= EndKey)
+			{
+				LastControlPointIndexInSegment = SplinePointIndex;
+				break;
+			}
+		}
+
+		if (FirstControlPointIndexInSegment == INDEX_NONE || LastControlPointIndexInSegment == INDEX_NONE)
+		{
+			// Invalid keys - unclear what caller is expecting
+			ensure(false);
+			return;
+		}
 
 		if (Spline->IsClosed() && EndKey <= StartKey) // allow for loop
 		{
-			LastControlPointInSegment += SplinePoints.Num();
+			LastControlPointIndexInSegment += SplinePoints.Num();
 		}
 
+		const double LoopKeyLength = Spline->IsClosed() ? (Spline->GetLoopKeyOffset() - SplinePoints.Last().InputKey) : 1.0;
+
+		const double FirstControlKey = SplinePoints[FirstControlPointIndexInSegment].InputKey;
+		const double PreviousControlKey = FirstControlPointIndexInSegment > 0 ? SplinePoints[FirstControlPointIndexInSegment - 1].InputKey : FirstControlKey - LoopKeyLength;
+		const double LastControlKey = SplinePoints[(LastControlPointIndexInSegment % SplinePoints.Num())].InputKey;
+		const double NextControlKey = ((LastControlPointIndexInSegment + 1) % SplinePoints.Num() == 0) ? (LastControlKey + LoopKeyLength) : SplinePoints[((LastControlPointIndexInSegment + 1) % SplinePoints.Num())].InputKey;
+
 		bool bNeedsStartTangentAdjustment = false;
 		bool bNeedsEndTangentAdjustment = false;
 
 		// Add left split point, prioritize existing control points when possible
-		const double KeyEpsilon = 1.0e-6f;
-		const double FracStartKey = FMath::Frac(StartKey);
-		if (FracStartKey < KeyEpsilon)
+		const double KeyEpsilon = 1.0e-6;
+
+		const double StartKeyDeltaToPreviousControlPoint = StartKey - PreviousControlKey;
+		const double StartKeyDeltaToFirstControlPoint = FirstControlKey - StartKey;
+
+		if (StartKeyDeltaToPreviousControlPoint < KeyEpsilon)
 		{
-			if (FracStartKey > 0) // otherwise point is already added in the control point loop
+			if (StartKeyDeltaToPreviousControlPoint > 0) // otherwise point is already added in the control point loop
 			{
-				AddExistingControlPoint(FMath::FloorToInt(StartKey));
+				AddExistingControlPoint((FirstControlPointIndexInSegment + SplinePoints.Num() - 1) % SplinePoints.Num());
 			}
 		}
-		else if (1.0 - FracStartKey < KeyEpsilon)
+		else if (StartKeyDeltaToFirstControlPoint < KeyEpsilon)
 		{
 			// will be added by the control point loop
 		}
 		else
 		{
-			AddSplitPoint(StartKey);
+			TEnumAsByte<ESplinePointType::Type> SplitPointType = GetPointTypeForTouchedPoint(SplinePoints[(FirstControlPointIndexInSegment + SplinePoints.Num() - 1) % SplinePoints.Num()].Type);
+			AddSplitPoint(StartKey, SplitPointType, PreviousControlKey, FirstControlKey);
 			bNeedsStartTangentAdjustment = true;
 		}
 
 		// Add control points that are taken as-is from the original spline points
-		for (int ControlPointIndex = FirstControlPointInSegment; ControlPointIndex <= LastControlPointInSegment; ++ControlPointIndex)
+		for (int ControlPointIndex = FirstControlPointIndexInSegment; ControlPointIndex <= LastControlPointIndexInSegment; ++ControlPointIndex)
 		{
-			AddExistingControlPoint((ControlPointIndex % SplinePoints.Num()));
+			AddExistingControlPoint(ControlPointIndex % SplinePoints.Num());
 		}
 
 		// Add right split point
-		const double FracEndKey = FMath::Frac(EndKey);
-		if (FracEndKey < KeyEpsilon)
+		const double EndKeyDeltaToLastControlPoint = EndKey - LastControlKey;
+		const double EndKeyDeltaToNextControlPoint = NextControlKey - EndKey;
+
+		if (EndKeyDeltaToLastControlPoint < KeyEpsilon)
 		{
 			// will be already added by the control point loop
 		}
-		else if (1.0 - FracEndKey < KeyEpsilon)
+		else if(EndKeyDeltaToNextControlPoint < KeyEpsilon)
 		{
-			AddExistingControlPoint(FMath::CeilToInt(EndKey));
+			AddExistingControlPoint((LastControlPointIndexInSegment + 1) % SplinePoints.Num());
 		}
 		else
 		{
-			AddSplitPoint(EndKey);
+			TEnumAsByte<ESplinePointType::Type> SplitPointType = GetPointTypeForTouchedPoint(SplinePoints[LastControlPointIndexInSegment].Type);
+			AddSplitPoint(EndKey, SplitPointType, LastControlKey, NextControlKey);
 			bNeedsEndTangentAdjustment = true;
 		}
 
@@ -128,16 +184,18 @@ namespace PCGSplitSpline
 		{
 			if (bNeedsStartTangentAdjustment)
 			{
-				// assumes key range is 1.
-				const double Ratio = (1.0 - FMath::Frac(SplitSplinePoints[0].InputKey));
-				SplitSplinePoints[1].ArriveTangent *= Ratio;
+				const double Ratio = ((FirstControlKey - PreviousControlKey) > KeyEpsilon) ? (StartKey - PreviousControlKey) / (FirstControlKey - PreviousControlKey) : 0.0;
+				FSplinePoint& FirstPreexistingPoint = SplitSplinePoints[1];
+				FirstPreexistingPoint.ArriveTangent *= (1.0 - Ratio);
+				FirstPreexistingPoint.Type = GetPointTypeForTouchedPoint(FirstPreexistingPoint.Type);
 			}
 
 			if (bNeedsEndTangentAdjustment)
 			{
-				// assumes key range is 1.
-				const double Ratio = FMath::Frac(SplitSplinePoints.Last().InputKey);
-				SplitSplinePoints[SplitSplinePoints.Num() - 2].LeaveTangent *= Ratio;
+				const double Ratio = ((NextControlKey - LastControlKey) > KeyEpsilon) ? (EndKey - LastControlKey) / (NextControlKey - LastControlKey) : 1.0;
+				FSplinePoint& LastPreexistingPoint = SplitSplinePoints[SplitSplinePoints.Num() - 2];
+				LastPreexistingPoint.LeaveTangent *= Ratio;
+				LastPreexistingPoint.Type = GetPointTypeForTouchedPoint(LastPreexistingPoint.Type);
 			}
 		}
 
@@ -146,10 +204,6 @@ namespace PCGSplitSpline
 		{
 			FSplinePoint& SplinePoint = SplitSplinePoints[SPIndex];
 			SplinePoint.InputKey = static_cast<float>(SPIndex);
-			if (SplinePoint.Type == ESplinePointType::Curve || SplinePoint.Type == ESplinePointType::CurveClamped)
-			{
-				SplitSplinePoints[SPIndex].Type = ESplinePointType::CurveCustomTangent;
-			}
 		}
 
 		// Initialize data now
@@ -338,12 +392,15 @@ bool FPCGSplitSplineElement::ExecuteInternal(FPCGContext* Context) const
 		const bool bIsClosed = SplineData->IsClosed();
 
 		// Remove values that are outside of the spline
+		const float FirstKey = SplinePoints.IsEmpty() ? 0.0f : SplinePoints[0].InputKey;
+		const float LastKey = SplinePoints.IsEmpty() ? 1.0f : (bIsClosed ? SplineData->GetLoopKeyOffset() : SplinePoints.Last().InputKey);
+
+		// Closed loops should properly cut at key = 0 and up to the last key + 1.
 		if (bIsClosed)
 		{
-			// Closed loops should properly cut at key = 0 and up to the last key + 1.
 			for (int SplitIndex = SplitKeys.Num() - 1; SplitIndex >= 0; --SplitIndex)
 			{
-				if (SplitKeys[SplitIndex] < SplinePoints[0].InputKey || SplitKeys[SplitIndex] > 1.0 + SplinePoints.Last().InputKey)
+				if (SplitKeys[SplitIndex] < FirstKey || SplitKeys[SplitIndex] > LastKey)
 				{
 					SplitKeys.RemoveAtSwap(SplitIndex);
 				}
@@ -353,7 +410,8 @@ bool FPCGSplitSplineElement::ExecuteInternal(FPCGContext* Context) const
 		{
 			for (int SplitIndex = SplitKeys.Num() - 1; SplitIndex >= 0; --SplitIndex)
 			{
-				if (SplitKeys[SplitIndex] <= SplinePoints[0].InputKey || SplitKeys[SplitIndex] >= SplinePoints.Last().InputKey)
+				// Implementation note: in open splines, there's no point in splitting at the first or last keys
+				if (SplitKeys[SplitIndex] <= FirstKey || SplitKeys[SplitIndex] >= LastKey)
 				{
 					SplitKeys.RemoveAtSwap(SplitIndex);
 				}
@@ -376,7 +434,7 @@ bool FPCGSplitSplineElement::ExecuteInternal(FPCGContext* Context) const
 		for (int SplitIndex = SplitKeys.Num() - 1; SplitIndex > 0; --SplitIndex)
 		{
 			if ((SplitKeys[SplitIndex] - SplitKeys[SplitIndex - 1] < KeyEpsilon) ||
-				(bIsClosed && SplitIndex == SplitKeys.Num() - 1 && (SplitKeys[0] + 1.0 + SplinePoints.Last().InputKey) - SplitKeys[SplitIndex] < KeyEpsilon))
+				(bIsClosed && SplitIndex == SplitKeys.Num() - 1 && (LastKey - SplitKeys[SplitIndex]) < KeyEpsilon))
 			{
 				SplitKeys.RemoveAt(SplitIndex);
 			}
@@ -388,8 +446,7 @@ bool FPCGSplitSplineElement::ExecuteInternal(FPCGContext* Context) const
 		// For open splines: we'll add a starting key (assumedly 0 in most cases) and a finishing key so we have proper segments [0 -> split key] ... [split key -> end]
 		if (bIsClosed)
 		{
-			const double FirstKey = SplitKeys[0];
-			SplitKeys.Add(FirstKey);
+			SplitKeys.Add(SplitKeys[0]);
 		}
 		else
 		{
diff --git a/Engine/Plugins/PCG/Source/PCG/Public/Data/PCGSplineData.h b/Engine/Plugins/PCG/Source/PCG/Public/Data/PCGSplineData.h
--- a/Engine/Plugins/PCG/Source/PCG/Public/Data/PCGSplineData.h
+++ b/Engine/Plugins/PCG/Source/PCG/Public/Data/PCGSplineData.h
@@ -91,11 +91,13 @@ public:
 	UE_API virtual FTransform GetTransformAtDistance(int SegmentIndex, FVector::FReal Distance, bool bWorldSpace = true, FBox* OutBounds = nullptr) const override;
 	UE_API virtual FVector::FReal GetCurvatureAtDistance(int SegmentIndex, FVector::FReal Distance) const override;
 	UE_API virtual float GetInputKeyAtDistance(int SegmentIndex, FVector::FReal Distance) const override;
+	UE_API virtual float GetInputKeyAtAlpha(float Alpha) const override;
 	UE_API virtual void GetTangentsAtSegmentStart(int SegmentIndex, FVector& OutArriveTangent, FVector& OutLeaveTangent) const override;
 	UE_API virtual FVector::FReal GetDistanceAtSegmentStart(int SegmentIndex) const override;
 	UE_API virtual FVector GetLocationAtAlpha(float Alpha) const override;
 	UE_API virtual FTransform GetTransformAtAlpha(float Alpha) const override;
 	virtual bool IsClosed() const override { return SplineStruct.bClosedLoop; }
+	UE_API float GetLoopKeyOffset() const;
 	UE_API virtual void WriteMetadataToPoint(float InputKey, FPCGPoint& OutPoint, UPCGMetadata* OutMetadata) const override;
 
 	virtual void AllocateMetadataEntries() override { SplineStruct.AllocateMetadataEntries(); }
diff --git a/Engine/Plugins/PCG/Source/PCG/Public/Data/PCGSplineStruct.h b/Engine/Plugins/PCG/Source/PCG/Public/Data/PCGSplineStruct.h
--- a/Engine/Plugins/PCG/Source/PCG/Public/Data/PCGSplineStruct.h
+++ b/Engine/Plugins/PCG/Source/PCG/Public/Data/PCGSplineStruct.h
@@ -94,10 +94,12 @@ struct FPCGSplineStruct
 	UE_API float GetAttributeValue(float InputKey, FName ChannelName) const;
 	
 	UE_API void UpdateSpline();
+	UE_API void SetLoopKeyOffset(float LoopKeyOffset);
 
 	UE_API int GetNumberOfSplineSegments() const;
 	UE_API int GetNumberOfPoints() const;
 	bool IsClosedLoop() const { return bClosedLoop; }
+	UE_API float GetLoopKeyOffset() const;
 	UE_API FVector::FReal GetSplineLength() const;
 	UE_API FBox GetBounds() const;
 	UE_API FBox GetSegmentBounds(int32 Index) const;
