[PCG] Small fix for closed splines inversion

#rnx

diff --git a/Engine/Plugins/PCG/Source/PCG/Private/Elements/PCGSplineDirection.cpp b/Engine/Plugins/PCG/Source/PCG/Private/Elements/PCGSplineDirection.cpp
--- a/Engine/Plugins/PCG/Source/PCG/Private/Elements/PCGSplineDirection.cpp
+++ b/Engine/Plugins/PCG/Source/PCG/Private/Elements/PCGSplineDirection.cpp
@@ -76,7 +76,8 @@ namespace PCGSplineDirection
 			 * noting that since each spline segment is calculated from [0..1) there is a slight inconsistency when
 			 * evaluated in reverse order, as effectively the reverse is [1..0) per segment.
 			 */
-			ESplinePointType::Type NewControlPointType = ConvertInterpCurveModeToSplinePointType(ControlPointsPosition.Points[(i > 0 ? i - 1 : i)].InterpMode);
+			const int PreviousPointIndex = (i > 0 ? i - 1 : (InputSplineData->IsClosed() ? ControlPointsPosition.Points.Num() - 1 : 0));
+			ESplinePointType::Type NewControlPointType = ConvertInterpCurveModeToSplinePointType(ControlPointsPosition.Points[PreviousPointIndex].InterpMode);
 
 			if ((NewControlPointType == ESplinePointType::Curve || NewControlPointType == ESplinePointType::CurveClamped || NewControlPointType == ESplinePointType::CurveCustomTangent))
 			{
