[PCG] Minor fix to custom HLSL declaration text.

#rnx

diff --git a/Engine/Plugins/PCG/Source/PCG/Private/Compute/Elements/PCGCustomHLSL.cpp b/Engine/Plugins/PCG/Source/PCG/Private/Compute/Elements/PCGCustomHLSL.cpp
--- a/Engine/Plugins/PCG/Source/PCG/Private/Compute/Elements/PCGCustomHLSL.cpp
+++ b/Engine/Plugins/PCG/Source/PCG/Private/Compute/Elements/PCGCustomHLSL.cpp
@@ -567,7 +567,7 @@ void UPCGCustomHLSLSettings::UpdateInputDeclarations()
 		}
 
 		InputDeclarations += FString::Format(TEXT(
-			"// Samples a virtual texture and gets all values that are available. Otherwise returns default values."
+			"// Samples a virtual texture and gets all values that are available. Otherwise returns default values.\n"
 			"void {0}_SampleVirtualTexture(\n"
 			"	uint InDataIndex,\n"
 			"	float3 InWorldPos,\n"
