[PCG] Compact TSet fix: replace Remove with iterator RemoveCurrent

#rnx

diff --git a/Engine/Plugins/PCG/Source/PCGEditor/Private/PCGEditorModule.cpp b/Engine/Plugins/PCG/Source/PCGEditor/Private/PCGEditorModule.cpp
--- a/Engine/Plugins/PCG/Source/PCGEditor/Private/PCGEditorModule.cpp
+++ b/Engine/Plugins/PCG/Source/PCGEditor/Private/PCGEditorModule.cpp
@@ -1095,7 +1095,7 @@ void FPCGEditorModule::ClearExecutedStacks(const UPCGGraph* InContainingGraph)
 
 		if (StackPtr->HasObject(InContainingGraph))
 		{
-			ExecutedStacks.Remove(StackPtr);
+			Iterator.RemoveCurrent();
 
 			if (TSet<FPCGStackSharedPtr>* StacksForComponent = ExecutedStacksPerSource.Find(StackPtr->GetRootSource()))
 			{
