[PCG] Expose GetGraph on the PCG Component

#rnx

diff --git a/Engine/Plugins/PCG/Source/PCG/Public/PCGComponent.h b/Engine/Plugins/PCG/Source/PCG/Public/PCGComponent.h
--- a/Engine/Plugins/PCG/Source/PCG/Public/PCGComponent.h
+++ b/Engine/Plugins/PCG/Source/PCG/Public/PCGComponent.h
@@ -172,7 +172,22 @@ public:
 
 	UE_API bool CanPartition() const;
 
+	/**
+	 * Returns the mutable underlying PCG Graph for this component.
+	 * Note that this is not the Graph Instance on the PCG Component, but the PCG Graph asset at the top of the instance hierarchy.
+	 * If it is needed to set Graph Parameters on the component, use the GraphInstance.
+	 */
+	UFUNCTION(BlueprintCallable, DisplayName="Get Mutable PCG Graph", Category = PCG)
 	UE_API UPCGGraph* GetGraph() const;
+
+	/**
+	 * Returns the const underlying PCG Graph for this component.
+	 * Note that this is not the Graph Instance on the PCG Component, but the PCG Graph asset at the top of the instance hierarchy.
+	 * If it is needed to set Graph Parameters on the component, use the GraphInstance.
+	 */
+	UFUNCTION(BlueprintCallable, DisplayName="Get Const PCG Graph", Category = PCG)
+	const UPCGGraph* GetConstGraph() const { return GetGraph(); }
+
 	UPCGGraphInstance* GetGraphInstance() const { return GraphInstance; }
 
 	UE_API void SetGraphLocal(UPCGGraphInterface* InGraph);
