[PCG] Blueprint : Point loop functions can require game thread execution (single threaded)
- This will cause errors if the blueprint element itself is not 'bRequiresGameThread'

diff --git a/Engine/Plugins/PCG/Source/PCG/Private/Blueprint/PCGBlueprintDeprecatedElement.cpp b/Engine/Plugins/PCG/Source/PCG/Private/Blueprint/PCGBlueprintDeprecatedElement.cpp
--- a/Engine/Plugins/PCG/Source/PCG/Private/Blueprint/PCGBlueprintDeprecatedElement.cpp
+++ b/Engine/Plugins/PCG/Source/PCG/Private/Blueprint/PCGBlueprintDeprecatedElement.cpp
@@ -18,7 +18,7 @@ void UPCGBlueprintElement::ExecuteWithContext_Implementation(FPCGContext& InCont
 	Execute(Input, Output);
 }
 
-void UPCGBlueprintElement::PointLoop(FPCGContext& InContext, const UPCGPointData* InData, UPCGPointData*& OutData, UPCGPointData* OptionalOutData) const
+void UPCGBlueprintElement::PointLoop(FPCGContext& InContext, const UPCGPointData* InData, UPCGPointData*& OutData, UPCGPointData* OptionalOutData, bool bGameThread) const
 {
 	if (!InData)
 	{
@@ -26,6 +26,12 @@ void UPCGBlueprintElement::PointLoop(FPCGContext& InContext, const UPCGPointData
 		return;
 	}
 
+	if (bGameThread && !IsInGameThread())
+	{
+		ThrowRequiresGameThreadException(&InContext, TEXT("PointLoop"));
+		return;
+	}
+
 	if (OptionalOutData)
 	{
 		OutData = OptionalOutData;
@@ -39,8 +45,8 @@ void UPCGBlueprintElement::PointLoop(FPCGContext& InContext, const UPCGPointData
 	const TArray<FPCGPoint>& InPoints = InData->GetPoints();
 	TArray<FPCGPoint>& OutPoints = OutData->GetMutablePoints();
 
-	bool bPreviousBPStateValue = true;
-	std::swap(bPreviousBPStateValue, InContext.AsyncState.bIsCallingBlueprint);
+	TGuardValue<bool> IsCallingBlueprint(CurrentContext->AsyncState.bIsCallingBlueprint, true);
+	TGuardValue<int32> NumTaskOverride(CurrentContext->AsyncState.NumAvailableTasks, bGameThread ? 1 : CurrentContext->AsyncState.NumAvailableTasks);
 
 	FPCGAsync::AsyncPointProcessing(&InContext, InPoints.Num(), OutPoints, [this, &InContext, InData, OutData, &InPoints](int32 Index, FPCGPoint& OutPoint)
 	{
@@ -51,11 +57,9 @@ void UPCGBlueprintElement::PointLoop(FPCGContext& InContext, const UPCGPointData
 
 		return PointLoopBody(InContext, InData, InPoints[Index], OutPoint, OutData->Metadata, Index);
 	});
-
-	std::swap(bPreviousBPStateValue, InContext.AsyncState.bIsCallingBlueprint);
 }
 
-void UPCGBlueprintElement::VariableLoop(FPCGContext& InContext, const UPCGPointData* InData, UPCGPointData*& OutData, UPCGPointData* OptionalOutData) const
+void UPCGBlueprintElement::VariableLoop(FPCGContext& InContext, const UPCGPointData* InData, UPCGPointData*& OutData, UPCGPointData* OptionalOutData, bool bGameThread) const
 {
 	if (!InData)
 	{
@@ -63,6 +67,12 @@ void UPCGBlueprintElement::VariableLoop(FPCGContext& InContext, const UPCGPointD
 		return;
 	}
 
+	if (bGameThread && !IsInGameThread())
+	{
+		ThrowRequiresGameThreadException(&InContext, TEXT("VariableLoop"));
+		return;
+	}
+
 	if (OptionalOutData)
 	{
 		OutData = OptionalOutData;
@@ -76,8 +86,8 @@ void UPCGBlueprintElement::VariableLoop(FPCGContext& InContext, const UPCGPointD
 	const TArray<FPCGPoint>& InPoints = InData->GetPoints();
 	TArray<FPCGPoint>& OutPoints = OutData->GetMutablePoints();
 
-	bool bPreviousBPStateValue = true;
-	std::swap(bPreviousBPStateValue, InContext.AsyncState.bIsCallingBlueprint);
+	TGuardValue<bool> IsCallingBlueprint(CurrentContext->AsyncState.bIsCallingBlueprint, true);
+	TGuardValue<int32> NumTaskOverride(CurrentContext->AsyncState.NumAvailableTasks, bGameThread ? 1 : CurrentContext->AsyncState.NumAvailableTasks);
 
 	FPCGAsync::AsyncMultiPointProcessing(&InContext, InPoints.Num(), OutPoints, [this, &InContext, InData, OutData, &InPoints](int32 Index)
 	{
@@ -88,11 +98,9 @@ void UPCGBlueprintElement::VariableLoop(FPCGContext& InContext, const UPCGPointD
 
 		return VariableLoopBody(InContext, InData, InPoints[Index], OutData->Metadata, Index);
 	});
-
-	std::swap(bPreviousBPStateValue, InContext.AsyncState.bIsCallingBlueprint);
 }
 
-void UPCGBlueprintElement::NestedLoop(FPCGContext& InContext, const UPCGPointData* InOuterData, const UPCGPointData* InInnerData, UPCGPointData*& OutData, UPCGPointData* OptionalOutData) const
+void UPCGBlueprintElement::NestedLoop(FPCGContext& InContext, const UPCGPointData* InOuterData, const UPCGPointData* InInnerData, UPCGPointData*& OutData, UPCGPointData* OptionalOutData, bool bGameThread) const
 {
 	if (!InOuterData || !InInnerData)
 	{
@@ -100,6 +108,12 @@ void UPCGBlueprintElement::NestedLoop(FPCGContext& InContext, const UPCGPointDat
 		return;
 	}
 
+	if (bGameThread && !IsInGameThread())
+	{
+		ThrowRequiresGameThreadException(&InContext, TEXT("NestedLoop"));
+		return;
+	}
+
 	if (OptionalOutData)
 	{
 		OutData = OptionalOutData;
@@ -115,8 +129,8 @@ void UPCGBlueprintElement::NestedLoop(FPCGContext& InContext, const UPCGPointDat
 	const TArray<FPCGPoint>& InInnerPoints = InInnerData->GetPoints();
 	TArray<FPCGPoint>& OutPoints = OutData->GetMutablePoints();
 
-	bool bPreviousBPStateValue = true;
-	std::swap(bPreviousBPStateValue, InContext.AsyncState.bIsCallingBlueprint);
+	TGuardValue<bool> IsCallingBlueprint(CurrentContext->AsyncState.bIsCallingBlueprint, true);
+	TGuardValue<int32> NumTaskOverride(CurrentContext->AsyncState.NumAvailableTasks, bGameThread ? 1 : CurrentContext->AsyncState.NumAvailableTasks);
 
 	FPCGAsync::AsyncPointProcessing(&InContext, InOuterPoints.Num() * InInnerPoints.Num(), OutPoints, [this, &InContext, InOuterData, InInnerData, OutData, &InOuterPoints, &InInnerPoints](int32 Index, FPCGPoint& OutPoint)
 	{
@@ -129,11 +143,9 @@ void UPCGBlueprintElement::NestedLoop(FPCGContext& InContext, const UPCGPointDat
 		const int32 InnerIndex = Index % InInnerPoints.Num();
 		return NestedLoopBody(InContext, InOuterData, InInnerData, InOuterPoints[OuterIndex], InInnerPoints[InnerIndex], OutPoint, OutData->Metadata, OuterIndex, InnerIndex);
 	});
-
-	std::swap(bPreviousBPStateValue, InContext.AsyncState.bIsCallingBlueprint);
 }
 
-void UPCGBlueprintElement::IterationLoop(FPCGContext& InContext, int64 NumIterations, UPCGPointData*& OutData, const UPCGSpatialData* InA, const UPCGSpatialData* InB, UPCGPointData* OptionalOutData) const
+void UPCGBlueprintElement::IterationLoop(FPCGContext& InContext, int64 NumIterations, UPCGPointData*& OutData, const UPCGSpatialData* InA, const UPCGSpatialData* InB, UPCGPointData* OptionalOutData, bool bGameThread) const
 {
 	if (NumIterations < 0)
 	{
@@ -141,6 +153,12 @@ void UPCGBlueprintElement::IterationLoop(FPCGContext& InContext, int64 NumIterat
 		return;
 	}
 
+	if (bGameThread && !IsInGameThread())
+	{
+		ThrowRequiresGameThreadException(&InContext, TEXT("IterationLoop"));
+		return;
+	}
+
 	if (OptionalOutData)
 	{
 		OutData = OptionalOutData;
@@ -158,8 +176,8 @@ void UPCGBlueprintElement::IterationLoop(FPCGContext& InContext, int64 NumIterat
 
 	TArray<FPCGPoint>& OutPoints = OutData->GetMutablePoints();
 
-	bool bPreviousBPStateValue = true;
-	std::swap(bPreviousBPStateValue, InContext.AsyncState.bIsCallingBlueprint);
+	TGuardValue<bool> IsCallingBlueprint(CurrentContext->AsyncState.bIsCallingBlueprint, true);
+	TGuardValue<int32> NumTaskOverride(CurrentContext->AsyncState.NumAvailableTasks, bGameThread ? 1 : CurrentContext->AsyncState.NumAvailableTasks);
 
 	FPCGAsync::AsyncPointProcessing(&InContext, NumIterations, OutPoints, [this, &InContext, InA, InB, OutData](int32 Index, FPCGPoint& OutPoint)
 	{
@@ -170,8 +188,6 @@ void UPCGBlueprintElement::IterationLoop(FPCGContext& InContext, int64 NumIterat
 
 		return IterationLoopBody(InContext, Index, InA, InB, OutPoint, OutData->Metadata);
 	});
-
-	std::swap(bPreviousBPStateValue, InContext.AsyncState.bIsCallingBlueprint);
 }
 
 FPCGContext& UPCGBlueprintElement::GetContext() const
diff --git a/Engine/Plugins/PCG/Source/PCG/Private/Elements/Blueprint/PCGBlueprintBaseElement.cpp b/Engine/Plugins/PCG/Source/PCG/Private/Elements/Blueprint/PCGBlueprintBaseElement.cpp
--- a/Engine/Plugins/PCG/Source/PCG/Private/Elements/Blueprint/PCGBlueprintBaseElement.cpp
+++ b/Engine/Plugins/PCG/Source/PCG/Private/Elements/Blueprint/PCGBlueprintBaseElement.cpp
@@ -2,11 +2,13 @@
 
 #include "Elements/Blueprint/PCGBlueprintBaseElement.h"
 
+#include "PCGModule.h"
 #include "PCGPin.h"
 #include "Elements/PCGExecuteBlueprint.h"
 
 #include "Algo/Transform.h"
 #include "Engine/World.h"
+#include "Helpers/PCGBlueprintHelpers.h"
 #include "Helpers/PCGHelpers.h"
 #include "Math/RandomStream.h"
 
@@ -304,4 +306,11 @@ bool UPCGBlueprintBaseElement::GetOutputPinByLabel(FName InPinLabel, FPCGPinProp
 	return false;
 }
 
+void UPCGBlueprintBaseElement::ThrowRequiresGameThreadException(FPCGContext* InContext, const FString& InFunctionName) const
+{
+	const FText FormatText = FText::Format(LOCTEXT("FunctionRequiresGameThread", "Function '{0}' requires game thread and is being called outside of game thread. Make sure to enable 'bRequiresGameThread' on your Blueprint."), FText::FromString(InFunctionName));
+	PCGE_LOG_C(Error, GraphAndLog, InContext, FormatText);
+	UPCGBlueprintHelpers::ThrowBlueprintException(FormatText);
+}
+
 #undef LOCTEXT_NAMESPACE
\ No newline at end of file
diff --git a/Engine/Plugins/PCG/Source/PCG/Private/Elements/Blueprint/PCGBlueprintPointProcessorElement.cpp b/Engine/Plugins/PCG/Source/PCG/Private/Elements/Blueprint/PCGBlueprintPointProcessorElement.cpp
--- a/Engine/Plugins/PCG/Source/PCG/Private/Elements/Blueprint/PCGBlueprintPointProcessorElement.cpp
+++ b/Engine/Plugins/PCG/Source/PCG/Private/Elements/Blueprint/PCGBlueprintPointProcessorElement.cpp
@@ -13,8 +13,14 @@ void UPCGBlueprintPointProcessorElement::PointLoopInitialize_Implementation(cons
 	OutputPointData->InitializeFromDataWithParams(InitializeFromDataParams);
 }
 
-void UPCGBlueprintPointProcessorElement::PointLoop(const UPCGBasePointData* InputPointData, UPCGBasePointData*& OutputPointData)
+void UPCGBlueprintPointProcessorElement::PointLoop(const UPCGBasePointData* InputPointData, UPCGBasePointData*& OutputPointData, bool bGameThread)
 {
+	if (bGameThread && !IsInGameThread())
+	{
+		ThrowRequiresGameThreadException(CurrentContext, TEXT("PointLoop"));
+		return;
+	}
+
 	if (!InputPointData)
 	{
 		return;
@@ -43,6 +49,7 @@ void UPCGBlueprintPointProcessorElement::PointLoop(const UPCGBasePointData* Inpu
 
 	{
 		TGuardValue<bool> IsCallingBlueprint(CurrentContext->AsyncState.bIsCallingBlueprint, true);
+		TGuardValue<int32> NumTaskOverride(CurrentContext->AsyncState.NumAvailableTasks, bGameThread ? 1 : CurrentContext->AsyncState.NumAvailableTasks);
 
 		FPCGAsync::AsyncProcessingRangeEx(
 			&CurrentContext->AsyncState,
@@ -69,8 +76,14 @@ void UPCGBlueprintPointProcessorElement::IterationLoopInitialize_Implementation(
 	OutputPointData->SetNumPoints(NumIterations);
 }
 
-void UPCGBlueprintPointProcessorElement::IterationLoop(int32 NumIterations, const UPCGBasePointData* InputPointDataA, const UPCGBasePointData* InputPointDataB, UPCGBasePointData*& OutputPointData)
+void UPCGBlueprintPointProcessorElement::IterationLoop(int32 NumIterations, const UPCGBasePointData* InputPointDataA, const UPCGBasePointData* InputPointDataB, UPCGBasePointData*& OutputPointData, bool bGameThread)
 {
+	if (bGameThread && !IsInGameThread())
+	{
+		ThrowRequiresGameThreadException(CurrentContext, TEXT("IterationLoop"));
+		return;
+	}
+
 	OutputPointData = FPCGContext::NewPointData_AnyThread(CurrentContext);
 	IterationLoopInitialize(NumIterations, InputPointDataA, InputPointDataB, OutputPointData);
 
@@ -96,6 +109,7 @@ void UPCGBlueprintPointProcessorElement::IterationLoop(int32 NumIterations, cons
 
 	{
 		TGuardValue<bool> IsCallingBlueprint(CurrentContext->AsyncState.bIsCallingBlueprint, true);
+		TGuardValue<int32> NumTaskOverride(CurrentContext->AsyncState.NumAvailableTasks, bGameThread ? 1 : CurrentContext->AsyncState.NumAvailableTasks);
 
 		FPCGAsync::AsyncProcessingRangeEx(
 			&CurrentContext->AsyncState,
diff --git a/Engine/Plugins/PCG/Source/PCG/Public/Elements/Blueprint/PCGBlueprintBaseElement.h b/Engine/Plugins/PCG/Source/PCG/Public/Elements/Blueprint/PCGBlueprintBaseElement.h
--- a/Engine/Plugins/PCG/Source/PCG/Public/Elements/Blueprint/PCGBlueprintBaseElement.h
+++ b/Engine/Plugins/PCG/Source/PCG/Public/Elements/Blueprint/PCGBlueprintBaseElement.h
@@ -181,6 +181,8 @@ public:
 #endif
 
 protected:
+	void ThrowRequiresGameThreadException(FPCGContext* InContext, const FString& InFunctionName) const;
+
 #if WITH_EDITOR
 	PCG_API void OnDependencyChanged(UObject* Object, FPropertyChangedEvent& PropertyChangedEvent);
 	PCG_API void UpdateDependencies();
diff --git a/Engine/Plugins/PCG/Source/PCG/Public/Elements/Blueprint/PCGBlueprintDeprecatedElement.h b/Engine/Plugins/PCG/Source/PCG/Public/Elements/Blueprint/PCGBlueprintDeprecatedElement.h
--- a/Engine/Plugins/PCG/Source/PCG/Public/Elements/Blueprint/PCGBlueprintDeprecatedElement.h
+++ b/Engine/Plugins/PCG/Source/PCG/Public/Elements/Blueprint/PCGBlueprintDeprecatedElement.h
@@ -108,19 +108,19 @@ public:
 
 	/** Calls the PointLoopBody function on all points */
 	UFUNCTION(BlueprintCallable, BlueprintPure = false, Category = "PCG|Flow Control", meta = (HideSelfPin = "true"))
-	PCG_API void PointLoop(UPARAM(ref) FPCGContext& InContext, const UPCGPointData* InData, UPCGPointData*& OutData, UPCGPointData* OptionalOutData = nullptr) const;
+	PCG_API void PointLoop(UPARAM(ref) FPCGContext& InContext, const UPCGPointData* InData, UPCGPointData*& OutData, UPCGPointData* OptionalOutData = nullptr, bool bGameThread = false) const;
 
 	/** Calls the VariableLoopBody function on all points, each call can return a variable number of points */
 	UFUNCTION(BlueprintCallable, BlueprintPure = false, Category = "PCG|Flow Control", meta = (HideSelfPin = "true"))
-	PCG_API void VariableLoop(UPARAM(ref) FPCGContext& InContext, const UPCGPointData* InData, UPCGPointData*& OutData, UPCGPointData* OptionalOutData = nullptr) const;
+	PCG_API void VariableLoop(UPARAM(ref) FPCGContext& InContext, const UPCGPointData* InData, UPCGPointData*& OutData, UPCGPointData* OptionalOutData = nullptr, bool bGameThread = false) const;
 
 	/** Calls the NestedLoopBody function on all nested loop pairs (e.g. (o, i) for all o in Outer, i in Inner) */
 	UFUNCTION(BlueprintCallable, BlueprintPure = false, Category = "PCG|FLow Control", meta = (HideSelfPin = "true"))
-	PCG_API void NestedLoop(UPARAM(ref) FPCGContext& InContext, const UPCGPointData* InOuterData, const UPCGPointData* InInnerData, UPCGPointData*& OutData, UPCGPointData* OptionalOutData = nullptr) const;
+	PCG_API void NestedLoop(UPARAM(ref) FPCGContext& InContext, const UPCGPointData* InOuterData, const UPCGPointData* InInnerData, UPCGPointData*& OutData, UPCGPointData* OptionalOutData = nullptr, bool bGameThread = false) const;
 
 	/** Calls the IterationLoopBody a fixed number of times, optional parameters are used to potentially initialized the Out Data, but otherwise are used to remove the need to have variables */
 	UFUNCTION(BlueprintCallable, BlueprintPure = false, Category = "PCG|Flow Control", meta = (HideSelfPin = "true"))
-	PCG_API void IterationLoop(UPARAM(ref) FPCGContext& InContext, int64 NumIterations, UPCGPointData*& OutData, const UPCGSpatialData* OptionalA = nullptr, const UPCGSpatialData* OptionalB = nullptr, UPCGPointData* OptionalOutData = nullptr) const;
+	PCG_API void IterationLoop(UPARAM(ref) FPCGContext& InContext, int64 NumIterations, UPCGPointData*& OutData, const UPCGSpatialData* OptionalA = nullptr, const UPCGSpatialData* OptionalB = nullptr, UPCGPointData* OptionalOutData = nullptr, bool bGameThread = false) const;
 
 	/** Creates a random stream from the settings & source component */
 	UFUNCTION(BlueprintCallable, BlueprintPure = false, Category = "PCG|Random")
diff --git a/Engine/Plugins/PCG/Source/PCG/Public/Elements/Blueprint/PCGBlueprintPointProcessorElement.h b/Engine/Plugins/PCG/Source/PCG/Public/Elements/Blueprint/PCGBlueprintPointProcessorElement.h
--- a/Engine/Plugins/PCG/Source/PCG/Public/Elements/Blueprint/PCGBlueprintPointProcessorElement.h
+++ b/Engine/Plugins/PCG/Source/PCG/Public/Elements/Blueprint/PCGBlueprintPointProcessorElement.h
@@ -27,7 +27,7 @@ public:
 	 * @param OutputPointData	PointLoop output point data. Can be modified.
 	 */
 	UFUNCTION(BlueprintCallable, BlueprintPure = false, Category = "PCG|Flow Control", meta = (HideSelfPin = "true"))
-	PCG_API void PointLoop(const UPCGBasePointData* InputPointData, UPCGBasePointData*& OutputPointData);
+	PCG_API void PointLoop(const UPCGBasePointData* InputPointData, UPCGBasePointData*& OutputPointData, bool bGameThread = false);
 
 	/**
 	 * Initialization of the PointLoop output point data. This is called once before the PointLoopBody calls.
@@ -56,7 +56,7 @@ public:
 	 * @param OutputPointData	IterationLoop output point data. Can be modified.
 	 */
 	UFUNCTION(BlueprintCallable, BlueprintPure = false, Category = "PCG|Flow Control", meta = (HideSelfPin = "true"))
-	PCG_API void IterationLoop(int32 NumIterations, const UPCGBasePointData* InputPointDataA, const UPCGBasePointData* InputPointDataB, UPCGBasePointData*& OutputPointData);
+	PCG_API void IterationLoop(int32 NumIterations, const UPCGBasePointData* InputPointDataA, const UPCGBasePointData* InputPointDataB, UPCGBasePointData*& OutputPointData, bool bGameThread = false);
 
 	/** 
 	 * Initialization of the IterationLoop output point data. This is called once before the IterationLoopBody calls. 
