[PCG] Fix update pin for subgraph nodes

The subgraph is expected to be in a graph instance, therefore, we should not teardown the callbacks, as the recipient of those callbacks are not changing, the graph inside that instance is.
Also InitializeCachedOverridableParams is done in OnSubgraphChanged, so no need to do it again there.

#rnx

diff --git a/Engine/Plugins/PCG/Source/PCG/Private/PCGSubgraph.cpp b/Engine/Plugins/PCG/Source/PCG/Private/PCGSubgraph.cpp
--- a/Engine/Plugins/PCG/Source/PCG/Private/PCGSubgraph.cpp
+++ b/Engine/Plugins/PCG/Source/PCG/Private/PCGSubgraph.cpp
@@ -115,18 +115,7 @@ void UPCGBaseSubgraphSettings::PostEditImport()
 
 void UPCGBaseSubgraphSettings::SetSubgraph(UPCGGraphInterface* InGraph)
 {
-#if WITH_EDITOR
-	TeardownCallbacks();
-#endif // WITH_EDITOR
-
 	SetSubgraphInternal(InGraph);
-
-#if WITH_EDITOR
-	SetupCallbacks();
-#endif // WITH_EDITOR
-
-	// Also, reconstruct overrides
-	InitializeCachedOverridableParams(/*bReset=*/true);
 }
 
 void UPCGBaseSubgraphSettings::BeginDestroy()
@@ -203,10 +192,10 @@ void UPCGBaseSubgraphSettings::OnSubgraphChanged(UPCGGraphInterface* InGraph, EP
 			ChangeType |= EPCGChangeType::Settings;
 		}
 
-		OnSettingsChangedDelegate.Broadcast(this, ChangeType);
-
 		// Also rebuild the overrides
 		InitializeCachedOverridableParams(/*bReset=*/true);
+
+		OnSettingsChangedDelegate.Broadcast(this, ChangeType);
 	}
 }
 #endif // WITH_EDITOR
