[PCG] Fixed bad memory usage in split spline.

#rnx

diff --git a/Engine/Plugins/PCG/Source/PCG/Private/Elements/PCGSplitSplines.cpp b/Engine/Plugins/PCG/Source/PCG/Private/Elements/PCGSplitSplines.cpp
--- a/Engine/Plugins/PCG/Source/PCG/Private/Elements/PCGSplitSplines.cpp
+++ b/Engine/Plugins/PCG/Source/PCG/Private/Elements/PCGSplitSplines.cpp
@@ -446,7 +446,8 @@ bool FPCGSplitSplineElement::ExecuteInternal(FPCGContext* Context) const
 		// For open splines: we'll add a starting key (assumedly 0 in most cases) and a finishing key so we have proper segments [0 -> split key] ... [split key -> end]
 		if (bIsClosed)
 		{
-			SplitKeys.Add(SplitKeys[0]);
+			const double FirstSplitKey = SplitKeys[0];
+			SplitKeys.Add(FirstSplitKey);
 		}
 		else
 		{
