[PCG] PCG Builder Volume
- Editor only actor used to limit the bounds of generation when using the PCG Builder commandlet

diff --git a/Engine/Plugins/PCG/Source/PCG/Private/Builder/PCGBuilderVolume.cpp b/Engine/Plugins/PCG/Source/PCG/Private/Builder/PCGBuilderVolume.cpp
new file mode 100644
--- /dev/null
+++ b/Engine/Plugins/PCG/Source/PCG/Private/Builder/PCGBuilderVolume.cpp
@@ -0,0 +1,17 @@
+// Copyright Epic Games, Inc. All Rights Reserved.
+
+#include "Builder/PCGBuilderVolume.h"
+#include "Engine/CollisionProfile.h"
+#include "Components/BrushComponent.h"
+
+#include UE_INLINE_GENERATED_CPP_BY_NAME(PCGBuilderVolume)
+
+APCGBuilderVolume::APCGBuilderVolume()
+	: Super(FObjectInitializer::Get().DoNotCreateDefaultSubobject(TEXT("Sprite")))
+{
+#if WITH_EDITORONLY_DATA
+	bIsSpatiallyLoaded = false;
+#endif
+
+	GetBrushComponent()->SetCollisionProfileName(UCollisionProfile::NoCollision_ProfileName);
+}
diff --git a/Engine/Plugins/PCG/Source/PCG/Public/Builder/PCGBuilderVolume.h b/Engine/Plugins/PCG/Source/PCG/Public/Builder/PCGBuilderVolume.h
new file mode 100644
--- /dev/null
+++ b/Engine/Plugins/PCG/Source/PCG/Public/Builder/PCGBuilderVolume.h
@@ -0,0 +1,23 @@
+// Copyright Epic Games, Inc. All Rights Reserved.
+#pragma once
+
+#include "CoreMinimal.h"
+#include "GameFramework/Volume.h"
+#include "PCGBuilderVolume.generated.h"
+
+UCLASS(DisplayName = "PCG Builder Volume", hidecategories = (Actor, Advanced, Display, Events, Object, Attachment, Info, Input, Blueprint, Layers, Tags, Replication), MinimalAPI)
+class APCGBuilderVolume : public AVolume
+{
+	GENERATED_BODY()
+
+public:
+	PCG_API APCGBuilderVolume();
+
+	virtual bool IsEditorOnly() const final { return true; }
+
+#if WITH_EDITOR
+	virtual bool ActorTypeSupportsDataLayer() const final { return false; }
+	virtual bool ActorTypeSupportsExternalDataLayer() const final { return false; }
+	virtual bool CanChangeIsSpatiallyLoadedFlag() const { return false; }
+#endif
+};
diff --git a/Engine/Plugins/PCG/Source/PCGEditor/Private/WorldPartitionBuilder/PCGWorldPartitionBuilder.cpp b/Engine/Plugins/PCG/Source/PCGEditor/Private/WorldPartitionBuilder/PCGWorldPartitionBuilder.cpp
--- a/Engine/Plugins/PCG/Source/PCGEditor/Private/WorldPartitionBuilder/PCGWorldPartitionBuilder.cpp
+++ b/Engine/Plugins/PCG/Source/PCGEditor/Private/WorldPartitionBuilder/PCGWorldPartitionBuilder.cpp
@@ -10,6 +10,7 @@
 #include "PCGSubsystem.h"
 #include "SPCGBuilderDialog.h"
 
+#include "Builder/PCGBuilderVolume.h"
 #include "Grid/PCGPartitionActor.h"
 #include "Helpers/PCGHelpers.h"
 
@@ -112,6 +113,8 @@ struct FPCGWorldPartitionBuilder
 	static void SetChainedDispatchToLocalComponents(UWorld* InWorld, bool bInChainedDispatch);
 
 	static void UpdateEditingMode(UPCGComponent* InComponent);
+
+	static FBox GetPCGBuilderVolumeBounds(UWorld* InWorld);
 private:
 	static void CollectSubGridComponentsToGenerateInternal(UPCGSubsystem* InPCGSubsystem, int32 InGridSizeIndex, const PCGHiGenGrid::FSizeArray& InGridSizes, UPCGComponent* InOriginalComponent, APCGPartitionActor* InPartitionActor, TArray<UPCGComponent*>& OutAdditionalComponentsToGenerate, TMap<UPCGComponent*, UPCGComponent*>& OutDependencies);
 
@@ -588,6 +591,12 @@ bool UPCGWorldPartitionBuilder::PreRun(UWorld* World, FPackageSourceControlHelpe
 
 	if (GetLoadingMode() == ELoadingMode::IterativeCells2D)
 	{
+		FBox PCGBuilderVolumeBounds = FPCGWorldPartitionBuilder::GetPCGBuilderVolumeBounds(World);
+		if (PCGBuilderVolumeBounds.IsValid)
+		{
+			IterativeWorldBounds = PCGBuilderVolumeBounds;
+		}
+		
 		// Disable Source Control until Post Run
 		PreviousProviderName = ISourceControlModule::Get().GetProvider().GetName();
 		if (PreviousProviderName != NAME_None)
@@ -1116,7 +1125,7 @@ bool UPCGWorldPartitionBuilder::SavePerCellPackages(UWorld* InWorld, FPackageSou
 		return false;
 	}
 
-	UE_LOG(LogPCGWorldPartitionBuilder, Display, TEXT("Cell package changes: %d modified, %d new, %d deleted."), CellNewPackageCount, CellModifiedPackageCount, CellDeletedPackageCount);
+	UE_LOG(LogPCGWorldPartitionBuilder, Display, TEXT("Cell package changes: %d modified, %d new, %d deleted."), CellModifiedPackageCount, CellNewPackageCount, CellDeletedPackageCount);
 	return true;
 }
 
@@ -1656,6 +1665,42 @@ void FPCGWorldPartitionBuilder::UpdateEditingMode(UPCGComponent* InComponent)
 	}
 }
 
+FBox FPCGWorldPartitionBuilder::GetPCGBuilderVolumeBounds(UWorld* InWorld)
+{
+	FBox PCGBuilderVolumeBounds(ForceInit);
+
+	check(InWorld);
+
+	// Override the minimap bounds if world partition minimap volumes exists
+	for (TActorIterator<APCGBuilderVolume> It(InWorld); It; ++It)
+	{
+		if (APCGBuilderVolume* BuilderVolume = *It)
+		{
+			if (FBox Bounds = BuilderVolume->GetBounds().GetBox(); Bounds.IsValid)
+			{
+				PCGBuilderVolumeBounds += Bounds;
+				UE_LOG(LogPCGWorldPartitionBuilder, Display, TEXT("Found PCGBuilderVolume actor '%s' label '%s' bounds : %s"),
+					*BuilderVolume->GetName(),
+					*BuilderVolume->GetActorNameOrLabel(),
+					*Bounds.ToString());
+			}
+			else
+			{
+				UE_LOG(LogPCGWorldPartitionBuilder, Warning, TEXT("Found PCGBuilderVolume actor '%s' label '%s' with invalid bounds (ignored)"),
+					*BuilderVolume->GetName(),
+					*BuilderVolume->GetActorNameOrLabel());
+			}
+		}
+	}
+
+	if (PCGBuilderVolumeBounds.IsValid)
+	{
+		UE_LOG(LogPCGWorldPartitionBuilder, Display, TEXT("PCG builder bounds '%s'"), *PCGBuilderVolumeBounds.ToString());
+	}
+
+	return PCGBuilderVolumeBounds;
+}
+
 TArray<UPackage*> FPCGWorldPartitionBuilder::GetDirtyPackages(const UWorld* InWorld, const TSet<FString>& AllowedAssetPackages)
 {
 	TArray<UPackage*> DirtyPackages;
