[PCG] [Bughawk] comment fixes (trivial)

#rnx

diff --git a/Engine/Plugins/PCG/Source/PCG/Public/Elements/Blueprint/PCGBlueprintPointProcessorElement.h b/Engine/Plugins/PCG/Source/PCG/Public/Elements/Blueprint/PCGBlueprintPointProcessorElement.h
--- a/Engine/Plugins/PCG/Source/PCG/Public/Elements/Blueprint/PCGBlueprintPointProcessorElement.h
+++ b/Engine/Plugins/PCG/Source/PCG/Public/Elements/Blueprint/PCGBlueprintPointProcessorElement.h
@@ -25,6 +25,7 @@ public:
 	 * Calls the PointLoopInitialize once and PointLoopBody a fixed number of times.
 	 * @param InputPointData	Input point data. Constant, must not be modified.
 	 * @param OutputPointData	PointLoop output point data. Can be modified.
+	 * @param bGameThread		If true, loop will run single-thread on the game thread.
 	 */
 	UFUNCTION(BlueprintCallable, BlueprintPure = false, Category = "PCG|Flow Control", meta = (HideSelfPin = "true"))
 	PCG_API void PointLoop(const UPCGBasePointData* InputPointData, UPCGBasePointData*& OutputPointData, bool bGameThread = false);
@@ -54,6 +55,7 @@ public:
 	 * @param InputPointDataA	Optional input point data, can be null. Constant, must not be modified.
 	 * @param InputPointDataB	Optional input point data, can be null. Constant, must not be modified.
 	 * @param OutputPointData	IterationLoop output point data. Can be modified.
+	 * @param bGameThread		If true, loop will run single-thread on the game thread.
 	 */
 	UFUNCTION(BlueprintCallable, BlueprintPure = false, Category = "PCG|Flow Control", meta = (HideSelfPin = "true"))
 	PCG_API void IterationLoop(int32 NumIterations, const UPCGBasePointData* InputPointDataA, const UPCGBasePointData* InputPointDataB, UPCGBasePointData*& OutputPointData, bool bGameThread = false);
diff --git a/Engine/Plugins/PCG/Source/PCGEditor/Private/WorldPartitionBuilder/PCGWorldPartitionBuilder.cpp b/Engine/Plugins/PCG/Source/PCGEditor/Private/WorldPartitionBuilder/PCGWorldPartitionBuilder.cpp
--- a/Engine/Plugins/PCG/Source/PCGEditor/Private/WorldPartitionBuilder/PCGWorldPartitionBuilder.cpp
+++ b/Engine/Plugins/PCG/Source/PCGEditor/Private/WorldPartitionBuilder/PCGWorldPartitionBuilder.cpp
@@ -1671,7 +1671,7 @@ FBox FPCGWorldPartitionBuilder::GetPCGBuilderVolumeBounds(UWorld* InWorld)
 
 	check(InWorld);
 
-	// Override the minimap bounds if world partition minimap volumes exists
+	// Override the builder bounds if pcg builder volumes exists
 	for (TActorIterator<APCGBuilderVolume> It(InWorld); It; ++It)
 	{
 		if (APCGBuilderVolume* BuilderVolume = *It)
