[PCG] Fixed issues with loop key offset confusion vs loop position

#rnx

diff --git a/Engine/Plugins/PCG/Source/PCG/Private/Data/PCGSplineStruct.cpp b/Engine/Plugins/PCG/Source/PCG/Private/Data/PCGSplineStruct.cpp
--- a/Engine/Plugins/PCG/Source/PCG/Private/Data/PCGSplineStruct.cpp
+++ b/Engine/Plugins/PCG/Source/PCG/Private/Data/PCGSplineStruct.cpp
@@ -468,8 +468,9 @@ void FPCGSplineStruct::SetLoopKeyOffset(float LoopKeyOffset)
 	{
 		const bool bStationaryEndpoints = false;
 		const bool bLoopPositionOverride = true;
+		const float LoopPosition = LoopKeyOffset + (WarninglessSplineCurves().Position.Points.IsEmpty() ? 0.0f : WarninglessSplineCurves().Position.Points.Last().InVal);
 
-		FSpline::FUpdateSplineParams Params{ bClosedLoop, bStationaryEndpoints, ReparamStepsPerSegment, bLoopPositionOverride, LoopKeyOffset };
+		FSpline::FUpdateSplineParams Params{ bClosedLoop, bStationaryEndpoints, ReparamStepsPerSegment, bLoopPositionOverride, LoopPosition };
 		WarninglessSplineCurves().UpdateSpline(Params.bClosedLoop, Params.bStationaryEndpoints, Params.ReparamStepsPerSegment, Params.bLoopPositionOverride, Params.LoopPosition, Params.Scale3D);
 		Spline.UpdateSpline(Params);
 	}
diff --git a/Engine/Plugins/PCG/Source/PCG/Private/Elements/PCGSplitSplines.cpp b/Engine/Plugins/PCG/Source/PCG/Private/Elements/PCGSplitSplines.cpp
--- a/Engine/Plugins/PCG/Source/PCG/Private/Elements/PCGSplitSplines.cpp
+++ b/Engine/Plugins/PCG/Source/PCG/Private/Elements/PCGSplitSplines.cpp
@@ -120,7 +120,7 @@ namespace PCGSplitSpline
 			LastControlPointIndexInSegment += SplinePoints.Num();
 		}
 
-		const double LoopKeyLength = Spline->IsClosed() ? (Spline->GetLoopKeyOffset() - SplinePoints.Last().InputKey) : 1.0;
+		const double LoopKeyLength = Spline->IsClosed() ? Spline->GetLoopKeyOffset() : 1.0;
 
 		const double FirstControlKey = SplinePoints[FirstControlPointIndexInSegment].InputKey;
 		const double PreviousControlKey = FirstControlPointIndexInSegment > 0 ? SplinePoints[FirstControlPointIndexInSegment - 1].InputKey : FirstControlKey - LoopKeyLength;
@@ -393,7 +393,7 @@ bool FPCGSplitSplineElement::ExecuteInternal(FPCGContext* Context) const
 
 		// Remove values that are outside of the spline
 		const float FirstKey = SplinePoints.IsEmpty() ? 0.0f : SplinePoints[0].InputKey;
-		const float LastKey = SplinePoints.IsEmpty() ? 1.0f : (bIsClosed ? SplineData->GetLoopKeyOffset() : SplinePoints.Last().InputKey);
+		const float LastKey = SplinePoints.IsEmpty() ? 1.0f : (SplinePoints.Last().InputKey + (bIsClosed ? SplineData->GetLoopKeyOffset() : 0.0f));
 
 		// Closed loops should properly cut at key = 0 and up to the last key + 1.
 		if (bIsClosed)
