[PCG] Detect if a PCG Component was already generated when loaded (bGeneratedOffline)
- Avoid generating PCGPartitionActor components on load if their original component was generated on demand and bGeneratedOffline. (log a warning instead)

diff --git a/Engine/Plugins/PCG/Source/PCG/Private/Grid/PCGPartitionActor.cpp b/Engine/Plugins/PCG/Source/PCG/Private/Grid/PCGPartitionActor.cpp
--- a/Engine/Plugins/PCG/Source/PCG/Private/Grid/PCGPartitionActor.cpp
+++ b/Engine/Plugins/PCG/Source/PCG/Private/Grid/PCGPartitionActor.cpp
@@ -220,9 +220,15 @@ void APCGPartitionActor::BeginPlay()
 			UPCGComponent* LocalComponent = It.Value;
 			if (OriginalComponent && LocalComponent)
 			{
+				// If OriginalComponent was already generated when loading then our LocalComponent should be in the same situation, we want to avoid generating and instead warn.
+				if (OriginalComponent->GenerationTrigger == EPCGComponentGenerationTrigger::GenerateOnDemand && OriginalComponent->IsGeneratedOffline())
+				{
+					LocalComponent->GenerationTrigger = OriginalComponent->GenerationTrigger;
+					UE_CLOG(!LocalComponent->bGenerated, LogPCG, Warning, TEXT("PartitionActor '%s' has a Local component '%s' which should have been generated offline"), *GetName(), *LocalComponent->GetName());
+				}
 				// If we have an original component that is generated (or generating), this one is automatically generated => GenerateOnLoad.
 				// But if its runtime generated then it's handled by the runtime generation scheduler.
-				if ((OriginalComponent->bGenerated || OriginalComponent->IsGenerating()) && ensure(!OriginalComponent->IsManagedByRuntimeGenSystem()))
+				else if ((OriginalComponent->bGenerated || OriginalComponent->IsGenerating()) && ensure(!OriginalComponent->IsManagedByRuntimeGenSystem())) 
 				{
 					LocalComponent->GenerationTrigger = EPCGComponentGenerationTrigger::GenerateOnLoad;
 				}
diff --git a/Engine/Plugins/PCG/Source/PCG/Private/PCGComponent.cpp b/Engine/Plugins/PCG/Source/PCG/Private/PCGComponent.cpp
--- a/Engine/Plugins/PCG/Source/PCG/Private/PCGComponent.cpp
+++ b/Engine/Plugins/PCG/Source/PCG/Private/PCGComponent.cpp
@@ -1902,7 +1902,6 @@ void UPCGComponent::BeginPlay()
 	if (OriginalComponent && bActivated && !bAlreadyGenerated && GenerationTrigger == EPCGComponentGenerationTrigger::GenerateOnLoad)
 	{
 		GenerateInternal(/*bForce=*/false, EPCGHiGenGrid::Uninitialized, EPCGComponentGenerationTrigger::GenerateOnLoad, {});
-		bRuntimeGenerated = true;
 	}
 }
 
@@ -2264,6 +2263,9 @@ void UPCGComponent::PostLoad()
 		SchedulingPolicy->SetShouldDisplayProperties(IsManagedByRuntimeGenSystem());
 #endif
 	}
+
+	// keep track if component was already generated when it got loaded
+	bGeneratedOffline = bGenerated;
 }
 
 #if WITH_EDITOR
@@ -4325,6 +4327,7 @@ FPCGComponentInstanceData::FPCGComponentInstanceData(const UPCGComponent* InSour
 	, CurrentEditingMode(InSourceComponent->CurrentEditingMode)
 	, PreviousEditingMode(InSourceComponent->PreviousEditingMode)
 	, bForceGenerateOnBPAddedToWorld(InSourceComponent->bForceGenerateOnBPAddedToWorld)
+	, bGeneratedOffline(InSourceComponent->bGeneratedOffline)
 #endif
 {
 #if WITH_EDITOR
@@ -4379,6 +4382,7 @@ void FPCGComponentInstanceData::ApplyToComponent(UActorComponent* Component, con
 		PCGComponent->bForceGenerateOnBPAddedToWorld = bForceGenerateOnBPAddedToWorld;
 		PCGComponent->PreviousEditingMode = PreviousEditingMode;	
 		PCGComponent->CurrentEditingMode = CurrentEditingMode;
+		PCGComponent->bGeneratedOffline = bGeneratedOffline;
 
 		for (const FDynamicTrackedKeyInstanceData& DynamicTrackedKeyInstanceData : DynamicallyTrackedKeysToSettings)
 		{
diff --git a/Engine/Plugins/PCG/Source/PCG/Public/PCGComponent.h b/Engine/Plugins/PCG/Source/PCG/Public/PCGComponent.h
--- a/Engine/Plugins/PCG/Source/PCG/Public/PCGComponent.h
+++ b/Engine/Plugins/PCG/Source/PCG/Public/PCGComponent.h
@@ -384,6 +384,7 @@ public:
 	UPROPERTY(BlueprintReadOnly, VisibleAnywhere, Category = Debug, NonTransactional, meta = (NoResetToDefault))
 	bool bGenerated = false;
 
+	UE_DEPRECATED(5.8, "No longer used.")
 	UPROPERTY(NonPIEDuplicateTransient)
 	bool bRuntimeGenerated = false;
 
@@ -397,6 +398,9 @@ public:
 	bool IsGenerating() const { return CurrentGenerationTask != InvalidPCGTaskId; }
 	bool IsCleaningUp() const { return CurrentCleanupTask != InvalidPCGTaskId; }
 
+	/** Return if this component was already generated when it was loaded */
+	bool IsGeneratedOffline() const { return bGeneratedOffline; }
+
 	/** Returns task ids to do internal chaining */
 	FPCGTaskId GetGenerationTaskId() const { return CurrentGenerationTask; }
 	FPCGTaskId GetCleanupTaskId() const { return CurrentCleanupTask; }
@@ -750,6 +754,9 @@ private:
 	// When doing a cleanup, locking resource modification. Used as sentinel.
 	bool GeneratedResourcesInaccessible = false;
 
+	// Flag to indicate whether this component was already generated on load.
+	bool bGeneratedOffline = false;
+
 	UPROPERTY(VisibleInstanceOnly, Category = Debug)
 	FBox LastGeneratedBounds = FBox(EForceInit::ForceInit);
 
@@ -875,6 +882,10 @@ protected:
 	UPROPERTY()
 	bool bForceGenerateOnBPAddedToWorld = false;
 
+	// UPCGComponent::bGeneratedOffline is not a property so we need to capture/apply it ourselves
+	UPROPERTY()
+	bool bGeneratedOffline = false;
+
 	// UPCGComponent::DynamicallyTrackedKeysToSettings is not a property so we need to capture/apply it ourselves.
 	// To avoid having to write custom serialization code the map is converted into an array for serialization.
 	UPROPERTY()
