[PCG] Expose the actor selection options on get landscape data

#rnx

diff --git a/Engine/Plugins/PCG/Source/PCG/Private/Elements/PCGTypedGetter.cpp b/Engine/Plugins/PCG/Source/PCG/Private/Elements/PCGTypedGetter.cpp
--- a/Engine/Plugins/PCG/Source/PCG/Private/Elements/PCGTypedGetter.cpp
+++ b/Engine/Plugins/PCG/Source/PCG/Private/Elements/PCGTypedGetter.cpp
@@ -17,7 +17,6 @@
 UPCGGetLandscapeSettings::UPCGGetLandscapeSettings()
 {
 	Mode = EPCGGetDataFromActorMode::ParseActorComponents;
-	ActorSelector.bShowActorFilter = false;
 	ActorSelector.bIncludeChildren = false;
 	ActorSelector.bShowActorSelectionClass = false;
 	ActorSelector.bSelectMultiple = true;
