[PCG] Extract the PCG Action row to be usable by other execution source types

diff --git a/Engine/Plugins/PCG/Source/PCG/Private/Graph/PCGGraphExecutor.cpp b/Engine/Plugins/PCG/Source/PCG/Private/Graph/PCGGraphExecutor.cpp
--- a/Engine/Plugins/PCG/Source/PCG/Private/Graph/PCGGraphExecutor.cpp
+++ b/Engine/Plugins/PCG/Source/PCG/Private/Graph/PCGGraphExecutor.cpp
@@ -3235,7 +3235,9 @@ namespace PCGGraphExecutor
 				PCGGraphExecutionLogging::LogGridLinkageTaskExecuteRetrieveScheduleGraph(InContext, ExecSourceWithData, InResourceKey);
 
 				// Wake up this task after graph has generated.
-				const FPCGTaskId GraphTaskId = ExecSourceWithData->GetExecutionState().Generate();
+				IPCGGraphExecutionState::FGenerateParams Params;
+				Params.bEvenIfAlreadyGenerated = true;
+				const FPCGTaskId GraphTaskId = ExecSourceWithData->GetExecutionState().Generate(Params);
 
 				// Update state and go to paused.
 				InContext->DynamicDependencies.Add(GraphTaskId);
diff --git a/Engine/Plugins/PCG/Source/PCG/Private/PCGComponentExecutionState.cpp b/Engine/Plugins/PCG/Source/PCG/Private/PCGComponentExecutionState.cpp
--- a/Engine/Plugins/PCG/Source/PCG/Private/PCGComponentExecutionState.cpp
+++ b/Engine/Plugins/PCG/Source/PCG/Private/PCGComponentExecutionState.cpp
@@ -211,13 +211,14 @@ FPCGTaskId FPCGComponentExecutionState::Generate(const FGenerateParams& InGenera
 
 	const EPCGHiGenGrid Grid = InGenerateParams.bGenerateLocalSources ? EPCGHiGenGrid::Uninitialized : GetGenerationGrid();
 	const EPCGComponentGenerationTrigger GenTrigger = IsManagedByRuntimeGenSystem() ? EPCGComponentGenerationTrigger::GenerateAtRuntime : EPCGComponentGenerationTrigger::GenerateOnDemand;
-	return Component->GenerateLocalGetTaskId(GenTrigger, /*bForce=*/true, Grid, InGenerateParams.Dependencies);
+	return Component->GenerateLocalGetTaskId(GenTrigger, InGenerateParams.bEvenIfAlreadyGenerated, Grid, InGenerateParams.Dependencies);
 }
 
 // UE_DEPRECATED(5.8)
 FPCGTaskId FPCGComponentExecutionState::GenerateLocalGetTaskId(EPCGHiGenGrid Grid)
 {
 	FGenerateParams GenerateParams;
+	GenerateParams.bEvenIfAlreadyGenerated = true;
 	GenerateParams.bGenerateLocalSources = Grid == EPCGHiGenGrid::Uninitialized;
 
 	return Generate(GenerateParams);
@@ -233,12 +234,11 @@ FPCGTaskId FPCGComponentExecutionState::Cleanup(const FCleanupParams& InCleanupP
 		return InvalidPCGTaskId;
 	}
 
-	// @todo_pcg: There is a difference in behavior for CleanupLocal vs CleanupLocalImmediate, where CleanupLocal cannot be
-	// called on local components directly, and will ALWAYS cleanup local components when called on an original component.
-	if (InCleanupParams.bCleanupLocalSources || IsLocalSource())
+	// Partitioned components will always cleanup their LCs, so warn if that option is toggled off
+	if (!InCleanupParams.bCleanupLocalSources && IsPartitioned())
 	{
-		UE_LOG(LogPCG, Warning, TEXT("Tried to schedule cleanup on a local execution source. This is not supported, call Cleanup() with bImmediate=true instead."))
-		return InvalidPCGTaskId;
+		UE_LOG(LogPCG, Warning, TEXT("Tried to schedule cleanup on a partitioned original execution source with bCleanupLocalSources=false, but that is not supported. Local components will be cleaned up."));
+		// doesn't need to return, can still clean up the OC.
 	}
 
 	return Component->CleanupLocal(InCleanupParams.bReleaseManagedResources, InCleanupParams.Dependencies);
diff --git a/Engine/Plugins/PCG/Source/PCG/Private/RuntimeGen/PCGRuntimeGenScheduler.cpp b/Engine/Plugins/PCG/Source/PCG/Private/RuntimeGen/PCGRuntimeGenScheduler.cpp
--- a/Engine/Plugins/PCG/Source/PCG/Private/RuntimeGen/PCGRuntimeGenScheduler.cpp
+++ b/Engine/Plugins/PCG/Source/PCG/Private/RuntimeGen/PCGRuntimeGenScheduler.cpp
@@ -1204,6 +1204,7 @@ void FPCGRuntimeGenScheduler::TickScheduleGeneration(TMap<FGridGenerationKey, do
 
 				// Force to refresh if the component is already generated.
 				IPCGGraphExecutionState::FGenerateParams GenerateParams;
+				GenerateParams.bEvenIfAlreadyGenerated = true;
 				GenerateParams.Dependencies = MoveTemp(Dependencies);
 
 				LocalComponent->GetExecutionState().Generate(GenerateParams);
diff --git a/Engine/Plugins/PCG/Source/PCG/Public/PCGGraphExecutionStateInterface.h b/Engine/Plugins/PCG/Source/PCG/Public/PCGGraphExecutionStateInterface.h
--- a/Engine/Plugins/PCG/Source/PCG/Public/PCGGraphExecutionStateInterface.h
+++ b/Engine/Plugins/PCG/Source/PCG/Public/PCGGraphExecutionStateInterface.h
@@ -35,6 +35,9 @@ class IPCGGraphExecutionState
 public:
 	struct FGenerateParams
 	{
+		/** If we should force generation even when the source is already generated. */
+		bool bEvenIfAlreadyGenerated = false;
+
 		/** If the execution source is partitioned, will forward the generate call to its registered local execution sources. */
 		bool bGenerateLocalSources = true;
 
diff --git a/Engine/Plugins/PCG/Source/PCGEditor/Private/Details/PCGComponentDetails.cpp b/Engine/Plugins/PCG/Source/PCGEditor/Private/Details/PCGComponentDetails.cpp
--- a/Engine/Plugins/PCG/Source/PCGEditor/Private/Details/PCGComponentDetails.cpp
+++ b/Engine/Plugins/PCG/Source/PCGEditor/Private/Details/PCGComponentDetails.cpp
@@ -1,73 +1,37 @@
 // Copyright Epic Games, Inc. All Rights Reserved.
 
 #include "Details/PCGComponentDetails.h"
+
 #include "PCGComponent.h"
-#include "PCGEditorStyle.h"
 #include "PCGGraph.h"
 #include "PCGSubsystem.h"
-#include "Framework/Application/SlateApplication.h"
-#include "Widgets/Text/STextBlock.h"
-#include "Widgets/Images/SImage.h"
-#include "Widgets/Input/SButton.h"
+#include "Widgets/SPCGExecutionSourceActionWidget.h"
+
+#include "DetailCategoryBuilder.h"
 #include "DetailLayoutBuilder.h"
 #include "DetailWidgetRow.h"
-#include "Modules/ModuleManager.h"
-#include "LevelEditor.h"
-#include "LevelEditorSubsystem.h"
 #include "Algo/NoneOf.h"
-#include "Elements/Framework/TypedElementSelectionSet.h"
+#include "GameFramework/Actor.h"
 
 #define LOCTEXT_NAMESPACE "PCGComponentDetails"
 
-static TAutoConsoleVariable<bool> CVarBroadcastComponentChanged(
-		TEXT("pcg.ComponentDetails.BroadcastComponentChanged"),
-		false,
-		TEXT("Controls whether we should broadcast component changes to the level editor"));
-
-namespace PCGComponentDetails
+TArray<TWeakInterfacePtr<IPCGGraphExecutionSource>> FPCGComponentDetails::GatherExecutionSourcesFromSelection(const TArray<TWeakObjectPtr<UObject>>& InObjectSelected) const
 {
-	void TriggerSelectionUpdate(UPCGComponent* Component, FPCGTaskId TaskId)
+	TArray<TWeakInterfacePtr<IPCGGraphExecutionSource>> ExecutionSources;
+	for (const TWeakObjectPtr<UObject>& Object : InObjectSelected)
 	{
-		if (!GEditor || !Component || TaskId == InvalidPCGTaskId)
+		if (IPCGGraphExecutionSource* ExecutionSource = Cast<IPCGGraphExecutionSource>(Object.Get()))
 		{
-			return;
-		}
-
-		if (UPCGSubsystem* PCGSubsystem = Component->GetSubsystem())
-		{
-			PCGSubsystem->ScheduleGeneric([]()
-			{
-				if (GEditor && GEditor->GetEditorSubsystem<ULevelEditorSubsystem>())
-				{
-					// Since we might have been moving things around (component wise) we need to force a selection set change, 
-					// Which in turn will trigger a component visualizer update.
-					if (UTypedElementSelectionSet* SelectionSet = GEditor->GetEditorSubsystem<ULevelEditorSubsystem>()->GetSelectionSet())
-					{
-						SelectionSet->OnChanged().Broadcast(SelectionSet);
-					}
-				}
-
-				return true;
-			}, Component, { TaskId });
+			ExecutionSources.Add(ExecutionSource);
 		}
 	}
-} // namespace PCGComponentDetails
 
-TSharedRef<IDetailCustomization> FPCGComponentDetails::MakeInstance()
-{
-	return MakeShareable(new FPCGComponentDetails());
+	return ExecutionSources;
 }
 
-void FPCGComponentDetails::GatherPCGComponentsFromSelection(const TArray<TWeakObjectPtr<UObject>>& InObjectSelected)
+TSharedRef<IDetailCustomization> FPCGComponentDetails::MakeInstance()
 {
-	for (const TWeakObjectPtr<UObject>& Object : InObjectSelected)
-	{
-		UPCGComponent* Component = Cast<UPCGComponent>(Object.Get());
-		if (ensure(Component))
-		{
-			SelectedComponents.Add(Component);
-		}
-	}
+	return MakeShareable(new FPCGComponentDetails());
 }
 
 void FPCGComponentDetails::CustomizeDetails(IDetailLayoutBuilder& DetailBuilder)
@@ -78,7 +42,7 @@ void FPCGComponentDetails::CustomizeDetails(IDetailLayoutBuilder& DetailBuilder)
 	TArray<TWeakObjectPtr<UObject>> ObjectsBeingCustomized;
 	DetailBuilder.GetObjectsBeingCustomized(ObjectsBeingCustomized);
 
-	GatherPCGComponentsFromSelection(ObjectsBeingCustomized);
+	TArray<TWeakInterfacePtr<IPCGGraphExecutionSource>> ExecutionSources = GatherExecutionSourcesFromSelection(ObjectsBeingCustomized);
 
 	if (AddDefaultProperties())
 	{
@@ -94,9 +58,9 @@ void FPCGComponentDetails::CustomizeDetails(IDetailLayoutBuilder& DetailBuilder)
 		}
 	}
 
-	for (TWeakObjectPtr<UPCGComponent>& ComponentPtr : SelectedComponents)
+	for (TWeakInterfacePtr<IPCGGraphExecutionSource>& ExecutionSource : ExecutionSources)
 	{
-		if (UPCGComponent* Component = ComponentPtr.Get())
+		if (UPCGComponent* Component = Cast<UPCGComponent>(ExecutionSource.Get()))
 		{
 			AActor* ComponentOwner = Component->GetOwner();
 			if (!ComponentOwner || (ComponentOwner->IsInLevelInstance() && !ComponentOwner->IsInEditLevelInstance()))
@@ -111,10 +75,10 @@ void FPCGComponentDetails::CustomizeDetails(IDetailLayoutBuilder& DetailBuilder)
 	if (Algo::NoneOf(ObjectsBeingCustomized, [](const TWeakObjectPtr<UObject>& Object) { return !Object.IsValid() || Object.Get()->IsA<UPCGComponent>(); }))
 	{
 		TArray<UObject*> Objects;
-		Objects.Reserve(SelectedComponents.Num());
-		Algo::Transform(SelectedComponents, Objects, [](const TWeakObjectPtr<UPCGComponent>& Component) -> UObject*
+		Objects.Reserve(ExecutionSources.Num());
+		Algo::Transform(ExecutionSources, Objects, [](const TWeakInterfacePtr<IPCGGraphExecutionSource>& ExecutionSource) -> UObject*
 		{
-			return Component.IsValid() ? CastChecked<UObject>(Component.Get()->GetGraphInstance(), ECastCheckedType::NullAllowed) : nullptr;
+			return ExecutionSource.IsValid() ? CastChecked<UObject>(ExecutionSource.Get()->GetExecutionState().GetGraphInstance(), ECastCheckedType::NullAllowed) : nullptr;
 		});
 		PCGCategory.AddExternalObjectProperty(Objects, TEXT("Graph"));
 		PCGCategory.AddExternalObjectProperty(Objects, TEXT("ParametersOverrides"));
@@ -125,316 +89,9 @@ void FPCGComponentDetails::CustomizeDetails(IDetailLayoutBuilder& DetailBuilder)
 	NewRow.ValueContent()
 		.MaxDesiredWidth(120.f)
 		[
-			SNew(SHorizontalBox)
-			+ SHorizontalBox::Slot()
-			.AutoWidth()
-			.Padding(2.0f, 0.0f)
-			.VAlign(VAlign_Fill)
-			[
-				SNew(SButton)
-				.OnClicked(this, &FPCGComponentDetails::OnGenerateClicked)
-				.ToolTipText(FText::FromString("Generates graph data. \nCtrl + Click flushes the cache and force generates."))
-				.Visibility(this, &FPCGComponentDetails::GenerateButtonVisible)
-				[
-					SNew(SHorizontalBox)
-					+SHorizontalBox::Slot()
-					.Padding(0.0f, 0.0f, 6.0f, 0.0f)
-					.VAlign(VAlign_Center)
-					.AutoWidth()
-					[
-						SNew(SImage)
-						.DesiredSizeOverride(FVector2D(16,16))
-						.Image_Lambda([]() { return FSlateApplication::Get().GetModifierKeys().IsControlDown() ? FPCGEditorStyle::Get().GetBrush("PCG.Command.ForceRegenClearCache") : FPCGEditorStyle::Get().GetBrush("PCG.Command.ForceRegen"); })
-						.ColorAndOpacity(FSlateColor::UseForeground())
-					]
-					+SHorizontalBox::Slot()
-					.VAlign(VAlign_Center)
-					.AutoWidth()
-					[
-						SNew(STextBlock)
-						.Font(IDetailLayoutBuilder::GetDetailFont())
-						.Text_Lambda([]() { return FSlateApplication::Get().GetModifierKeys().IsControlDown() ? LOCTEXT("ForceRegenerateButton", "Force Generate") : LOCTEXT("GenerateButton", "Generate"); })
-					]
-				]
-			]
-			+ SHorizontalBox::Slot()
-			.AutoWidth()
-			.Padding(2.0f, 0.0f)
-			.VAlign(VAlign_Fill)
-			[
-				SNew(SButton)
-				.OnClicked(this, &FPCGComponentDetails::OnCancelClicked)
-				.Visibility(this, &FPCGComponentDetails::CancelButtonVisible)
-				[
-					SNew(SHorizontalBox)
-					+ SHorizontalBox::Slot()
-					.Padding(0.0f, 0.0f, 6.0f, 0.0f)
-					.VAlign(VAlign_Center)
-					.AutoWidth()
-					[
-						SNew(SImage)
-						.DesiredSizeOverride(FVector2D(16, 16))
-						.Image(FPCGEditorStyle::Get().GetBrush("PCG.Command.StopRegen"))
-						.ColorAndOpacity(FSlateColor::UseForeground())
-					]
-					+ SHorizontalBox::Slot()
-					.VAlign(VAlign_Center)
-					.AutoWidth()
-					[
-						SNew(STextBlock)
-						.Font(IDetailLayoutBuilder::GetDetailFont())
-						.Text(LOCTEXT("CancelButton", "Cancel"))
-					]
-				]
-			]
-			+ SHorizontalBox::Slot()
-			.AutoWidth()
-			.VAlign(VAlign_Fill)
-			.Padding(2.0f, 0.0f)
-			[
-				SNew(SButton)
-				.VAlign(VAlign_Center)
-				.OnClicked(this, &FPCGComponentDetails::OnCleanupClicked)
-				.ToolTipText(FText::FromString("Cleans up graph data. \nCtrl + Click purges all components and attached actors tagged as created by PCG."))
-				.Visibility(this, &FPCGComponentDetails::CleanupButtonVisible)
-				[
-					SNew(STextBlock)
-					.Font(IDetailLayoutBuilder::GetDetailFont())
-					.Text(LOCTEXT("CleanupButton", "Cleanup"))
-				]
-			]
-			+ SHorizontalBox::Slot()
-			.AutoWidth()
-			.VAlign(VAlign_Fill)
-			.Padding(2.0f, 0.0f)
-			[
-				SNew(SButton)
-				.VAlign(VAlign_Center)
-				.OnClicked(this, &FPCGComponentDetails::OnRefreshClicked)
-				.Visibility(this, &FPCGComponentDetails::RefreshButtonVisible)
-				[
-					SNew(STextBlock)
-					.Font(IDetailLayoutBuilder::GetDetailFont())
-					.Text(LOCTEXT("RefreshButton", "Refresh"))
-				]
-			]
-			+ SHorizontalBox::Slot()
-			.AutoWidth()
-			.VAlign(VAlign_Fill)
-			.Padding(2.0f, 0.0f)
-			[
-				SNew(SButton)
-				.VAlign(VAlign_Center)
-				.OnClicked(this, &FPCGComponentDetails::OnClearPCGLinkClicked)
-				.Visibility(this, &FPCGComponentDetails::ClearPCGLinkButtonVisible)
-				[
-					SNew(STextBlock)
-					.Font(IDetailLayoutBuilder::GetDetailFont())
-					.Text(LOCTEXT("ClearPCGLinkButton", "Clear PCG Link"))
-				]
-			]
+			SNew(SPCGExecutionSourceActionWidget)
+			.ExecutionSources(ExecutionSources)
 		];
-
-	// Attach to generated delegate on the selected components
-	for (TWeakObjectPtr<UPCGComponent>& Component : SelectedComponents)
-	{
-		if (Component.IsValid())
-		{
-			Component.Get()->OnPCGGraphGeneratedDelegate.AddRaw(this, &FPCGComponentDetails::OnGraphChanged);
-			Component.Get()->OnPCGGraphCleanedDelegate.AddRaw(this, &FPCGComponentDetails::OnGraphChanged);
-		}
-	}
-}
-
-void FPCGComponentDetails::PendingDelete()
-{
-	// Detach from the generated delegate
-	for (TWeakObjectPtr<UPCGComponent>& Component : SelectedComponents)
-	{
-		if (Component.IsValid())
-		{
-			Component.Get()->OnPCGGraphGeneratedDelegate.RemoveAll(this);
-			Component.Get()->OnPCGGraphCleanedDelegate.RemoveAll(this);
-		}
-	}
-}
-
-EVisibility FPCGComponentDetails::GenerateButtonVisible() const
-{
-	for (const TWeakObjectPtr<UPCGComponent>& Component : SelectedComponents)
-	{
-		// If component is runtime generated then generate/cleanup is managed by the scheduler.
-		if (Component.IsValid() && !Component->IsGenerating() && !Component->IsManagedByRuntimeGenSystem())
-		{
-			return EVisibility::Visible;
-		}
-	}
-
-	return EVisibility::Collapsed;
-}
-
-EVisibility FPCGComponentDetails::CancelButtonVisible() const
-{
-	for (const TWeakObjectPtr<UPCGComponent>& Component : SelectedComponents)
-	{
-		// If component is runtime generated then generate/cleanup is managed by the scheduler.
-		if (Component.IsValid() && Component->IsGenerating() && !Component->IsManagedByRuntimeGenSystem())
-		{
-			return EVisibility::Visible;
-		}
-	}
-
-	return EVisibility::Collapsed;
-}
-
-EVisibility FPCGComponentDetails::CleanupButtonVisible() const
-{
-	for (const TWeakObjectPtr<UPCGComponent>& Component : SelectedComponents)
-	{
-		// If component is runtime generated then generate/cleanup is managed by the scheduler.
-		if (Component.IsValid() && !Component->IsManagedByRuntimeGenSystem())
-		{
-			return EVisibility::Visible;
-		}
-	}
-
-	return EVisibility::Collapsed;
-}
-
-EVisibility FPCGComponentDetails::RefreshButtonVisible() const
-{
-	for (const TWeakObjectPtr<UPCGComponent>& Component : SelectedComponents)
-	{
-		if (Component.IsValid() && Component->IsManagedByRuntimeGenSystem())
-		{
-			return EVisibility::Visible;
-		}
-	}
-
-	return EVisibility::Collapsed;
-}
-
-EVisibility FPCGComponentDetails::ClearPCGLinkButtonVisible() const
-{
-	for (const TWeakObjectPtr<UPCGComponent>& Component : SelectedComponents)
-	{
-		if (Component.IsValid() && !Component->IsManagedByRuntimeGenSystem())
-		{
-			return EVisibility::Visible;
-		}
-	}
-
-	return EVisibility::Collapsed;
-}
-
-
-FReply FPCGComponentDetails::OnGenerateClicked()
-{
-	const FModifierKeysState ModifierKeys = FSlateApplication::Get().GetModifierKeys();
-	const bool bIsControlDown = ModifierKeys.IsControlDown();
-
-	bool bFlushedCache = false;
-
-	for (TWeakObjectPtr<UPCGComponent>& SelectedComponentPtr : SelectedComponents)
-	{
-		UPCGComponent* PCGComponent = SelectedComponentPtr.Get();
-
-		if (IsValid(PCGComponent) && !PCGComponent->IsGenerating() && !PCGComponent->IsManagedByRuntimeGenSystem())
-		{
-			// Flush cache if needed
-			if (bIsControlDown && !bFlushedCache)
-			{
-				if (UPCGSubsystem* PCGSubsystem = PCGComponent->GetSubsystem())
-				{
-					PCGSubsystem->FlushCache();
-					bFlushedCache = true;
-				}
-			}
-
-			PCGComponent->Generate(/*bForce=*/bIsControlDown);
-			PCGComponentDetails::TriggerSelectionUpdate(PCGComponent, PCGComponent->GetGenerationTaskId());
-		}
-	}
-
-	return FReply::Handled();
-}
-
-FReply FPCGComponentDetails::OnCancelClicked()
-{
-	for (TWeakObjectPtr<UPCGComponent>& Component : SelectedComponents)
-	{
-		if (Component.IsValid() && Component->IsGenerating() && !Component->IsManagedByRuntimeGenSystem())
-		{
-			Component.Get()->CancelGeneration();
-		}
-	}
-
-	return FReply::Handled();
-}
-
-FReply FPCGComponentDetails::OnRefreshClicked()
-{
-	for (TWeakObjectPtr<UPCGComponent>& Component : SelectedComponents)
-	{
-		if (Component.IsValid() && Component->IsManagedByRuntimeGenSystem())
-		{
-			// Trigger the deepest refresh - re-initialize the PAs.
-			Component.Get()->Refresh(EPCGChangeType::Structural | EPCGChangeType::GenerationGrid, /*bCancelExistingRefresh=*/true);
-		}
-	}
-
-	return FReply::Handled();
-}
-
-FReply FPCGComponentDetails::OnClearPCGLinkClicked()
-{
-	for (TWeakObjectPtr<UPCGComponent>& Component : SelectedComponents)
-	{
-		if (Component.IsValid() && !Component->IsManagedByRuntimeGenSystem())
-		{
-			Component.Get()->ClearPCGLink();
-		}
-	}
-
-	return FReply::Handled();
-}
-
-void FPCGComponentDetails::OnGraphChanged(UPCGComponent* InComponent)
-{
-	if (!InComponent || !InComponent->GetOwner())
-	{
-		return;
-	}
-
-	// Notify editor that some components might have changed
-	if (CVarBroadcastComponentChanged.GetValueOnGameThread())
-	{
-		FLevelEditorModule& LevelEditor = FModuleManager::LoadModuleChecked<FLevelEditorModule>("LevelEditor");
-		LevelEditor.BroadcastComponentsEdited();
-	}
-}
-
-FReply FPCGComponentDetails::OnCleanupClicked()
-{
-	for (TWeakObjectPtr<UPCGComponent>& Component : SelectedComponents)
-	{
-		if (Component.IsValid() && !Component->IsManagedByRuntimeGenSystem())
-		{
-			FModifierKeysState ModifierKeys = FSlateApplication::Get().GetModifierKeys();
-			if (ModifierKeys.IsControlDown())
-			{
-				Component.Get()->CleanupLocalDeleteAllGeneratedObjects({});
-			}
-			else
-			{
-				Component.Get()->Cleanup();
-			}
-
-			PCGComponentDetails::TriggerSelectionUpdate(Component.Get(), Component->GetCleanupTaskId());
-		}
-	}
-
-	return FReply::Handled();
 }
 
 #undef LOCTEXT_NAMESPACE
diff --git a/Engine/Plugins/PCG/Source/PCGEditor/Private/Details/PCGVolumeDetails.cpp b/Engine/Plugins/PCG/Source/PCGEditor/Private/Details/PCGVolumeDetails.cpp
--- a/Engine/Plugins/PCG/Source/PCGEditor/Private/Details/PCGVolumeDetails.cpp
+++ b/Engine/Plugins/PCG/Source/PCGEditor/Private/Details/PCGVolumeDetails.cpp
@@ -10,8 +10,10 @@ TSharedRef<IDetailCustomization> FPCGVolumeDetails::MakeInstance()
 	return MakeShareable(new FPCGVolumeDetails());
 }
 
-void FPCGVolumeDetails::GatherPCGComponentsFromSelection(const TArray<TWeakObjectPtr<UObject>>& InObjectSelected)
+TArray<TWeakInterfacePtr<IPCGGraphExecutionSource>> FPCGVolumeDetails::GatherExecutionSourcesFromSelection(const TArray<TWeakObjectPtr<UObject>>& InObjectSelected) const
 {
+	TArray<TWeakInterfacePtr<IPCGGraphExecutionSource>> ExecutionSources;
+
 	for (const TWeakObjectPtr<UObject>& Object : InObjectSelected)
 	{
 		if (APCGVolume* Volume = Cast<APCGVolume>(Object))
@@ -22,9 +24,11 @@ void FPCGVolumeDetails::GatherPCGComponentsFromSelection(const TArray<TWeakObjec
 			{
 				if (Component)
 				{
-					SelectedComponents.Add(Component);
+					ExecutionSources.Add(Component);
 				}
 			}
 		}
 	}
+
+	return ExecutionSources;
 }
\ No newline at end of file
diff --git a/Engine/Plugins/PCG/Source/PCGEditor/Private/Widgets/SPCGExecutionSourceActionWidget.cpp b/Engine/Plugins/PCG/Source/PCGEditor/Private/Widgets/SPCGExecutionSourceActionWidget.cpp
new file mode 100644
--- /dev/null
+++ b/Engine/Plugins/PCG/Source/PCGEditor/Private/Widgets/SPCGExecutionSourceActionWidget.cpp
@@ -0,0 +1,353 @@
+// Copyright Epic Games, Inc. All Rights Reserved.
+
+#include "Widgets/SPCGExecutionSourceActionWidget.h"
+
+#include "PCGCommon.h"
+#include "PCGComponent.h"
+#include "PCGEditorStyle.h"
+#include "Subsystems/IPCGBaseSubsystem.h"
+
+#include "DetailLayoutBuilder.h"
+#include "Editor.h"
+#include "LevelEditorSubsystem.h"
+#include "Elements/Framework/TypedElementSelectionSet.h"
+#include "Framework/Application/SlateApplication.h"
+#include "Widgets/Images/SImage.h"
+#include "Widgets/Input/SButton.h"
+#include "Widgets/Text/STextBlock.h"
+
+#define LOCTEXT_NAMESPACE "PCGExecutionSourceActionWidget"
+
+namespace PCGExecutionSourceActionWidget
+{
+	void TriggerSelectionUpdate(IPCGGraphExecutionSource* ExecutionSource, FPCGTaskId TaskId)
+	{
+		if (!GEditor || !ExecutionSource || TaskId == InvalidPCGTaskId)
+		{
+			return;
+		}
+
+		if (IPCGBaseSubsystem* PCGSubsystem = ExecutionSource->GetExecutionState().GetSubsystem())
+		{
+			FPCGScheduleGenericParams Params(
+				/*InOperation=*/ [](FPCGContext*) -> bool
+				{
+					if (GEditor && GEditor->GetEditorSubsystem<ULevelEditorSubsystem>())
+					{
+						// Since we might have been moving things around (component wise) we need to force a selection set change, 
+						// Which in turn will trigger a component visualizer update.
+						if (UTypedElementSelectionSet* SelectionSet = GEditor->GetEditorSubsystem<ULevelEditorSubsystem>()->GetSelectionSet())
+						{
+							SelectionSet->OnChanged().Broadcast(SelectionSet);
+						}
+					}
+
+					return true;
+				},
+				/*InExecutionSource= */ ExecutionSource,
+				/*InExecutionDependencies=*/ { TaskId },
+				/*InDataDependencies=*/ {},
+				/*bSupportBasePointDataInput=*/ true);
+
+			PCGSubsystem->ScheduleGeneric(Params);
+		}
+	}
+}
+
+void SPCGExecutionSourceActionWidget::Construct(const FArguments& InArgs)
+{
+	ExecutionSources = InArgs._ExecutionSources.Get();
+
+	ChildSlot
+	[
+		SNew(SHorizontalBox)
+		+ SHorizontalBox::Slot()
+		.AutoWidth()
+		.Padding(2.0f, 0.0f)
+		.VAlign(VAlign_Fill)
+		[
+			SNew(SButton)
+			.OnClicked(this, &SPCGExecutionSourceActionWidget::OnGenerateClicked)
+			.ToolTipText(FText::FromString("Generates graph data. \nCtrl + Click flushes the cache and force generates."))
+			.Visibility(this, &SPCGExecutionSourceActionWidget::GetGenerateButtonVisibility)
+			[
+				SNew(SHorizontalBox)
+				+ SHorizontalBox::Slot()
+				.Padding(0.0f, 0.0f, 6.0f, 0.0f)
+				.VAlign(VAlign_Center)
+				.AutoWidth()
+				[
+					SNew(SImage)
+					.DesiredSizeOverride(FVector2D(16, 16))
+					.Image_Lambda([]() { return FSlateApplication::Get().GetModifierKeys().IsControlDown() ? FPCGEditorStyle::Get().GetBrush("PCG.Command.ForceRegenClearCache") : FPCGEditorStyle::Get().GetBrush("PCG.Command.ForceRegen"); })
+					.ColorAndOpacity(FSlateColor::UseForeground())
+				]
+				+ SHorizontalBox::Slot()
+				.VAlign(VAlign_Center)
+				.AutoWidth()
+				[
+					SNew(STextBlock)
+					.Font(IDetailLayoutBuilder::GetDetailFont())
+					.Text_Lambda([]() { return FSlateApplication::Get().GetModifierKeys().IsControlDown() ? LOCTEXT("ForceRegenerateButton", "Force Generate") : LOCTEXT("GenerateButton", "Generate"); })
+				]
+			]
+		]
+		+ SHorizontalBox::Slot()
+		.AutoWidth()
+		.Padding(2.0f, 0.0f)
+		.VAlign(VAlign_Fill)
+		[
+			SNew(SButton)
+			.OnClicked(this, &SPCGExecutionSourceActionWidget::OnCancelClicked)
+			.Visibility(this, &SPCGExecutionSourceActionWidget::GetCancelButtonVisibility)
+			[
+				SNew(SHorizontalBox)
+				+ SHorizontalBox::Slot()
+				.Padding(0.0f, 0.0f, 6.0f, 0.0f)
+				.VAlign(VAlign_Center)
+				.AutoWidth()
+				[
+					SNew(SImage)
+					.DesiredSizeOverride(FVector2D(16, 16))
+					.Image(FPCGEditorStyle::Get().GetBrush("PCG.Command.StopRegen"))
+					.ColorAndOpacity(FSlateColor::UseForeground())
+				]
+				+ SHorizontalBox::Slot()
+				.VAlign(VAlign_Center)
+				.AutoWidth()
+				[
+					SNew(STextBlock)
+					.Font(IDetailLayoutBuilder::GetDetailFont())
+					.Text(LOCTEXT("CancelButton", "Cancel"))
+				]
+			]
+		]
+		+ SHorizontalBox::Slot()
+		.AutoWidth()
+		.VAlign(VAlign_Fill)
+		.Padding(2.0f, 0.0f)
+		[
+			SNew(SButton)
+			.VAlign(VAlign_Center)
+			.OnClicked(this, &SPCGExecutionSourceActionWidget::OnCleanupClicked)
+			.ToolTipText(FText::FromString("Cleans up graph data. \nCtrl + Click purges all generated artifacts tagged as created by PCG."))
+			.Visibility(this, &SPCGExecutionSourceActionWidget::GetCleanupButtonVisibility)
+			[
+				SNew(STextBlock)
+				.Font(IDetailLayoutBuilder::GetDetailFont())
+				.Text(LOCTEXT("CleanupButton", "Cleanup"))
+			]
+		]
+		+ SHorizontalBox::Slot()
+		.AutoWidth()
+		.VAlign(VAlign_Fill)
+		.Padding(2.0f, 0.0f)
+		[
+			SNew(SButton)
+			.VAlign(VAlign_Center)
+			.OnClicked(this, &SPCGExecutionSourceActionWidget::OnRefreshClicked)
+			.Visibility(this, &SPCGExecutionSourceActionWidget::GetRefreshButtonVisibility)
+			[
+				SNew(STextBlock)
+				.Font(IDetailLayoutBuilder::GetDetailFont())
+				.Text(LOCTEXT("RefreshButton", "Refresh"))
+			]
+		]
+		+ SHorizontalBox::Slot()
+		.AutoWidth()
+		.VAlign(VAlign_Fill)
+		.Padding(2.0f, 0.0f)
+		[
+			SNew(SButton)
+			.VAlign(VAlign_Center)
+			.OnClicked(this, &SPCGExecutionSourceActionWidget::OnClearPCGLinkClicked)
+			.Visibility(this, &SPCGExecutionSourceActionWidget::GetClearPCGLinkButtonVisibility)
+			[
+				SNew(STextBlock)
+				.Font(IDetailLayoutBuilder::GetDetailFont())
+				.Text(LOCTEXT("ClearPCGLinkButton", "Clear PCG Link"))
+			]
+		]
+	];
+}
+
+EVisibility SPCGExecutionSourceActionWidget::GetGenerateButtonVisibility() const
+{
+	for (const TWeakInterfacePtr<IPCGGraphExecutionSource>& ExecutionSource : ExecutionSources)
+	{
+		// If execution source is runtime generated then generate/cleanup is managed by the scheduler.
+		if (ExecutionSource.IsValid() && !ExecutionSource->GetExecutionState().IsGenerating() && !ExecutionSource->GetExecutionState().IsManagedByRuntimeGenSystem())
+		{
+			return EVisibility::Visible;
+		}
+	}
+
+	return EVisibility::Collapsed;
+}
+
+EVisibility SPCGExecutionSourceActionWidget::GetCancelButtonVisibility() const
+{
+	for (const TWeakInterfacePtr<IPCGGraphExecutionSource>& ExecutionSource : ExecutionSources)
+	{
+		// If execution source is runtime generated then generate/cleanup is managed by the scheduler.
+		if (ExecutionSource.IsValid() && ExecutionSource->GetExecutionState().IsGenerating() && !ExecutionSource->GetExecutionState().IsManagedByRuntimeGenSystem())
+		{
+			return EVisibility::Visible;
+		}
+	}
+
+	return EVisibility::Collapsed;
+}
+
+EVisibility SPCGExecutionSourceActionWidget::GetCleanupButtonVisibility() const
+{
+	for (const TWeakInterfacePtr<IPCGGraphExecutionSource>& ExecutionSource : ExecutionSources)
+	{
+		// If execution source is runtime generated then generate/cleanup is managed by the scheduler.
+		if (ExecutionSource.IsValid() && !ExecutionSource->GetExecutionState().IsManagedByRuntimeGenSystem())
+		{
+			return EVisibility::Visible;
+		}
+	}
+
+	return EVisibility::Collapsed;
+}
+
+EVisibility SPCGExecutionSourceActionWidget::GetRefreshButtonVisibility() const
+{
+	for (const TWeakInterfacePtr<IPCGGraphExecutionSource>& ExecutionSource : ExecutionSources)
+	{
+		if (ExecutionSource.IsValid() && ExecutionSource->GetExecutionState().IsManagedByRuntimeGenSystem())
+		{
+			return EVisibility::Visible;
+		}
+	}
+
+	return EVisibility::Collapsed;
+}
+
+EVisibility SPCGExecutionSourceActionWidget::GetClearPCGLinkButtonVisibility() const
+{
+	for (const TWeakInterfacePtr<IPCGGraphExecutionSource>& ExecutionSource : ExecutionSources)
+	{
+		if (ExecutionSource.IsValid() && !ExecutionSource->GetExecutionState().IsManagedByRuntimeGenSystem())
+		{
+			return EVisibility::Visible;
+		}
+	}
+
+	return EVisibility::Collapsed;
+}
+
+
+FReply SPCGExecutionSourceActionWidget::OnGenerateClicked()
+{
+	const FModifierKeysState ModifierKeys = FSlateApplication::Get().GetModifierKeys();
+	const bool bIsControlDown = ModifierKeys.IsControlDown();
+
+	bool bFlushedCache = false;
+
+	for (TWeakInterfacePtr<IPCGGraphExecutionSource>& ExecutionSourcePtr : ExecutionSources)
+	{
+		IPCGGraphExecutionSource* ExecutionSource = ExecutionSourcePtr.Get();
+
+		if (ExecutionSource && !ExecutionSource->GetExecutionState().IsGenerating() && !ExecutionSource->GetExecutionState().IsManagedByRuntimeGenSystem())
+		{
+			// Flush cache if needed
+			if (bIsControlDown && !bFlushedCache)
+			{
+				if (IPCGBaseSubsystem* PCGSubsystem = ExecutionSource->GetExecutionState().GetSubsystem())
+				{
+					PCGSubsystem->FlushCache();
+					bFlushedCache = true;
+				}
+			}
+
+			IPCGGraphExecutionState::FGenerateParams Params;
+			Params.bEvenIfAlreadyGenerated = bIsControlDown;
+			FPCGTaskId GenerationTask = ExecutionSource->GetExecutionState().Generate(Params);
+						
+			PCGExecutionSourceActionWidget::TriggerSelectionUpdate(ExecutionSource, GenerationTask);
+		}
+	}
+
+	return FReply::Handled();
+}
+
+FReply SPCGExecutionSourceActionWidget::OnCancelClicked()
+{
+	for (TWeakInterfacePtr<IPCGGraphExecutionSource>& ExecutionSource : ExecutionSources)
+	{
+		if (ExecutionSource.IsValid() && ExecutionSource->GetExecutionState().IsGenerating() && !ExecutionSource->GetExecutionState().IsManagedByRuntimeGenSystem())
+		{
+			ExecutionSource->GetExecutionState().Cancel();
+		}
+	}
+
+	return FReply::Handled();
+}
+
+FReply SPCGExecutionSourceActionWidget::OnRefreshClicked()
+{
+	for (TWeakInterfacePtr<IPCGGraphExecutionSource>& ExecutionSource : ExecutionSources)
+	{
+		// @todo_pcg: support this in IPCGGraphExecutionSource
+		if (UPCGComponent* Component = Cast<UPCGComponent>(ExecutionSource.Get()))
+		{
+			if (IsValid(Component) && Component->IsManagedByRuntimeGenSystem())
+			{
+				// Trigger the deepest refresh - re-initialize the PAs.
+				Component->Refresh(EPCGChangeType::Structural | EPCGChangeType::GenerationGrid, /*bCancelExistingRefresh=*/true);
+			}
+		}
+	}
+
+	return FReply::Handled();
+}
+
+FReply SPCGExecutionSourceActionWidget::OnClearPCGLinkClicked()
+{
+	for (TWeakInterfacePtr<IPCGGraphExecutionSource>& ExecutionSource : ExecutionSources)
+	{
+		// @todo_pcg: support this in IPCGGraphExecutionSource
+		if (UPCGComponent* Component = Cast<UPCGComponent>(ExecutionSource.Get()))
+		{
+			if (IsValid(Component) && !Component->IsManagedByRuntimeGenSystem())
+			{
+				Component->ClearPCGLink();
+			}
+		}
+	}
+
+	return FReply::Handled();
+}
+
+FReply SPCGExecutionSourceActionWidget::OnCleanupClicked()
+{
+	for (TWeakInterfacePtr<IPCGGraphExecutionSource>& ExecutionSource : ExecutionSources)
+	{
+		if (ExecutionSource.IsValid() && !ExecutionSource->GetExecutionState().IsManagedByRuntimeGenSystem())
+		{
+			FModifierKeysState ModifierKeys = FSlateApplication::Get().GetModifierKeys();
+
+			FPCGTaskId CleanupTaskId = InvalidPCGTaskId;
+
+			UPCGComponent* Component = Cast<UPCGComponent>(ExecutionSource.Get());
+			if (ModifierKeys.IsControlDown() && Component)
+			{
+				Component->CleanupLocalDeleteAllGeneratedObjects({});
+				CleanupTaskId = Component->GetCleanupTaskId();
+			}
+			else
+			{
+				CleanupTaskId = ExecutionSource->GetExecutionState().Cleanup();
+			}
+
+			PCGExecutionSourceActionWidget::TriggerSelectionUpdate(ExecutionSource.Get(), CleanupTaskId);
+		}
+	}
+
+	return FReply::Handled();
+}
+
+#undef LOCTEXT_NAMESPACE
diff --git a/Engine/Plugins/PCG/Source/PCGEditor/Public/Details/PCGComponentDetails.h b/Engine/Plugins/PCG/Source/PCGEditor/Public/Details/PCGComponentDetails.h
--- a/Engine/Plugins/PCG/Source/PCGEditor/Public/Details/PCGComponentDetails.h
+++ b/Engine/Plugins/PCG/Source/PCGEditor/Public/Details/PCGComponentDetails.h
@@ -3,12 +3,11 @@
 #pragma once
 
 #include "IDetailCustomization.h"
-#include "UObject/WeakObjectPtrTemplates.h"
 
-class FReply;
-struct EVisibility;
+#include "PCGGraphExecutionStateInterface.h"
 
-class UPCGComponent;
+#include "UObject/WeakInterfacePtr.h"
+#include "UObject/WeakObjectPtrTemplates.h"
 
 class FPCGComponentDetails : public IDetailCustomization
 {
@@ -17,26 +16,9 @@ public:
 
 	/** ~Begin IDetailCustomization interface */
 	virtual void CustomizeDetails(IDetailLayoutBuilder& DetailBuilder) override;
-	virtual void PendingDelete() override;
 	/** ~End IDetailCustomization interface */
 
 protected:
-	virtual void GatherPCGComponentsFromSelection(const TArray<TWeakObjectPtr<UObject>>& InObjectSelected);
+	virtual TArray<TWeakInterfacePtr<IPCGGraphExecutionSource>> GatherExecutionSourcesFromSelection(const TArray<TWeakObjectPtr<UObject>>& InObjectSelected) const;
 	virtual bool AddDefaultProperties() const { return true; }
-
-private:
-	EVisibility GenerateButtonVisible() const;
-	EVisibility CancelButtonVisible() const;
-	EVisibility CleanupButtonVisible() const;
-	EVisibility RefreshButtonVisible() const;
-	EVisibility ClearPCGLinkButtonVisible() const;
-	FReply OnGenerateClicked();
-	FReply OnCancelClicked();
-	FReply OnCleanupClicked();
-	FReply OnRefreshClicked();
-	FReply OnClearPCGLinkClicked();
-	void OnGraphChanged(UPCGComponent* InComponent);
-
-protected:
-	TArray<TWeakObjectPtr<UPCGComponent>> SelectedComponents;
 };
diff --git a/Engine/Plugins/PCG/Source/PCGEditor/Public/Details/PCGVolumeDetails.h b/Engine/Plugins/PCG/Source/PCGEditor/Public/Details/PCGVolumeDetails.h
--- a/Engine/Plugins/PCG/Source/PCGEditor/Public/Details/PCGVolumeDetails.h
+++ b/Engine/Plugins/PCG/Source/PCGEditor/Public/Details/PCGVolumeDetails.h
@@ -13,7 +13,7 @@ public:
 
 protected:
 	/** ~Begin FPCGComponentDetails interface */
-	virtual void GatherPCGComponentsFromSelection(const TArray<TWeakObjectPtr<UObject>>& InObjectSelected) override;
+	virtual TArray<TWeakInterfacePtr<IPCGGraphExecutionSource>> GatherExecutionSourcesFromSelection(const TArray<TWeakObjectPtr<UObject>>& InObjectSelected) const override;
 	virtual bool AddDefaultProperties() const override { return false; }
 	/** ~End FPCGComponentDetails interface */
 };
diff --git a/Engine/Plugins/PCG/Source/PCGEditor/Public/Widgets/SPCGExecutionSourceActionWidget.h b/Engine/Plugins/PCG/Source/PCGEditor/Public/Widgets/SPCGExecutionSourceActionWidget.h
new file mode 100644
--- /dev/null
+++ b/Engine/Plugins/PCG/Source/PCGEditor/Public/Widgets/SPCGExecutionSourceActionWidget.h
@@ -0,0 +1,36 @@
+// Copyright Epic Games, Inc. All Rights Reserved.
+
+#pragma once
+
+#include "PCGGraphExecutionStateInterface.h"
+
+#include "UObject/WeakInterfacePtr.h"
+#include "Widgets/SCompoundWidget.h"
+
+class FReply;
+struct EVisibility;
+
+class SPCGExecutionSourceActionWidget : public SCompoundWidget
+{
+public:
+	SLATE_BEGIN_ARGS(SPCGExecutionSourceActionWidget) {}
+		SLATE_ATTRIBUTE(TArray<TWeakInterfacePtr<IPCGGraphExecutionSource>>, ExecutionSources)
+	SLATE_END_ARGS()
+
+	PCGEDITOR_API void Construct(const FArguments& InArgs);
+
+private:
+	EVisibility GetGenerateButtonVisibility() const;
+	EVisibility GetCancelButtonVisibility() const;
+	EVisibility GetCleanupButtonVisibility() const;
+	EVisibility GetRefreshButtonVisibility() const;
+	EVisibility GetClearPCGLinkButtonVisibility() const;
+	
+	FReply OnGenerateClicked();
+	FReply OnCancelClicked();
+	FReply OnCleanupClicked();
+	FReply OnRefreshClicked();
+	FReply OnClearPCGLinkClicked();
+	
+	TArray<TWeakInterfacePtr<IPCGGraphExecutionSource>> ExecutionSources;
+};
\ No newline at end of file
