[PCG] Optimize compute graph compilation.

Also add some profile scopes in schedule/compilation work.

#rnx

diff --git a/Engine/Plugins/PCG/Source/PCG/Private/Graph/PCGGraphCompiler.cpp b/Engine/Plugins/PCG/Source/PCG/Private/Graph/PCGGraphCompiler.cpp
--- a/Engine/Plugins/PCG/Source/PCG/Private/Graph/PCGGraphCompiler.cpp
+++ b/Engine/Plugins/PCG/Source/PCG/Private/Graph/PCGGraphCompiler.cpp
@@ -147,6 +147,7 @@ void FPCGGraphCompilerCache::RemoveFromCacheRecursive(UPCGGraph* InGraph)
 
 TArray<FPCGGraphTask> FPCGGraphCompiler::CompileGraph(UPCGGraph* InGraph, FPCGTaskId& NextId, FPCGStackContext& InOutStackContext)
 {
+	TRACE_CPUPROFILER_EVENT_SCOPE(FPCGGraphCompiler::CompileGraph);
 	if (!InGraph)
 	{
 		return TArray<FPCGGraphTask>();
@@ -402,8 +403,9 @@ TArray<FPCGGraphTask> FPCGGraphCompiler::CompileGraph(UPCGGraph* InGraph, FPCGTa
 
 void FPCGGraphCompiler::Compile(UPCGGraph* InGraph)
 {
-	bool bAlreadyCached = false;
+	TRACE_CPUPROFILER_EVENT_SCOPE(FPCGGraphCompiler::Compile);
 
+	bool bAlreadyCached = false;
 	{
 		PCG::TSharedScopeLock ReadLock(Cache.GraphToTaskMapLock);
 		bAlreadyCached = Cache.GraphToTaskMap.Contains(InGraph);
@@ -1109,6 +1111,8 @@ void FPCGGraphCompiler::PostCullStackCleanup(TArray<FPCGGraphTask>& InCompiledTa
 
 void FPCGGraphCompiler::CalculateDynamicActivePinDependencies(FPCGTaskId InTaskId, TArray<FPCGGraphTask>& InOutCompiledTasks)
 {
+	TRACE_CPUPROFILER_EVENT_SCOPE(FPCGGraphCompiler::CalculateDynamicActivePinDependencies);
+
 	if (!ensure(InOutCompiledTasks.IsValidIndex(InTaskId)))
 	{
 		return;
@@ -1237,6 +1241,8 @@ void FPCGGraphCompiler::CalculateDynamicActivePinDependencies(FPCGTaskId InTaskI
 
 TArray<FPCGGraphTask> FPCGGraphCompiler::GetCompiledTasks(UPCGGraph* InGraph, uint32 GenerationGridSize, FPCGStackContext& OutStackContext, bool bIsTopGraph)
 {
+	TRACE_CPUPROFILER_EVENT_SCOPE(FPCGGraphCompiler::GetCompiledTasks);
+
 #if !WITH_EDITOR
 	// In standalone builds, try using cooked tasks first.
 	if (InGraph)
@@ -1327,6 +1333,7 @@ void FPCGGraphCompiler::OffsetNodeIds(TArray<FPCGGraphTask>& Tasks, FPCGTaskId O
 
 void FPCGGraphCompiler::CompileTopGraph(UPCGGraph* InGraph, uint32 GenerationGridSize)
 {
+	TRACE_CPUPROFILER_EVENT_SCOPE(FPCGGraphCompiler::CompileTopGraph);
 	bool bAlreadyCached = false;
 
 	{
@@ -1354,6 +1361,8 @@ void FPCGGraphCompiler::CompileTopGraph(UPCGGraph* InGraph, uint32 GenerationGri
 
 	if (PCGGraphCompiler::CVarEnableTaskStaticCulling.GetValueOnAnyThread())
 	{
+		TRACE_CPUPROFILER_EVENT_SCOPE(TaskStaticCulling);
+
 		// Remove reroute nodes before execution grid setup, as grid linkages need final nodes to connect from/to.
 		CullTasks(CompiledTasks, /*bAddPassthroughWires=*/true, [](const FPCGGraphTask& InTask) { return InTask.Node && Cast<UPCGRerouteSettings>(InTask.Node->GetSettings()) && InTask.Node->HasInboundEdges(); });
 
diff --git a/Engine/Plugins/PCG/Source/PCG/Private/Graph/PCGGraphCompilerGPU.cpp b/Engine/Plugins/PCG/Source/PCG/Private/Graph/PCGGraphCompilerGPU.cpp
--- a/Engine/Plugins/PCG/Source/PCG/Private/Graph/PCGGraphCompilerGPU.cpp
+++ b/Engine/Plugins/PCG/Source/PCG/Private/Graph/PCGGraphCompilerGPU.cpp
@@ -264,6 +264,7 @@ void FPCGGraphCompilerGPU::ExpandGPUNodeKernelsToTasks(
 	TArray<FPCGGraphTask>& InOutCompiledTasks,
 	TMap<FPCGTaskId, UPCGComputeKernel*>& OutTaskIdToKernel)
 {
+	TRACE_CPUPROFILER_EVENT_SCOPE(FPCGGraphCompilerGPU::ExpandGPUNodeKernelsToTasks);
 	const int NumTasksBefore = InOutCompiledTasks.Num();
 
 	// Build successors map, only for nodes that are downstream from one or more of the relevant GPU tasks.
@@ -880,14 +881,6 @@ void FPCGGraphCompilerGPU::BuildComputeGraphTask(
 	{
 		InOutCompiledTasks[InGPUGraphTaskId].Element = MakeShared<FPCGComputeGraphElement>(ComputeGraphIndex);
 	}
-
-	if (FApp::CanEverRender() && IsInGameThread() && ensure(ComputeGraph))
-	{
-		TRACE_CPUPROFILER_EVENT_SCOPE(UpdateResources);
-
-		// Compile shader resources and create render proxies.
-		ComputeGraph->UpdateResources(/*bSync=*/!ComputeFramework::IsDeferredCompilation());
-	}
 }
 
 void FPCGGraphCompilerGPU::BuildComputeGraphStaticData(
@@ -1145,6 +1138,23 @@ void FPCGGraphCompilerGPU::CompileComputeGraph(
 {
 	TRACE_CPUPROFILER_EVENT_SCOPE(FPCGGraphCompilerGPU::CompileComputeGraph);
 
+	TMap<const UComputeDataInterface*, TArray<FShaderFunctionDefinition>> DataInterfaceToInputs;
+	TMap<const UComputeDataInterface*, TArray<FShaderFunctionDefinition>> DataInterfaceToOutputs;
+	{
+		TRACE_CPUPROFILER_EVENT_SCOPE(CacheDataInterfaceInputOutputs);
+		for (const TPair<TPair</*Node task*/FPCGTaskId, /*Pin label*/FName>, UPCGComputeDataInterface*>& PinAndDataInterface : InOutputPinToDataInterface)
+		{
+			// FindOrAdd is safe because the data interfaces in InOutputPinToDataInterface will be unique by definition, because every output pin maps to its own DI.
+			TArray<FShaderFunctionDefinition>& InputDefinitions = DataInterfaceToInputs.FindOrAdd(PinAndDataInterface.Value);
+			InputDefinitions.Reserve(128);
+			PinAndDataInterface.Value->GetSupportedInputs(InputDefinitions);
+
+			TArray<FShaderFunctionDefinition>& OutputDefinitions = DataInterfaceToOutputs.FindOrAdd(PinAndDataInterface.Value);
+			OutputDefinitions.Reserve(128);
+			PinAndDataInterface.Value->GetSupportedOutputs(OutputDefinitions);
+		}
+	}
+
 	// Build compute graph by traversing the set of collapsed tasks in execution order (order of queued kernel invocations matters).
 	TSet<FPCGTaskId> RemainingTasks;
 	TArray<FPCGTaskId> ReadyTasks;
@@ -1448,11 +1458,20 @@ void FPCGGraphCompilerGPU::CompileComputeGraph(
 
 		InOutComputeGraph->KernelToNode.Add(Node);
 
-		auto SetupAllInputBindings = [&KernelWithBindings, InOutComputeGraph](const FDataInterfaceInfo& InDataInterfaceInfo)
+		auto SetupAllInputBindings = [&KernelWithBindings, InOutComputeGraph, &DataInterfaceToInputs](const FDataInterfaceInfo& InDataInterfaceInfo)
 		{
 			const UComputeDataInterface* DataInterface = InOutComputeGraph->DataInterfaces[InDataInterfaceInfo.Index];
-			TArray<FShaderFunctionDefinition> Functions;
-			DataInterface->GetSupportedInputs(Functions);
+			check(DataInterface);
+
+			TArray<FShaderFunctionDefinition>* FunctionsFromCache = DataInterfaceToInputs.Find(DataInterface);
+			TArray<FShaderFunctionDefinition> RetrievedFunctions; // Used as local storage if we need to retrieve the functions here.
+			if (!FunctionsFromCache)
+			{
+				RetrievedFunctions.Reserve(128);
+				DataInterface->GetSupportedInputs(RetrievedFunctions);
+			}
+
+			TArray<FShaderFunctionDefinition>& Functions = FunctionsFromCache ? *FunctionsFromCache : RetrievedFunctions;
 
 			for (int FuncIndex = 0; FuncIndex < Functions.Num(); ++FuncIndex)
 			{
@@ -1463,13 +1482,28 @@ void FPCGGraphCompilerGPU::CompileComputeGraph(
 				Binding.BindingFunctionName = Functions[FuncIndex].Name;
 				Binding.DataInterfaceBindingIndex = FuncIndex;
 			}
+
+			if (!FunctionsFromCache)
+			{
+				// Augment the DI function cache with any additional data interfaces created for the current task.
+				DataInterfaceToInputs.Add(DataInterface, MoveTemp(RetrievedFunctions));
+			}
 		};
 
-		auto SetupAllOutputBindings = [&KernelWithBindings, InOutComputeGraph](const FDataInterfaceInfo& InDataInterfaceInfo)
+		auto SetupAllOutputBindings = [&KernelWithBindings, InOutComputeGraph, &DataInterfaceToOutputs](const FDataInterfaceInfo& InDataInterfaceInfo)
 		{
 			const UComputeDataInterface* DataInterface = InOutComputeGraph->DataInterfaces[InDataInterfaceInfo.Index];
-			TArray<FShaderFunctionDefinition> Functions;
-			DataInterface->GetSupportedOutputs(Functions);
+			check(DataInterface);
+
+			TArray<FShaderFunctionDefinition>* FunctionsFromCache = DataInterfaceToOutputs.Find(DataInterface);
+			TArray<FShaderFunctionDefinition> RetrievedFunctions; // Used as local storage if we need to retrieve the functions here.
+			if (!FunctionsFromCache)
+			{
+				RetrievedFunctions.Reserve(128);
+				DataInterface->GetSupportedOutputs(RetrievedFunctions);
+			}
+
+			TArray<FShaderFunctionDefinition>& Functions = FunctionsFromCache ? *FunctionsFromCache : RetrievedFunctions;
 
 			for (int FuncIndex = 0; FuncIndex < Functions.Num(); ++FuncIndex)
 			{
@@ -1480,6 +1514,12 @@ void FPCGGraphCompilerGPU::CompileComputeGraph(
 				Binding.BindingFunctionName = Functions[FuncIndex].Name;
 				Binding.DataInterfaceBindingIndex = FuncIndex;
 			}
+
+			if (!FunctionsFromCache)
+			{
+				// Augment the DI function cache with any additional data interfaces created for the current task.
+				DataInterfaceToOutputs.Add(DataInterface, MoveTemp(RetrievedFunctions));
+			}
 		};
 
 		// Bind data interfaces.
@@ -1551,8 +1591,7 @@ void FPCGGraphCompilerGPU::CompileComputeGraph(
 			// Add functions for external inputs/outputs which must be fulfilled by DIs
 			for (FInterfaceBinding& Binding : KernelWithBindings.InputDataBindings)
 			{
-				TArray<FShaderFunctionDefinition> Functions;
-				Binding.DataInterface->GetSupportedInputs(Functions);
+				TArray<FShaderFunctionDefinition>& Functions = DataInterfaceToInputs.FindChecked(Binding.DataInterface);
 				check(Functions.IsValidIndex(Binding.DataInterfaceBindingIndex));
 
 				FShaderFunctionDefinition FuncDef = Functions[Binding.DataInterfaceBindingIndex];
@@ -1567,8 +1606,7 @@ void FPCGGraphCompilerGPU::CompileComputeGraph(
 
 			for (FInterfaceBinding& Binding : KernelWithBindings.OutputDataBindings)
 			{
-				TArray<FShaderFunctionDefinition> Functions;
-				Binding.DataInterface->GetSupportedOutputs(Functions);
+				TArray<FShaderFunctionDefinition>& Functions = DataInterfaceToOutputs.FindChecked(Binding.DataInterface);
 				check(Functions.IsValidIndex(Binding.DataInterfaceBindingIndex));
 
 				FShaderFunctionDefinition FuncDef = Functions[Binding.DataInterfaceBindingIndex];
@@ -1582,9 +1620,11 @@ void FPCGGraphCompilerGPU::CompileComputeGraph(
 			}
 		}
 
-		auto AddAllEdgesForKernel = [&KernelWithBindings, InOutComputeGraph, &OriginalToVirtualInputPin](int32 InKernelIndex, bool bInEdgesAreInputs)
+		auto AddAllEdgesForKernel = [&KernelWithBindings, InOutComputeGraph, &OriginalToVirtualInputPin, &DataInterfaceToInputs, &DataInterfaceToOutputs](int32 InKernelIndex, bool bInEdgesAreInputs)
 		{
+			TRACE_CPUPROFILER_EVENT_SCOPE(AddAllEdgesForKernel);
 			TArray<FInterfaceBinding>& Bindings = bInEdgesAreInputs ? KernelWithBindings.InputDataBindings : KernelWithBindings.OutputDataBindings;
+			InOutComputeGraph->GraphEdges.Reserve(InOutComputeGraph->GraphEdges.Num() + Bindings.Num());
 
 			// Add all graph edges for bindings, which means include all functions that the data interfaces expose.
 			for (int KernelBindingIndex = 0; KernelBindingIndex < Bindings.Num(); ++KernelBindingIndex)
@@ -1608,16 +1648,8 @@ void FPCGGraphCompilerGPU::CompileComputeGraph(
 
 				if (Binding.Prefix != NAME_None)
 				{
-					TArray<FShaderFunctionDefinition> DataInterfaceFunctions;
-					if (bInEdgesAreInputs)
-					{
-						DataInterface->GetSupportedInputs(DataInterfaceFunctions);
-					}
-					else
-					{
-						DataInterface->GetSupportedOutputs(DataInterfaceFunctions);
-					}
-
+					TArray<FShaderFunctionDefinition>& DataInterfaceFunctions = bInEdgesAreInputs ? DataInterfaceToInputs.FindChecked(DataInterface) : DataInterfaceToOutputs.FindChecked(DataInterface);
+					
 					Edge.BindingFunctionNameOverride = FString::Format(
 						TEXT("{0}_{1}"),
 						{ Binding.Prefix.ToString(), DataInterfaceFunctions[Edge.DataInterfaceBindingIndex].Name }
diff --git a/Engine/Plugins/PCG/Source/PCG/Private/Graph/PCGGraphExecutor.cpp b/Engine/Plugins/PCG/Source/PCG/Private/Graph/PCGGraphExecutor.cpp
--- a/Engine/Plugins/PCG/Source/PCG/Private/Graph/PCGGraphExecutor.cpp
+++ b/Engine/Plugins/PCG/Source/PCG/Private/Graph/PCGGraphExecutor.cpp
@@ -220,6 +220,7 @@ FPCGTaskId FPCGGraphExecutor::Schedule(
 
 FPCGTaskId FPCGGraphExecutor::ScheduleGraph(const FPCGScheduleGraphParams& InParams)
 {
+	TRACE_CPUPROFILER_EVENT_SCOPE(FPCGGraphExecutor::ScheduleGraph);
 	check(InParams.ExecutionSource);
 
 	PCGGraphExecutionLogging::LogGraphSchedule(InParams.ExecutionSource, InParams.Graph);
