[PCG] Windowing for generate landscape textures so we only produce output data for covering the generation cell.

In grass map test level the final texture is <20% of the size before.

#rnx

diff --git a/Engine/Plugins/PCG/Shaders/Private/PCGGrassMapUnpackerCS.usf b/Engine/Plugins/PCG/Shaders/Private/PCGGrassMapUnpackerCS.usf
--- a/Engine/Plugins/PCG/Shaders/Private/PCGGrassMapUnpackerCS.usf
+++ b/Engine/Plugins/PCG/Shaders/Private/PCGGrassMapUnpackerCS.usf
@@ -18,20 +18,21 @@
 //  -----------------------------------------
 Texture2D<float4> InPackedGrassMaps;
 
-// Grass maps output. Each slice is a single grass combined grass map.
+// Full output domain consists of multiple components, but we only run on a "window" that covers the generation cell bounds.
+// The shader writes to an output texture with height data and a texture array where each slice is a single grass combined grass map.
 // Each component is an inner grid cell. Res is vert count and includes outer edges, so a row/col is removed from each internal boundary.
 //
 //                         -------------------------------
 //                       -------------------------------  |
 //                      |  Res-1  |  Res-1  |    Res    | |
 //                      |         |         |           | |
-//  (Res-1) * NumTilesY |Res-1    |         |           | |
-//         + 1          |         |         |           | |
-//                      |---------+---------+-----------| |
-//                      |         |         |           | |
-//                      |Res-1    |         |           | |
-//                      |         |         |           | |
-//                      |         |         |           | |
+//  (Res-1) * NumTilesY |Res-1    |     +---------+     | |
+//         + 1          |         |     |         |     | |
+//                      |---------+-----|  Exec   |-----| |
+//                      |         |     |  Bounds |     | |
+//                      |Res-1    |     |         |     | |
+//                      |         |     +---------+     | |
+//                      |         |      InWindow       | |
 //                      |---------+---------+-----------| |
 //                      |         |         |           | |
 //                      |         |         |           | |
@@ -52,21 +53,25 @@ uint InLandscapeComponentResolution; // One texel per vert, so number of quads +
 uint InNumGrassMapPasses;
 float3 InLandscapeGridScale; // x == LS Actor DrawScale.X, y == LS Actor DrawScale.y, z == LS Actor DrawScale.z / 128.0f (ZSCALE)
 float InLandscapeLocationZ;
-uint2 InOutputResolution;
 uint InOutputHeight;
+uint4 InWindow; // xy == Window start coord, zw = Window end coord
 
-// Dispatch group count is (InOutputResolution.xy / 8, NumGrassMaps).
+// Dispatch group count is ((InWindow.zw - InWindow.xy + 1) / 8, Max(1, NumGrassMaps)).
 [numthreads(THREADGROUPSIZE_X, THREADGROUPSIZE_Y, THREADGROUPSIZE_Z)]
-void PCGGrassMapUnpacker_CS(uint3 DispatchThreadId : SV_DispatchThreadID)
+void PCGGrassMapUnpacker_CS(const uint3 WindowCoord_Slice : SV_DispatchThreadID)
 {
-	if (any(DispatchThreadId.xy >= InOutputResolution))
+	const uint2 WindowCoord = WindowCoord_Slice.xy;
+	const uint2 FullCoord = InWindow.xy + WindowCoord;
+	if (any(FullCoord > InWindow.zw))
 	{
 		return;
 	}
 
+	const uint Slice = WindowCoord_Slice.z;
+
 	// The rendered landscape components in the input are not in any particular order, so lookup which component index for this output tile.
 	// Clamp because we are dividing by resolution-1 and last col/row of texture will overflow.
-	const uint2 TileCoords = min(DispatchThreadId.xy / (InLandscapeComponentResolution - 1u), InNumTiles - 1u);
+	const uint2 TileCoords = min(FullCoord / (InLandscapeComponentResolution - 1u), InNumTiles - 1u);
 
 	const uint TileLinearIndex = TileCoords.y * InNumTiles.x + TileCoords.x;
 	const int ComponentIndex = InLinearTileIndexToComponentIndex[TileLinearIndex][0];
@@ -75,34 +80,33 @@ void PCGGrassMapUnpacker_CS(uint3 DispatchThreadId : SV_DispatchThreadID)
 	if (ComponentIndex == -1)
 	{
 		// First layer of threads also deals with height.
-		if (InOutputHeight == 1 && DispatchThreadId.z == 0)
+		if (InOutputHeight == 1 && Slice == 0)
 		{
-			OutUnpackedHeight[DispatchThreadId.xy] = 0.0f;
+			OutUnpackedHeight[WindowCoord] = 0.0f;
 		}
 
-		OutUnpackedGrassMaps[DispatchThreadId] = 0.0f;
+		OutUnpackedGrassMaps[uint3(WindowCoord, Slice)] = 0.0f;
 		return;
 	}
 
-	// Coordinates into input grass maps (see packing above).
-	const uint GrassMapIndex = DispatchThreadId.z;
-	const uint PassIndex = (GrassMapIndex + 2) / 4;
+	// Coordinates into input height & grass maps (see packing above).
+	const uint PassIndex = (Slice + 2) / 4;
 	const uint ComponentXOffset = (ComponentIndex * InNumGrassMapPasses + PassIndex) * InLandscapeComponentResolution;
 
 	// Coordinates into packed data which are laid out horizontally (see above).
-	const uint2 LocalTileCoords = uint2(DispatchThreadId.x - TileCoords.x * (InLandscapeComponentResolution - 1), DispatchThreadId.y - TileCoords.y * (InLandscapeComponentResolution - 1));
+	const uint2 LocalTileCoords = uint2(FullCoord.x - TileCoords.x * (InLandscapeComponentResolution - 1), FullCoord.y - TileCoords.y * (InLandscapeComponentResolution - 1));
 	const uint2 InputCoords = uint2(LocalTileCoords.x + ComponentXOffset, LocalTileCoords.y);
 
 	// Heights. The first layer threads also deal with height.
-	if (InOutputHeight == 1 && DispatchThreadId.z == 0)
+	if (InOutputHeight == 1 && Slice == 0)
 	{
 		const float RelativeHeight = (UnpackHeight(InPackedGrassMaps[InputCoords].xy) - LANDSCAPE_MID_VALUE) * InLandscapeGridScale.z * TERRAIN_ZSCALE;
-		OutUnpackedHeight[DispatchThreadId.xy] = RelativeHeight + InLandscapeLocationZ;
+		OutUnpackedHeight[WindowCoord] = RelativeHeight + InLandscapeLocationZ;
 	}
 
 	// Grass map source channel (see packing above).
-	const uint InputChannelIndex = (GrassMapIndex + 2) % 4;
+	const uint InputChannelIndex = (Slice + 2) % 4;
 
 	// Grass maps
-	OutUnpackedGrassMaps[DispatchThreadId] = InPackedGrassMaps[InputCoords][InputChannelIndex];
+	OutUnpackedGrassMaps[uint3(WindowCoord, Slice)] = InPackedGrassMaps[InputCoords][InputChannelIndex];
 }
diff --git a/Engine/Plugins/PCG/Source/PCG/Private/Elements/PCGGenerateLandscapeTextures.cpp b/Engine/Plugins/PCG/Source/PCG/Private/Elements/PCGGenerateLandscapeTextures.cpp
--- a/Engine/Plugins/PCG/Source/PCG/Private/Elements/PCGGenerateLandscapeTextures.cpp
+++ b/Engine/Plugins/PCG/Source/PCG/Private/Elements/PCGGenerateLandscapeTextures.cpp
@@ -51,6 +51,9 @@ namespace PCGGenerateLandscapeTextures
 		TEXT("Draws debug boxes around landscapes for which grass maps are generated, colored by the current task ID."));
 #endif
 
+	// todo_pcg: Remove after running tests.
+	TAutoConsoleVariable<bool> CVarWindowedGeneration(TEXT("pcg.Grass.WindowedGeneration"), false, TEXT(""));
+
 	static bool IsTextureFullyStreamedIn(UTexture* InTexture)
 	{
 		const bool bCheckForLODTransition = true;
@@ -329,7 +332,7 @@ bool FPCGGenerateLandscapeTexturesElement::ExecuteInternal(FPCGContext* InContex
 							}
 						}
 
-						Context->GrassMapBounds = LandscapeComponentBounds;
+						Context->TotalLandscapeComponentBounds = LandscapeComponentBounds;
 
 						Context->LandscapeComponentExtent = LandscapeComponent->Bounds.BoxExtent.X * 2.0;
 						Context->TexelSizeWorld = Context->LandscapeComponentExtent / static_cast<double>(LandscapeProxy->ComponentSizeQuads);
@@ -346,7 +349,7 @@ bool FPCGGenerateLandscapeTexturesElement::ExecuteInternal(FPCGContext* InContex
 					}
 					else
 					{
-						Context->GrassMapBounds += LandscapeComponentBounds;
+						Context->TotalLandscapeComponentBounds += LandscapeComponentBounds;
 
 						const UMaterialInterface* LandscapeProxyMaterial = Context->LandscapeProxy->GetLandscapeMaterial();
 						const UMaterialInterface* LandscapeMaterial = LandscapeComponent->GetLandscapeMaterial();
@@ -528,8 +531,8 @@ bool FPCGGenerateLandscapeTexturesElement::ExecuteInternal(FPCGContext* InContex
 			// Landscape components are not ordered, so store a 2d index of each component within the grass map.
 			// TODO could likely sort the LandscapeComponents array so the order is known and no indices need looking up.
 			LandscapeTileCoords.Add(FIntVector2(
-				(LandscapeComponent->Bounds.Origin.X - Context->GrassMapBounds.Min.X) / Context->LandscapeComponentExtent,
-				(LandscapeComponent->Bounds.Origin.Y - Context->GrassMapBounds.Min.Y) / Context->LandscapeComponentExtent));
+				(LandscapeComponent->Bounds.Origin.X - Context->TotalLandscapeComponentBounds.Min.X) / Context->LandscapeComponentExtent,
+				(LandscapeComponent->Bounds.Origin.Y - Context->TotalLandscapeComponentBounds.Min.Y) / Context->LandscapeComponentExtent));
 		}
 
 		Context->LandscapeGrassWeightExporter = new FLandscapeGrassWeightExporter(
@@ -550,6 +553,66 @@ bool FPCGGenerateLandscapeTexturesElement::ExecuteInternal(FPCGContext* InContex
 			return true;
 		}
 
+		const int32 ComponentSizeQuads = LandscapeProxy->ComponentSizeQuads;
+		const double QuadSizeWorld = Context->LandscapeComponentExtent / ComponentSizeQuads;
+		// Each corner of a quad in the landscape corresponds to one texel in the grass map. 
+		const uint32 LandscapeComponentResolution = ComponentSizeQuads + 1;
+		// Internal boundaries between components in the output share a row/column, so these are dropped to get a correct world space output.
+		const FVector TotalLandscapeComponentExtent = Context->TotalLandscapeComponentBounds.GetExtent() * 2;
+		const uint32 NumTilesX = FMath::RoundToInt(TotalLandscapeComponentExtent.X / Context->LandscapeComponentExtent);
+		const uint32 NumTilesY = FMath::RoundToInt(TotalLandscapeComponentExtent.Y / Context->LandscapeComponentExtent);
+		const FIntPoint FullExportResolution((LandscapeComponentResolution - 1u)* NumTilesX + 1u, (LandscapeComponentResolution - 1u)* NumTilesY + 1u);
+
+		FIntPoint WindowStart, WindowEnd;
+
+		if (PCGGenerateLandscapeTextures::CVarWindowedGeneration.GetValueOnGameThread())
+		{
+			const FBox CellBounds = Context->ExecutionSource->GetExecutionState().GetBounds();
+			const FBox& TotalBounds = Context->TotalLandscapeComponentBounds;
+
+			UE_LOG(LogPCG, Verbose, TEXT("PCGUnpackGrassMap: %s, Cell bounds (%.2f, %.2f -> %.2f, %.2f), Total bounds (%.2f, %.2f -> %.2f, %.2f), efficiency %.2f"),
+				*Context->GetExecutionSourceName(),
+				CellBounds.Min.X, CellBounds.Min.Y, CellBounds.Max.X, CellBounds.Max.Y,
+				TotalBounds.Min.X, TotalBounds.Min.Y, TotalBounds.Max.X, TotalBounds.Max.Y,
+				(CellBounds.Max.X - CellBounds.Min.X)* (CellBounds.Max.Y - CellBounds.Min.Y) / ((TotalBounds.Max.X - TotalBounds.Min.X) * (TotalBounds.Max.Y - TotalBounds.Min.Y)));
+
+			WindowStart.X = FMath::Clamp(FMath::Floor((CellBounds.Min.X - TotalBounds.Min.X) / QuadSizeWorld), 0, FullExportResolution.X - 1);
+			WindowStart.Y = FMath::Clamp(FMath::Floor((CellBounds.Min.Y - TotalBounds.Min.Y) / QuadSizeWorld), 0, FullExportResolution.Y - 1);
+			WindowEnd.X = FMath::Clamp(FMath::CeilToInt((CellBounds.Max.X - TotalBounds.Min.X) / QuadSizeWorld), WindowStart.X, FullExportResolution.X - 1);
+			WindowEnd.Y = FMath::Clamp(FMath::CeilToInt((CellBounds.Max.Y - TotalBounds.Min.Y) / QuadSizeWorld), WindowStart.Y, FullExportResolution.Y - 1);
+		}
+		else
+		{
+			WindowStart = FIntPoint::ZeroValue;
+			WindowEnd.X = FullExportResolution.X - 1;
+			WindowEnd.Y = FullExportResolution.Y - 1;
+		}
+
+		const FIntPoint OutputResolution(WindowEnd.X - WindowStart.X + 1, WindowEnd.Y - WindowStart.Y + 1);
+
+		UE_LOG(LogPCG, Verbose, TEXT("PCGUnpackGrassMap: %s, Export resolution %s, Window %s -> %s, Window proportion %.2f"),
+			*Context->GetExecutionSourceName(),
+			*FullExportResolution.ToString(),
+			*WindowStart.ToString(),
+			*WindowEnd.ToString(),
+			(OutputResolution.X* OutputResolution.Y) / double(FullExportResolution.X * FullExportResolution.Y));
+
+		FBox OutputBounds;
+		OutputBounds.Min.X = Context->TotalLandscapeComponentBounds.Min.X + WindowStart.X * QuadSizeWorld;
+		OutputBounds.Min.Y = Context->TotalLandscapeComponentBounds.Min.Y + WindowStart.Y * QuadSizeWorld;
+		OutputBounds.Max.X = OutputBounds.Min.X + OutputResolution.X * QuadSizeWorld;
+		OutputBounds.Max.Y = OutputBounds.Min.Y + OutputResolution.Y * QuadSizeWorld;
+		OutputBounds.Min.Z = Context->TotalLandscapeComponentBounds.Min.Z;
+		OutputBounds.Max.Z = Context->TotalLandscapeComponentBounds.Max.Z;
+
+		// Dilate the horizontal extents by half a texel in both X and Y. This will place the texel centers on the vert grid.
+		OutputBounds.Min.X -= QuadSizeWorld / 2.0;
+		OutputBounds.Min.Y -= QuadSizeWorld / 2.0;
+		OutputBounds.Max.X += QuadSizeWorld / 2.0;
+		OutputBounds.Max.Y += QuadSizeWorld / 2.0;
+
+		Context->OutputTextureTransform = FTransform(FQuat::Identity, OutputBounds.GetCenter(), OutputBounds.GetExtent());
+
 		UE::RenderCommandPipe::FSyncScope SyncScope;
 
 		RenderCaptureInterface::FScopedCapture RenderCapture(PCGGenerateLandscapeTextures::GTriggerGPUCaptureDispatches > 0, TEXT("PCGLandscapeGrassmapCapture"));
@@ -557,7 +620,7 @@ bool FPCGGenerateLandscapeTexturesElement::ExecuteInternal(FPCGContext* InContex
 		
 		TWeakPtr<FPCGContextHandle> ContextHandle = Context->GetOrCreateHandle();
 
-		ENQUEUE_RENDER_COMMAND(GenerateGrassMaps)([ContextHandle, LandscapeComponentExtent=Context->LandscapeComponentExtent, LandscapeTileCoords=MoveTemp(LandscapeTileCoords), ComponentSizeQuads=LandscapeProxy->ComponentSizeQuads, LandscapeGridScale, LandscapeZ, bGenerateHeightMap=Settings->bGenerateHeightMap](FRHICommandListImmediate& RHICmdList)
+		ENQUEUE_RENDER_COMMAND(GenerateGrassMaps)([ContextHandle, LandscapeComponentExtent=Context->LandscapeComponentExtent, LandscapeTileCoords=MoveTemp(LandscapeTileCoords), LandscapeComponentResolution, LandscapeGridScale, LandscapeZ, bGenerateHeightMap=Settings->bGenerateHeightMap, NumTilesX, NumTilesY, OutputResolution, WindowStart, WindowEnd](FRHICommandListImmediate& RHICmdList)
 		{
 			LLM_SCOPE_BYTAG(PCG);
 
@@ -582,14 +645,6 @@ bool FPCGGenerateLandscapeTexturesElement::ExecuteInternal(FPCGContext* InContex
 			// Generate grass maps. All will be generated to a single texture.
 			Context->LandscapeGrassWeightExporter->RenderLandscapeComponentToTexture_RenderThread(GraphBuilder, GrassMapTexture);
 
-			const FVector GrassMapExtent = Context->GrassMapBounds.GetExtent() * 2;
-			const uint32 NumTilesX = FMath::RoundToInt(GrassMapExtent.X / LandscapeComponentExtent);
-			const uint32 NumTilesY = FMath::RoundToInt(GrassMapExtent.Y / LandscapeComponentExtent);
-
-			// Each corner of a quad in the landscape corresponds to one texel in the grass map. 
-			const uint32 LandscapeComponentResolution = ComponentSizeQuads + 1;
-			// Internal boundaries between components in the output share a row/column, so these are dropped to get a correct world space output.
-			const FIntPoint OutputResolution((LandscapeComponentResolution - 1u) * NumTilesX + 1u, (LandscapeComponentResolution - 1u) * NumTilesY + 1u);
 			const uint32 NumGrassTypes = Context->NumGrassTypes;
 			const uint32 NumOutputSlices = FMath::Max(NumGrassTypes, 1u);
 
@@ -622,8 +677,8 @@ bool FPCGGenerateLandscapeTexturesElement::ExecuteInternal(FPCGContext* InContex
 			Parameters->InLandscapeGridScale = static_cast<FVector3f>(LandscapeGridScale);
 			Parameters->InLandscapeLocationZ = static_cast<float>(LandscapeZ);
 			Parameters->InLandscapeComponentResolution = LandscapeComponentResolution;
-			Parameters->InOutputResolution = FUintVector2(OutputResolution.X, OutputResolution.Y);
 			Parameters->InOutputHeight = bGenerateHeightMap ? 1u : 0u;
+			Parameters->InWindow = FUintVector4(WindowStart.X, WindowStart.Y, WindowEnd.X, WindowEnd.Y);
 
 			// The first 2 channels are reserved for height data. See illustration of packing in PCGGrassMapUnpackerCS.usf.
 			Parameters->InNumGrassMapPasses = FMath::DivideAndRoundUp<int>(NumGrassTypes + 2, 4);
@@ -647,6 +702,9 @@ bool FPCGGenerateLandscapeTexturesElement::ExecuteInternal(FPCGContext* InContex
 
 			FComputeShaderUtils::AddPass(GraphBuilder, RDG_EVENT_NAME("PCGUnpackGrassMap"), ERDGPassFlags::Compute, Shader, Parameters, FIntVector(GroupCountX, GroupCountY, GroupCountZ));
 
+			UE_LOG(LogPCG, Verbose, TEXT("PCGUnpackGrassMap: %s, Components %d, OutputResolution %s, OutputSlices %u"),
+				*Context->GetExecutionSourceName(), LandscapeTileCoords.Num(), *OutputResolution.ToString(), NumOutputSlices);
+
 			// Export the output textures so they can be used downstream.
 			TRefCountPtr<IPooledRenderTarget> HeightsExported;
 			if (bGenerateHeightMap)
@@ -705,22 +763,12 @@ bool FPCGGenerateLandscapeTexturesElement::ExecuteInternal(FPCGContext* InContex
 		return true;
 	}
 
-	if (!Context->OutputTextureTransform)
-	{
-		// Dilate the horizontal extents by half a texel in both X and Y. This will place the texel centers on the vert grid.
-		FVector Extent = Context->GrassMapBounds.GetExtent();
-		Extent.X += Context->TexelSizeWorld / 2.0;
-		Extent.Y += Context->TexelSizeWorld / 2.0;
-
-		Context->OutputTextureTransform = FTransform(FQuat::Identity, Context->GrassMapBounds.GetCenter(), Extent);
-	}
-
 	bool bAllTexturesInitialized = true;
 
 	if (Context->HeightHandle)
 	{
 		check(Context->HeightTextureData);
-		bAllTexturesInitialized &= Context->HeightTextureData->Initialize(Context->HeightHandle, /*TextureIndex=*/0, *Context->OutputTextureTransform, Settings->bSkipReadbackToCPU);
+		bAllTexturesInitialized &= Context->HeightTextureData->Initialize(Context->HeightHandle, /*TextureIndex=*/0, Context->OutputTextureTransform, Settings->bSkipReadbackToCPU);
 	}
 
 	if (Context->GrassMapHandle)
@@ -735,7 +783,7 @@ bool FPCGGenerateLandscapeTexturesElement::ExecuteInternal(FPCGContext* InContex
 			bAllTexturesInitialized &= TextureData->Initialize(
 				Context->GrassMapHandle,
 				/*TextureIndex=*/Context->SelectedGrassTypes[DataIndex].Get<1>(),
-				*Context->OutputTextureTransform,
+				Context->OutputTextureTransform,
 				Settings->bSkipReadbackToCPU);
 		}
 	}
diff --git a/Engine/Plugins/PCG/Source/PCG/Public/Elements/PCGGenerateLandscapeTextures.h b/Engine/Plugins/PCG/Source/PCG/Public/Elements/PCGGenerateLandscapeTextures.h
--- a/Engine/Plugins/PCG/Source/PCG/Public/Elements/PCGGenerateLandscapeTextures.h
+++ b/Engine/Plugins/PCG/Source/PCG/Public/Elements/PCGGenerateLandscapeTextures.h
@@ -100,7 +100,7 @@ public:
 	TArray<TWeakObjectPtr<ULandscapeComponent>> LandscapeComponents;
 
 	/** World-space bounds containing all of the landscape components given to the grass weight exporter. */
-	FBox GrassMapBounds = FBox(EForceInit::ForceInit);
+	FBox TotalLandscapeComponentBounds = FBox(EForceInit::ForceInit);
 
 	/** Extent (side length) of each landscape component. */
 	double LandscapeComponentExtent = 0.0;
@@ -122,7 +122,7 @@ public:
 
 	double TexelSizeWorld = 0.0;
 
-	TOptional<FTransform> OutputTextureTransform;
+	FTransform OutputTextureTransform;
 
 protected:
 	virtual void AddExtraStructReferencedObjects(FReferenceCollector& Collector) override;
diff --git a/Engine/Plugins/PCG/Source/PCGCompute/Internal/PCGGrassMapUnpackerCS.h b/Engine/Plugins/PCG/Source/PCGCompute/Internal/PCGGrassMapUnpackerCS.h
--- a/Engine/Plugins/PCG/Source/PCGCompute/Internal/PCGGrassMapUnpackerCS.h
+++ b/Engine/Plugins/PCG/Source/PCGCompute/Internal/PCGGrassMapUnpackerCS.h
@@ -28,8 +28,8 @@ public:
 		SHADER_PARAMETER(uint32, InNumGrassMapPasses)
 		SHADER_PARAMETER(FVector3f, InLandscapeGridScale)
 		SHADER_PARAMETER(float, InLandscapeLocationZ)
-		SHADER_PARAMETER(FUintVector2, InOutputResolution)
 		SHADER_PARAMETER(uint32, InOutputHeight)
+		SHADER_PARAMETER(FUintVector4, InWindow)
 	END_SHADER_PARAMETER_STRUCT()
 
 public:
