[PCG] Add settings function to get commonly used settings.

#rnx

diff --git a/Engine/Plugins/PCG/Source/PCG/Private/Helpers/PCGSettingsHelpers.cpp b/Engine/Plugins/PCG/Source/PCG/Private/Helpers/PCGSettingsHelpers.cpp
--- a/Engine/Plugins/PCG/Source/PCG/Private/Helpers/PCGSettingsHelpers.cpp
+++ b/Engine/Plugins/PCG/Source/PCG/Private/Helpers/PCGSettingsHelpers.cpp
@@ -10,6 +10,8 @@
 
 #include "UObject/EnumProperty.h"
 
+#include UE_INLINE_GENERATED_CPP_BY_NAME(PCGSettingsHelpers)
+
 namespace PCGSettingsHelpers
 {
 	template<typename T>
@@ -270,8 +272,6 @@ namespace PCGSettingsHelpers
 
 		TArray<FPCGSettingsOverridableParam> Res;
 
-		const EFieldIteratorFlags::SuperClassFlags SuperFlag = InConfig.bExcludeSuperProperties ? EFieldIteratorFlags::ExcludeSuper : EFieldIteratorFlags::IncludeSuper;
-
 		check(InClass);
 
 		auto GatherAliases = [](const FProperty* InProperty, FPCGSettingsOverridableParam& InParam) -> void
@@ -294,114 +294,190 @@ namespace PCGSettingsHelpers
 #endif // WITH_EDITOR
 		};
 
-		for (TFieldIterator<FProperty> InputIt(InClass, SuperFlag, EFieldIteratorFlags::ExcludeDeprecated); InputIt; ++InputIt)
+		const UStruct* StopClass = InConfig.bExcludeSuperProperties ? InClass->GetSuperStruct() : InConfig.StopClass;
+		for (const UStruct* CurClass = InClass; CurClass != nullptr && (StopClass == nullptr || CurClass  != StopClass); CurClass = CurClass->GetSuperStruct())
 		{
-			const FProperty* Property = *InputIt;
-			if (!Property)
+			for (TFieldIterator<FProperty> InputIt(CurClass, EFieldIteratorFlags::ExcludeSuper, EFieldIteratorFlags::ExcludeDeprecated); InputIt; ++InputIt)
 			{
-				continue;
-			}
+				const FProperty* Property = *InputIt;
+				if (!Property)
+				{
+					continue;
+				}
 
-			bool bValid = true;
+				bool bValid = true;
 
 #if WITH_EDITOR
-			auto HasMetadata = [Property](const TArray<FName>& InValues, bool bIsConjunction) -> bool
-			{
-				bool bFound = false;
-				for (const FName& Metadata : InValues)
-				{
-					if (Property->HasMetaData(Metadata) && !bIsConjunction)
+				auto HasMetadata = [Property](const TArray<FName>& InValues, bool bIsConjunction) -> bool
 					{
-						bFound = true;
-						if (!bIsConjunction)
+						bool bFound = false;
+						for (const FName& Metadata : InValues)
 						{
-							break;
+							if (Property->HasMetaData(Metadata) && !bIsConjunction)
+							{
+								bFound = true;
+								if (!bIsConjunction)
+								{
+									break;
+								}
+							}
+							else if (bIsConjunction)
+							{
+								bFound = false;
+								break;
+							}
 						}
-					}
-					else if (bIsConjunction)
-					{
-						bFound = false;
-						break;
-					}
-				}
 
-				return bFound;
-			};
+						return bFound;
+					};
 
-			if (!InConfig.IncludeMetadataValues.IsEmpty())
-			{
-				bValid &= HasMetadata(InConfig.IncludeMetadataValues, InConfig.bIncludeMetadataIsConjunction);
-			}
+				if (!InConfig.IncludeMetadataValues.IsEmpty())
+				{
+					bValid &= HasMetadata(InConfig.IncludeMetadataValues, InConfig.bIncludeMetadataIsConjunction);
+				}
 
-			if (!InConfig.ExcludeMetadataValues.IsEmpty())
-			{
-				bValid &= !HasMetadata(InConfig.ExcludeMetadataValues, InConfig.bExcludeMetadataIsConjunction);
-			}
+				if (!InConfig.ExcludeMetadataValues.IsEmpty())
+				{
+					bValid &= !HasMetadata(InConfig.ExcludeMetadataValues, InConfig.bExcludeMetadataIsConjunction);
+				}
 #endif // WITH_EDITOR
 
-			auto HasPropertyFlags = [Property](uint64 InFlags, bool bIsConjunction) -> bool
-			{
-				return bIsConjunction ? Property->HasAllPropertyFlags(InFlags) : Property->HasAnyPropertyFlags(InFlags);
-			};
-
-			if (InConfig.IncludePropertyFlags != 0)
-			{
-				bValid &= HasPropertyFlags(InConfig.IncludePropertyFlags, InConfig.bIncludePropertyFlagsIsConjunction);
-			}
+				auto HasPropertyFlags = [Property](uint64 InFlags, bool bIsConjunction) -> bool
+					{
+						return bIsConjunction ? Property->HasAllPropertyFlags(InFlags) : Property->HasAnyPropertyFlags(InFlags);
+					};
 
-			if (InConfig.ExcludePropertyFlags != 0)
-			{
-				bValid &= !HasPropertyFlags(InConfig.ExcludePropertyFlags, InConfig.bExcludePropertyFlagsIsConjunction);
-			}
+				if (InConfig.IncludePropertyFlags != 0)
+				{
+					bValid &= HasPropertyFlags(InConfig.IncludePropertyFlags, InConfig.bIncludePropertyFlagsIsConjunction);
+				}
 
-			// Don't allow to override the seed if the settings doesn't use the seed.
-			if (!InConfig.bUseSeed && Property->GetOwnerClass() && Property->GetOwnerClass()->IsChildOf<UPCGSettings>())
-			{
-				bValid &= (Property->GetFName() != GET_MEMBER_NAME_CHECKED(UPCGSettings, Seed));
-			}
+				if (InConfig.ExcludePropertyFlags != 0)
+				{
+					bValid &= !HasPropertyFlags(InConfig.ExcludePropertyFlags, InConfig.bExcludePropertyFlagsIsConjunction);
+				}
 
-			if (!bValid)
-			{
-				continue;
-			}
+				// Don't allow to override the seed if the settings doesn't use the seed.
+				if (!InConfig.bUseSeed && Property->GetOwnerClass() && Property->GetOwnerClass()->IsChildOf<UPCGSettings>())
+				{
+					bValid &= (Property->GetFName() != GET_MEMBER_NAME_CHECKED(UPCGSettings, Seed));
+				}
 
-			auto RecursiveExtraction = [&InConfig, &Res, &LabelCache, Property, InClass, &GatherAliases, &InAlreadySeenProperties](const UStruct* NextClass)
-			{
-				// Reached max depth
-				if (InConfig.MaxStructDepth == 0)
+				if (!bValid)
 				{
-					return;
+					continue;
 				}
 
-				// Use the seed, and don't check metadata for PCG overridable.
-				FPCGGetAllOverridableParamsConfig RecurseConfig = InConfig;
-				InAlreadySeenProperties.Add(Property);
-				RecurseConfig.bUseSeed = true;
+				auto RecursiveExtraction = [&InConfig, &Res, &LabelCache, Property, InClass, &GatherAliases, &InAlreadySeenProperties](const UStruct* NextClass)
+					{
+						// Reached max depth
+						if (InConfig.MaxStructDepth == 0)
+						{
+							return;
+						}
+
+						// Use the seed, and don't check metadata for PCG overridable.
+						FPCGGetAllOverridableParamsConfig RecurseConfig = InConfig;
+						InAlreadySeenProperties.Add(Property);
+						RecurseConfig.bUseSeed = true;
 #if WITH_EDITOR
-				RecurseConfig.IncludeMetadataValues.Remove(PCGObjectMetadata::Overridable);
-				RecurseConfig.IncludeMetadataValues.Remove(PCGObjectMetadata::OverridableCPUAndGPU);
-				RecurseConfig.IncludeMetadataValues.Remove(PCGObjectMetadata::OverridableCPUAndGPUWithReadback);
+						RecurseConfig.IncludeMetadataValues.Remove(PCGObjectMetadata::Overridable);
+						RecurseConfig.IncludeMetadataValues.Remove(PCGObjectMetadata::OverridableCPUAndGPU);
+						RecurseConfig.IncludeMetadataValues.Remove(PCGObjectMetadata::OverridableCPUAndGPUWithReadback);
 
-				// Also check if we have some filtering on the child properties
-				TArray<FString> ChildPropertiesFiltering;
-				if (Property->HasMetaData(PCGObjectMetadata::OverridableChildProperties))
-				{
-					const FString& ChildPropertiesMetadata = Property->GetMetaData(PCGObjectMetadata::OverridableChildProperties);
-					ChildPropertiesFiltering = PCGHelpers::GetStringArrayFromCommaSeparatedList(ChildPropertiesMetadata);
-				}
+						// Also check if we have some filtering on the child properties
+						TArray<FString> ChildPropertiesFiltering;
+						if (Property->HasMetaData(PCGObjectMetadata::OverridableChildProperties))
+						{
+							const FString& ChildPropertiesMetadata = Property->GetMetaData(PCGObjectMetadata::OverridableChildProperties);
+							ChildPropertiesFiltering = PCGHelpers::GetStringArrayFromCommaSeparatedList(ChildPropertiesMetadata);
+						}
 #endif // WITH_EDITOR
 
-				if (RecurseConfig.MaxStructDepth > 0)
+						if (RecurseConfig.MaxStructDepth > 0)
+						{
+							RecurseConfig.MaxStructDepth--;
+						}
+
+						for (FPCGSettingsOverridableParam& ChildParam : GetAllOverridableParams_Internal(NextClass, RecurseConfig, InAlreadySeenProperties))
+						{
+							FName Label = ChildParam.Label;
+							bool bHasNameClash = false;
+#if WITH_EDITOR
+							// Don't check for label clash, as they would all be equal to None in non-editor build
+							const int32 CachedLabelIndex = LabelCache.IndexOfByKey(Label);
+							if (CachedLabelIndex != INDEX_NONE)
+							{
+								// If we have a clash, we will use the full path, so mark this param and the other that clashed to use the full path.
+								Res[CachedLabelIndex].bHasNameClash = true;
+								Res[CachedLabelIndex].Label = FName(Res[CachedLabelIndex].GetPropertyPath());
+								bHasNameClash = true;
+							}
+
+							LabelCache.AddUnique(Label);
+
+							// Also check if the current property path is filtered out
+							if (!ChildPropertiesFiltering.IsEmpty() && !ChildPropertiesFiltering.Contains(ChildParam.GetPropertyPath()))
+							{
+								continue;
+							}
+#endif // WITH_EDITOR
+
+							FPCGSettingsOverridableParam& Param = Res.Emplace_GetRef();
+							Param.PropertiesNames.Add(Property->GetFName());
+							Param.PropertiesNames.Append(std::move(ChildParam.PropertiesNames));
+							Param.Properties.Add(Property);
+							Param.Properties.Append(std::move(ChildParam.Properties));
+							Param.PropertyClass = InClass;
+#if WITH_EDITOR
+							Param.UnderlyingType = ChildParam.UnderlyingType;
+#endif //WITH_EDITOR
+
+							if (!ChildParam.bHasNameClash && !bHasNameClash)
+							{
+								Param.Label = Label;
+							}
+							else
+							{
+#if WITH_EDITOR
+								Param.Label = FName(Param.GetPropertyPath());
+#endif // WITH_EDITOR
+								Param.bHasNameClash = true;
+							}
+
+#if WITH_EDITOR
+							Param.bSupportsGPU = ChildParam.bSupportsGPU;
+							Param.bRequiresGPUReadback = ChildParam.bRequiresGPUReadback;
+							GatherAliases(Property, Param);
+							// For all already found aliases, add them to the current param, with the index incremented by 1.
+							for (TPair<int32, FPCGPropertyAliases>& It : ChildParam.MapOfAliases)
+							{
+								Param.MapOfAliases.Emplace(It.Key + 1, std::move(It.Value));
+							}
+#endif // WITH_EDITOR
+						}
+
+						InAlreadySeenProperties.Pop(EAllowShrinking::No);
+					};
+
+				const FProperty* PropertyToCheck = Property;
+				if (InConfig.bExtractArrays && Property->IsA<FArrayProperty>())
 				{
-					RecurseConfig.MaxStructDepth--;
+					PropertyToCheck = CastFieldChecked<FArrayProperty>(Property)->Inner;
 				}
 
-				for (FPCGSettingsOverridableParam& ChildParam : GetAllOverridableParams_Internal(NextClass, RecurseConfig, InAlreadySeenProperties))
+				// Always extract instanced objects, but be careful with infinite recursion (class having an instanced property on itself)
+				if (PropertyToCheck->IsA<FObjectProperty>() && !InAlreadySeenProperties.Contains(PropertyToCheck) && (InConfig.bExtractObjects || PropertyToCheck->HasAllPropertyFlags(CPF_InstancedReference)))
+				{
+					RecursiveExtraction(CastFieldChecked<FObjectProperty>(PropertyToCheck)->PropertyClass);
+				}
+				// Validate that the property can be overriden by params
+				else if (PCGAttributeAccessorHelpers::IsPropertyAccessorSupported(PropertyToCheck))
 				{
-					FName Label = ChildParam.Label;
+					FName Label = NAME_None;
 					bool bHasNameClash = false;
 #if WITH_EDITOR
-					// Don't check for label clash, as they would all be equal to None in non-editor build
+					Label = *Property->GetAuthoredName();
 					const int32 CachedLabelIndex = LabelCache.IndexOfByKey(Label);
 					if (CachedLabelIndex != INDEX_NONE)
 					{
@@ -412,97 +488,25 @@ namespace PCGSettingsHelpers
 					}
 
 					LabelCache.AddUnique(Label);
-
-					// Also check if the current property path is filtered out
-					if (!ChildPropertiesFiltering.IsEmpty() && !ChildPropertiesFiltering.Contains(ChildParam.GetPropertyPath()))
-					{
-						continue;
-					}
 #endif // WITH_EDITOR
 
 					FPCGSettingsOverridableParam& Param = Res.Emplace_GetRef();
 					Param.PropertiesNames.Add(Property->GetFName());
-					Param.PropertiesNames.Append(std::move(ChildParam.PropertiesNames));
 					Param.Properties.Add(Property);
-					Param.Properties.Append(std::move(ChildParam.Properties));
 					Param.PropertyClass = InClass;
+					Param.bHasNameClash = bHasNameClash;
 #if WITH_EDITOR
-					Param.UnderlyingType = ChildParam.UnderlyingType;
-#endif //WITH_EDITOR
-					
-					if (!ChildParam.bHasNameClash && !bHasNameClash)
-					{
-						Param.Label = Label;
-					}
-					else
-					{
-#if WITH_EDITOR
-						Param.Label = FName(Param.GetPropertyPath());
-#endif // WITH_EDITOR
-						Param.bHasNameClash = true;
-					}
-
-#if WITH_EDITOR
-					Param.bSupportsGPU = ChildParam.bSupportsGPU;
-					Param.bRequiresGPUReadback = ChildParam.bRequiresGPUReadback;
+					Param.Label = bHasNameClash ? FName(Param.GetPropertyPath()) : Label;
+					Param.bSupportsGPU = Property->HasMetaData(PCGObjectMetadata::OverridableCPUAndGPU) || Property->HasMetaData(PCGObjectMetadata::OverridableCPUAndGPUWithReadback);
+					Param.bRequiresGPUReadback = Property->HasMetaData(PCGObjectMetadata::OverridableCPUAndGPUWithReadback);
 					GatherAliases(Property, Param);
-					// For all already found aliases, add them to the current param, with the index incremented by 1.
-					for (TPair<int32, FPCGPropertyAliases>& It : ChildParam.MapOfAliases)
-					{
-						Param.MapOfAliases.Emplace(It.Key + 1, std::move(It.Value));
-					}
+					Param.UnderlyingType = PCGAttributeAccessorHelpers::GetMetadataTypeForProperty(Property);
 #endif // WITH_EDITOR
 				}
-
-				InAlreadySeenProperties.Pop(EAllowShrinking::No);
-			};
-
-			const FProperty* PropertyToCheck = Property;
-			if (InConfig.bExtractArrays && Property->IsA<FArrayProperty>())
-			{
-				PropertyToCheck = CastFieldChecked<FArrayProperty>(Property)->Inner;
-			}
-
-			// Always extract instanced objects, but be careful with infinite recursion (class having an instanced property on itself)
-			if (PropertyToCheck->IsA<FObjectProperty>() && !InAlreadySeenProperties.Contains(PropertyToCheck) && (InConfig.bExtractObjects || PropertyToCheck->HasAllPropertyFlags(CPF_InstancedReference)))
-			{
-				RecursiveExtraction(CastFieldChecked<FObjectProperty>(PropertyToCheck)->PropertyClass);
-			}
-			// Validate that the property can be overriden by params
-			else if (PCGAttributeAccessorHelpers::IsPropertyAccessorSupported(PropertyToCheck))
-			{
-				FName Label = NAME_None;
-				bool bHasNameClash = false;
-#if WITH_EDITOR
-				Label = *Property->GetAuthoredName();
-				const int32 CachedLabelIndex = LabelCache.IndexOfByKey(Label);
-				if (CachedLabelIndex != INDEX_NONE)
+				else if (const FStructProperty* StructProperty = CastField<FStructProperty>(PropertyToCheck))
 				{
-					// If we have a clash, we will use the full path, so mark this param and the other that clashed to use the full path.
-					Res[CachedLabelIndex].bHasNameClash = true;
-					Res[CachedLabelIndex].Label = FName(Res[CachedLabelIndex].GetPropertyPath());
-					bHasNameClash = true;
+					RecursiveExtraction(StructProperty->Struct);
 				}
-
-				LabelCache.AddUnique(Label);
-#endif // WITH_EDITOR
-
-				FPCGSettingsOverridableParam& Param = Res.Emplace_GetRef();
-				Param.PropertiesNames.Add(Property->GetFName());
-				Param.Properties.Add(Property);
-				Param.PropertyClass = InClass;
-				Param.bHasNameClash = bHasNameClash;
-#if WITH_EDITOR
-				Param.Label = bHasNameClash ? FName(Param.GetPropertyPath()) : Label;
-				Param.bSupportsGPU = Property->HasMetaData(PCGObjectMetadata::OverridableCPUAndGPU) || Property->HasMetaData(PCGObjectMetadata::OverridableCPUAndGPUWithReadback);
-				Param.bRequiresGPUReadback = Property->HasMetaData(PCGObjectMetadata::OverridableCPUAndGPUWithReadback);
-				GatherAliases(Property, Param);
-				Param.UnderlyingType = PCGAttributeAccessorHelpers::GetMetadataTypeForProperty(Property);
-#endif // WITH_EDITOR
-			}
-			else if (const FStructProperty* StructProperty = CastField<FStructProperty>(PropertyToCheck))
-			{
-				RecursiveExtraction(StructProperty->Struct);
 			}
 		}
 
@@ -572,3 +576,35 @@ PRAGMA_ENABLE_DEPRECATION_WARNINGS
 		}
 	}
 }
+
+PCG_API TArray<FPCGSettingsPropertyDefinition> UPCGSettingsHelpers::GetCommonlyUsedProperties(TSubclassOf<UPCGSettings> InSettingsClass)
+{
+	using namespace PCGSettingsHelpers;
+
+	FPCGGetAllOverridableParamsConfig Config;
+	Config.ExcludePropertyFlags = EPropertyFlags::CPF_AdvancedDisplay;
+	Config.IncludePropertyFlags = EPropertyFlags::CPF_Edit;
+	Config.StopClass = UPCGSettings::StaticClass();
+
+	TArray<FPCGSettingsPropertyDefinition> Properties;
+	for (const FPCGSettingsOverridableParam& Param : GetAllOverridableParams(InSettingsClass, Config))
+	{
+		const FProperty* Property = Param.Properties.IsEmpty() ? nullptr : Param.Properties[0];
+		if (Property)
+		{
+			FString Tooltip;
+#if WITH_EDITOR
+			Tooltip = Property->GetToolTipText().ToString();
+#endif // WITH_EDITOR
+
+			Properties.Add(FPCGSettingsPropertyDefinition{ 
+				Param.Label.ToString() , 
+				Param.GetPropertyPath(), 
+				Tooltip,
+				Property->GetCPPType() });
+		}
+	}
+
+	return Properties;
+}
+
diff --git a/Engine/Plugins/PCG/Source/PCG/Public/Helpers/PCGSettingsHelpers.h b/Engine/Plugins/PCG/Source/PCG/Public/Helpers/PCGSettingsHelpers.h
--- a/Engine/Plugins/PCG/Source/PCG/Public/Helpers/PCGSettingsHelpers.h
+++ b/Engine/Plugins/PCG/Source/PCG/Public/Helpers/PCGSettingsHelpers.h
@@ -14,6 +14,8 @@
 #include "Metadata/Accessors/IPCGAttributeAccessor.h"
 #include "Metadata/Accessors/PCGAttributeAccessorKeys.h"
 
+#include "PCGSettingsHelpers.generated.h"
+
 #define UE_API PCG_API
 
 class UPCGComponent;
@@ -21,6 +23,36 @@ class UPCGNode;
 class UPCGPin;
 struct FPCGDataCollection;
 
+
+USTRUCT(BlueprintType)
+struct FPCGSettingsPropertyDefinition
+{
+	GENERATED_BODY()
+
+	UPROPERTY(BlueprintReadOnly, Category = Node)
+	FString Name;
+
+	UPROPERTY(BlueprintReadOnly, Category = Node)
+	FString FullName;
+
+	UPROPERTY(BlueprintReadOnly, Category = Node)
+	FString Tooltip;
+
+	UPROPERTY(BlueprintReadOnly, Category = Node)
+	FString CppType;
+};
+
+UCLASS(MinimalAPI, BlueprintType)
+class UPCGSettingsHelpers : public UBlueprintFunctionLibrary
+{
+	GENERATED_BODY()
+
+public:
+	/** Returns the list of overridable properties from the settings class. */
+	UFUNCTION(BlueprintCallable, Category = Node)
+	static PCG_API TArray<FPCGSettingsPropertyDefinition> GetCommonlyUsedProperties(TSubclassOf<UPCGSettings> InSettingsClass);
+};
+
 namespace PCGSettingsHelpers
 {
 	/** Sets data from the params to a given property, matched on a name basis */
@@ -103,6 +135,10 @@ namespace PCGSettingsHelpers
 		// Don't look for properties from parents
 		bool bExcludeSuperProperties = false;
 
+		// Don't look for properties from 'StopClass' and its parents (Should be a subclass of the looked-up class) nullptr means go all the way to root.
+		//  if bExcludeSuperProperties is true, this is ignored.
+		const UStruct* StopClass = nullptr;
+
 #if WITH_EDITOR
 		// List of metadata values to find in property metadata. Only works in editor builds as metadata on property is not available elsewise.
 		TArray<FName> IncludeMetadataValues;
