[PCG] Added command to change to any PCG editor toolmode

diff --git a/Engine/Plugins/PCG/Source/PCGEditor/Private/EditorMode/PCGEdMode.cpp b/Engine/Plugins/PCG/Source/PCGEditor/Private/EditorMode/PCGEdMode.cpp
--- a/Engine/Plugins/PCG/Source/PCGEditor/Private/EditorMode/PCGEdMode.cpp
+++ b/Engine/Plugins/PCG/Source/PCGEditor/Private/EditorMode/PCGEdMode.cpp
@@ -58,6 +58,18 @@ namespace PCGEditorMode
 			return ToolName ? *ToolName : FString(TEXT("<Invalid ToolName>"));
 		}
 	}
+
+	static FAutoConsoleCommand CommandEnablePCGTool(
+		TEXT("pcg.EnableTool"),
+		TEXT("Enables toolmode and selects the specified tool."),
+		FConsoleCommandWithArgsDelegate::CreateLambda([](const TArray<FString>& Args)
+		{
+			GLevelEditorModeTools().ActivateMode(UPCGEditorMode::EM_PCGEditorModeId, /*bToggle=*/false);
+			if (UEdMode* ThisMode = GLevelEditorModeTools().GetActiveScriptableMode(UPCGEditorMode::EM_PCGEditorModeId); ThisMode && !Args.IsEmpty())
+			{
+				ThisMode->GetInteractiveToolsContext()->StartTool(Args[0]);
+			}
+		}));
 }
 
 UPCGEditorMode::UPCGEditorMode()
