[PCG] Improve Scene Capture node

Add inclusion/exclusion by tag
Add exclusion for PCG artifacts
Support overriding the sampler filter
Fix bug where rotations on X and Z axis cause the capture view to be skewed
Rotate scene capture 90 degrees during execution instead of rotating the resulting texture data by 90 degrees

diff --git a/Engine/Plugins/PCG/Source/PCG/Private/Elements/PCGSceneCapture.cpp b/Engine/Plugins/PCG/Source/PCG/Private/Elements/PCGSceneCapture.cpp
--- a/Engine/Plugins/PCG/Source/PCG/Private/Elements/PCGSceneCapture.cpp
+++ b/Engine/Plugins/PCG/Source/PCG/Private/Elements/PCGSceneCapture.cpp
@@ -2,9 +2,14 @@
 
 #include "Elements/PCGSceneCapture.h"
 
+#include "Data/PCGRenderTargetData.h"
+#include "Helpers/PCGHelpers.h"
 #include "Helpers/PCGSettingsHelpers.h"
 
-#include "Data/PCGRenderTargetData.h"
+#include "EngineUtils.h"
+#include "Algo/AnyOf.h"
+#include "Components/PrimitiveComponent.h"
+#include "GameFramework/Actor.h"
 
 #include UE_INLINE_GENERATED_CPP_BY_NAME(PCGSceneCapture)
 
@@ -79,7 +84,7 @@ bool FPCGSceneCaptureElement::ExecuteInternal(FPCGContext* InContext) const
 		}
 
 		FTransform CaptureTransform = FTransform::Identity;
-		FBox CaptureBounds(EForceInit::ForceInit);
+		FBox CaptureBounds(-FVector(1.0), FVector(1.0));
 		bool bUnionWasCreated = false;
 
 		if (const UPCGSpatialData* BoundingShape = Context->InputData.GetSpatialUnionOfInputsByPin(Context, PCGSceneCaptureConstants::BoundingShapeLabel, bUnionWasCreated))
@@ -89,8 +94,10 @@ bool FPCGSceneCaptureElement::ExecuteInternal(FPCGContext* InContext) const
 		}
 		else
 		{
+			const FBox OriginalActorLocalBounds = ExecutionState.GetOriginalLocalSpaceBounds();
 			CaptureTransform = ExecutionState.GetTransform();
-			CaptureBounds = ExecutionState.GetBounds();
+			CaptureTransform.SetScale3D(CaptureTransform.GetScale3D() * 0.5 * (OriginalActorLocalBounds.Max - OriginalActorLocalBounds.Min));
+			CaptureBounds = CaptureBounds.TransformBy(CaptureTransform);
 		}
 
 		if (!CaptureBounds.IsValid || CaptureBounds.GetVolume() <= UE_DOUBLE_SMALL_NUMBER)
@@ -103,11 +110,11 @@ bool FPCGSceneCaptureElement::ExecuteInternal(FPCGContext* InContext) const
 		const FQuat CaptureBoundsRotation = CaptureTransform.GetRotation().GetNormalized(); // Zero quat will return identity when normalized.
 		const FVector CaptureBoundsSize = CaptureBounds.GetSize();
 		const FVector SceneCaptureLocation = CaptureBoundsLocation + CaptureBoundsRotation.GetUpVector() * (CaptureBoundsSize.Z / 2.0);
-		const FQuat SceneCaptureRotation = CaptureBoundsRotation * FQuat(CaptureBoundsRotation.GetRightVector(), FMath::DegreesToRadians(90.0f));
+		const FQuat SceneCaptureRotation = FQuat(CaptureBoundsRotation.GetUpVector(), FMath::DegreesToRadians(-90.0f)) * FQuat(CaptureBoundsRotation.GetRightVector(), FMath::DegreesToRadians(90.0f)) * CaptureBoundsRotation;
 		const float OrthoWidth = FMath::Max(CaptureBoundsSize.X, CaptureBoundsSize.Y);
 
 		Context->RenderTargetTransform = FTransform(
-			CaptureBoundsRotation * FQuat(CaptureBoundsRotation.GetUpVector(), FMath::DegreesToRadians(90.0f)),
+			CaptureBoundsRotation,
 			CaptureBoundsLocation,
 			FVector(OrthoWidth / 2.0f));
 
@@ -150,6 +157,55 @@ bool FPCGSceneCaptureElement::ExecuteInternal(FPCGContext* InContext) const
 		Context->SceneCaptureComponent->SetWorldLocation(SceneCaptureLocation);
 		Context->SceneCaptureComponent->SetWorldRotation(SceneCaptureRotation);
 
+		// Gather actors and primitives for inclusion/exclusion filtering.
+		{
+			TRACE_CPUPROFILER_EVENT_SCOPE(FilterPrimitivesForCapture);
+
+			Context->SceneCaptureComponent->PrimitiveRenderMode = Settings->bOnlyCaptureContentMatchingIncludedTags ? ESceneCapturePrimitiveRenderMode::PRM_UseShowOnlyList : ESceneCapturePrimitiveRenderMode::PRM_RenderScenePrimitives;
+
+			for (TActorIterator<AActor> It(ExecutionState.GetWorld()); It; ++It)
+			{
+				AActor* Actor = *It;
+
+				if (IsValid(Actor))
+				{
+					// Hide any actors marked for cleanup or exclusion.
+					if (Actor->ActorHasTag(PCGHelpers::MarkedForCleanupPCGTag)
+						|| (Settings->bExcludePCGContent && Actor->ActorHasTag(PCGHelpers::DefaultPCGActorTag))
+						|| Algo::AnyOf(Settings->ExcludedTags, [Actor](const FName& InTag) { return Actor->ActorHasTag(InTag); }))
+					{
+						Context->SceneCaptureComponent->HiddenActors.Add(Actor);
+					}
+					// Show any actors marked for inclusion.
+					else if (Settings->bOnlyCaptureContentMatchingIncludedTags && Algo::AnyOf(Settings->IncludedTags, [Actor](const FName& InTag) { return Actor->ActorHasTag(InTag); }))
+					{
+						Context->SceneCaptureComponent->ShowOnlyActors.Add(Actor);
+					}
+
+					TInlineComponentArray<UPrimitiveComponent*> Components(Actor);
+
+					for (UPrimitiveComponent* Component : Components)
+					{
+						if (IsValid(Component))
+						{
+							// Hide any components marked for cleanup or exclusion.
+							if (Component->ComponentHasTag(PCGHelpers::MarkedForCleanupPCGTag)
+								|| (Settings->bExcludePCGContent && Component->ComponentHasTag(PCGHelpers::DefaultPCGTag))
+								|| Algo::AnyOf(Settings->ExcludedTags, [Component](const FName& InTag) { return Component->ComponentHasTag(InTag); }))
+							{
+								Context->SceneCaptureComponent->HiddenComponents.Add(Component);
+							}
+							// Show any components marked for inclusion.
+							else if (Settings->bOnlyCaptureContentMatchingIncludedTags && Algo::AnyOf(Settings->IncludedTags, [Component](const FName& InTag) { return Component->ComponentHasTag(InTag); }))
+							{
+								Context->SceneCaptureComponent->ShowOnlyComponents.Add(Component);
+							}
+						}
+					}
+				}
+			}
+		}
+
 		// Perform scene capture.
 		Context->SceneCaptureComponent->CaptureSceneDeferred();
 		Context->bSubmittedSceneCapture = true;
@@ -176,6 +232,7 @@ bool FPCGSceneCaptureElement::ExecuteInternal(FPCGContext* InContext) const
 
 	// Initialize render target data.
 	UPCGRenderTargetData* RenderTargetData = FPCGContext::NewObject_AnyThread<UPCGRenderTargetData>(InContext);
+	RenderTargetData->Filter = Settings->Filter;
 	RenderTargetData->Initialize(Context->RenderTarget, Context->RenderTargetTransform, Settings->bSkipReadbackToCPU, /*bInTakeOwnershipOfRenderTarget=*/true);
 
 	FPCGTaggedData& OutputData = InContext->OutputData.TaggedData.Emplace_GetRef();
diff --git a/Engine/Plugins/PCG/Source/PCG/Public/Elements/PCGSceneCapture.h b/Engine/Plugins/PCG/Source/PCG/Public/Elements/PCGSceneCapture.h
--- a/Engine/Plugins/PCG/Source/PCG/Public/Elements/PCGSceneCapture.h
+++ b/Engine/Plugins/PCG/Source/PCG/Public/Elements/PCGSceneCapture.h
@@ -4,6 +4,7 @@
 
 #include "PCGContext.h"
 #include "PCGSettings.h"
+#include "Data/PCGTextureData.h"
 
 #include "Components/SceneCaptureComponent2D.h"
 #include "Engine/TextureRenderTarget2D.h"
@@ -40,10 +41,32 @@ public:
 	UPROPERTY(BlueprintReadWrite, EditAnywhere, Category = Settings, meta=(PCG_Overridable))
 	TEnumAsByte<enum ESceneCaptureSource> CaptureSource = ESceneCaptureSource::SCS_SceneColorHDR;
 
+	/** Only capture actors and primitive components that have tags listed in Included Tags. */
+	UPROPERTY(BlueprintReadWrite, EditAnywhere, Category = "Settings", meta = (PCG_Overridable))
+	bool bOnlyCaptureContentMatchingIncludedTags = false;
+
+	// @todo_pcg: Allow overriding this array when array overrides are supported.
+	/** Tags to match against actors/components when selecting for inclusion in the scene capture. */
+	UPROPERTY(BlueprintReadWrite, EditAnywhere, Category = "Settings", meta = (EditCondition = "bOnlyCaptureContentMatchingIncludedTags", EditConditionHides))
+	TArray<FName> IncludedTags;
+
+	// @todo_pcg: Allow overriding this array when array overrides are supported.
+	/** Tags to match against actors/components when selecting for exclusion from the scene capture. */
+	UPROPERTY(BlueprintReadWrite, EditAnywhere, Category = "Settings")
+	TArray<FName> ExcludedTags;
+
+	/** Exclude content created from PCG. */
+	UPROPERTY(BlueprintReadWrite, EditAnywhere, Category = "Settings", DisplayName = "Exclude PCG Content", meta = (PCG_Overridable))
+	bool bExcludePCGContent = true;
+
 	/** Size of a texel in the render target in world units (cm). */
 	UPROPERTY(BlueprintReadWrite, EditAnywhere, Category = Settings, meta=(UIMin = "1.0", ClampMin = "1.0", PCG_Overridable))
 	float TexelSize = 50.0f;
 
+	/** Method used to determine the value for a sample based on the value of nearby texels. */
+	UPROPERTY(BlueprintReadWrite, EditAnywhere, Category = Settings, meta = (PCG_Overridable))
+	EPCGTextureFilter Filter = EPCGTextureFilter::Bilinear;
+
 	/** Skip CPU readback during initialization of the render target data. */
 	UPROPERTY(BlueprintReadWrite, EditAnywhere, Category = Settings)
 	bool bSkipReadbackToCPU = false;
