[PCG] Add functionality to Custom HLSL kernel to allow doing reductions and atomic counting.

* Support static zero initialization of Custom HLSL kernel output pins, allowing all threads to assume values are initially 0
* Expose AtomicAdd

#rnx

diff --git a/Engine/Plugins/PCG/Source/PCG/Private/Compute/Elements/PCGCustomHLSL.cpp b/Engine/Plugins/PCG/Source/PCG/Private/Compute/Elements/PCGCustomHLSL.cpp
--- a/Engine/Plugins/PCG/Source/PCG/Private/Compute/Elements/PCGCustomHLSL.cpp
+++ b/Engine/Plugins/PCG/Source/PCG/Private/Compute/Elements/PCGCustomHLSL.cpp
@@ -322,7 +322,15 @@ void UPCGCustomHLSLSettings::CreateKernels(FPCGGPUCompilationContext& InOutConte
 
 	CreateParams.bRequiresOverridableParams = true;
 
-	PCGKernelHelpers::CreateKernel<UPCGCustomHLSLKernel>(InOutContext, CreateParams, OutKernels, OutEdges);
+	UPCGCustomHLSLKernel* Kernel = PCGKernelHelpers::CreateKernel<UPCGCustomHLSLKernel>(InOutContext, CreateParams, OutKernels, OutEdges);
+
+	for (const FPCGPinPropertiesGPU& OutputPin : OutputPins)
+	{
+		if (OutputPin.PropertiesGPU.bInitializeToZero)
+		{
+			Kernel->SetOutputPinRequiresZeroInitialization(OutputPin.Label);
+		}
+	}
 }
 
 TArray<FPCGDataTypeIdentifier> UPCGCustomHLSLSettings::GetAllowedInputTypes()
@@ -685,10 +693,9 @@ void UPCGCustomHLSLSettings::UpdateOutputDeclarations()
 
 		OutputDeclarations += FString::Format(TEXT(
 			"// Valid types: bool, int, float, float2, float3, float4, Rotator (float3), Quat (float4), Transform (float4x4), StringKey (int), Name (uint2)\n"
-			"\n"
+			"// Example: {0}_SetFloat({0}_DataIndex, ElementIndex, 'MyFloatAttr', MyValue);\n"
 			"void {0}_Set{type}(uint DataIndex, uint ElementIndex, 'AttributeName', {type} Value);\n"
-			"\n"
-			"// Example: {0}_SetFloat({0}_DataIndex, ElementIndex, 'MyFloatAttr', MyValue);\n"),
+			"int {0}_AtomicAddInt(uint DataIndex, uint ElementIndex, 'IntAttributeName', int ValueToAdd); // Returns value before it was incremented.\n"),
 			{ bMultiPin ? PCGHLSLElement::PinDeclTemplateStr : DataCollectionDataPins[0] });
 
 		OutputDeclarations += TEXT("\n");
@@ -1006,6 +1013,10 @@ void UPCGCustomHLSLSettings::UpdatePinSettings()
 
 		Properties.bAllowEditMultipleData = true;
 		Properties.PropertiesGPU.bAllowEditDataCount = true;
+
+		// Static initialization to zero is allowed for pins that the kernel is not controlling, and is not currently implemented for textures.
+		Properties.PropertiesGPU.bShowInitializeToZeroSetting = !bInitModeDrivenByKernel
+			&& (!!(Properties.AllowedTypes & EPCGDataType::Point) || !!(Properties.AllowedTypes & EPCGDataType::Param));
 	}
 }
 
diff --git a/Engine/Plugins/PCG/Source/PCG/Private/Compute/Elements/PCGCustomHLSLKernel.cpp b/Engine/Plugins/PCG/Source/PCG/Private/Compute/Elements/PCGCustomHLSLKernel.cpp
--- a/Engine/Plugins/PCG/Source/PCG/Private/Compute/Elements/PCGCustomHLSLKernel.cpp
+++ b/Engine/Plugins/PCG/Source/PCG/Private/Compute/Elements/PCGCustomHLSLKernel.cpp
@@ -30,6 +30,8 @@ namespace PCGCustomHLSLKernel
 {
 	constexpr TCHAR AttributeFunctionGetKeyword[] = { TEXT("Get") };
 	constexpr TCHAR AttributeFunctionSetKeyword[] = { TEXT("Set") };
+	constexpr TCHAR AttributeFunctionAtomicAddKeyword[] = { TEXT("AtomicAdd") };
+	constexpr TCHAR AttributeFunctionAtomicAddType[] = { TEXT("Int") };
 	constexpr TCHAR CopyElementFunctionKeyword[] = { TEXT("CopyElementFrom") };
 	constexpr TCHAR InitializeFunctionKeyword[] = { TEXT("InitializePoint") };
 	constexpr TCHAR StoreFunctionKeyword[] = { TEXT("Store") };
@@ -847,36 +849,38 @@ void UPCGCustomHLSLKernel::ParseShaderSource()
 		{
 			for (const FString& AttributeTypeString : AttributeTypeStrings)
 			{
-				AttributeKeywords.Add(PinStr + "_" + PCGCustomHLSLKernel::AttributeFunctionSetKeyword + AttributeTypeString);
+				AttributeKeywords.Add(PinStr + TEXT("_") + PCGCustomHLSLKernel::AttributeFunctionSetKeyword + AttributeTypeString);
 			}
 
+			AttributeKeywords.Add(PinProps.Label.ToString() + TEXT("_") + PCGCustomHLSLKernel::AttributeFunctionAtomicAddKeyword + PCGCustomHLSLKernel::AttributeFunctionAtomicAddType);
+
 			for (const FPCGPinProperties& InputPinProps : CustomHLSLSettings->InputPins)
 			{
 				if (InputPinProps.AllowedTypes == PinProps.AllowedTypes)
 				{
-					CopyElementKeywords.Add(PinStr + "_" + PCGCustomHLSLKernel::CopyElementFunctionKeyword + "_" + InputPinProps.Label.ToString());
+					CopyElementKeywords.Add(PinStr + TEXT("_") + PCGCustomHLSLKernel::CopyElementFunctionKeyword + TEXT("_") + InputPinProps.Label.ToString());
 				}
 			}
 
-			InitializeKeywords.Add(PinStr + "_" + PCGCustomHLSLKernel::InitializeFunctionKeyword);
+			InitializeKeywords.Add(PinStr + TEXT("_") + PCGCustomHLSLKernel::InitializeFunctionKeyword);
 		}
 		else if (!!(PinProps.AllowedTypes & EPCGDataType::BaseTexture))
 		{
-			InitializeKeywords.Add(PinStr + "_" + PCGCustomHLSLKernel::StoreFunctionKeyword);
+			InitializeKeywords.Add(PinStr + TEXT("_") + PCGCustomHLSLKernel::StoreFunctionKeyword);
 		}
 
 		if (!!(PinProps.AllowedTypes & EPCGDataType::Point))
 		{
-			PointPropertySetterKeywords.Add(PinStr + "_" + PCGCustomHLSLKernel::SetPositionKeyword);
-			PointPropertySetterKeywords.Add(PinStr + "_" + PCGCustomHLSLKernel::SetRotationKeyword);
-			PointPropertySetterKeywords.Add(PinStr + "_" + PCGCustomHLSLKernel::SetScaleKeyword);
-			PointPropertySetterKeywords.Add(PinStr + "_" + PCGCustomHLSLKernel::SetBoundsMinKeyword);
-			PointPropertySetterKeywords.Add(PinStr + "_" + PCGCustomHLSLKernel::SetBoundsMaxKeyword);
-			PointPropertySetterKeywords.Add(PinStr + "_" + PCGCustomHLSLKernel::SetColorKeyword);
-			PointPropertySetterKeywords.Add(PinStr + "_" + PCGCustomHLSLKernel::SetDensityKeyword);
-			PointPropertySetterKeywords.Add(PinStr + "_" + PCGCustomHLSLKernel::SetSteepnessKeyword);
-			PointPropertySetterKeywords.Add(PinStr + "_" + PCGCustomHLSLKernel::SetSeedKeyword);
-			PointPropertySetterKeywords.Add(PinStr + "_" + PCGCustomHLSLKernel::SetPointTransformKeyword);
+			PointPropertySetterKeywords.Add(PinStr + TEXT("_") + PCGCustomHLSLKernel::SetPositionKeyword);
+			PointPropertySetterKeywords.Add(PinStr + TEXT("_") + PCGCustomHLSLKernel::SetRotationKeyword);
+			PointPropertySetterKeywords.Add(PinStr + TEXT("_") + PCGCustomHLSLKernel::SetScaleKeyword);
+			PointPropertySetterKeywords.Add(PinStr + TEXT("_") + PCGCustomHLSLKernel::SetBoundsMinKeyword);
+			PointPropertySetterKeywords.Add(PinStr + TEXT("_") + PCGCustomHLSLKernel::SetBoundsMaxKeyword);
+			PointPropertySetterKeywords.Add(PinStr + TEXT("_") + PCGCustomHLSLKernel::SetColorKeyword);
+			PointPropertySetterKeywords.Add(PinStr + TEXT("_") + PCGCustomHLSLKernel::SetDensityKeyword);
+			PointPropertySetterKeywords.Add(PinStr + TEXT("_") + PCGCustomHLSLKernel::SetSteepnessKeyword);
+			PointPropertySetterKeywords.Add(PinStr + TEXT("_") + PCGCustomHLSLKernel::SetSeedKeyword);
+			PointPropertySetterKeywords.Add(PinStr + TEXT("_") + PCGCustomHLSLKernel::SetPointTransformKeyword);
 		}
 	}
 
@@ -937,7 +941,7 @@ void UPCGCustomHLSLKernel::ParseShaderSource()
 			const FString& FuncStr = InAttributeFunction.FuncStr;
 			const FString& TypeStr = InAttributeFunction.TypeStr;
 			const FString& NameStr = InAttributeFunction.NameStr;
-			const FString UsageString = PinStr + "_" + FuncStr + TypeStr;
+			const FString UsageString = PinStr + TEXT("_") + FuncStr + TypeStr;
 
 			if (PinStr.IsEmpty() || FuncStr.IsEmpty() || TypeStr.IsEmpty() || NameStr.IsEmpty())
 			{
@@ -984,20 +988,36 @@ void UPCGCustomHLSLKernel::ParseShaderSource()
 					TokenString.FindLastChar('_', DelimiterIndex);
 
 					AttributeFunctionMatch.PinStr = TokenString.Left(DelimiterIndex);
-					AttributeFunctionMatch.FuncStr = TokenString.Mid(DelimiterIndex + 1, 3);
-					AttributeFunctionMatch.TypeStr = TokenString.Mid(DelimiterIndex + 4, Token.Range.Len() - DelimiterIndex + 1);
+
+					int32 TypeIndex = 0;
+					for (; TypeIndex < AttributeTypeStrings.Num(); ++TypeIndex)
+					{
+						if (TokenString.EndsWith(AttributeTypeStrings[TypeIndex]))
+						{
+							break;
+						}
+					}
+
+					if (TypeIndex == AttributeTypeStrings.Num())
+					{
+						UE_LOG(LogPCG, Error, TEXT("Invalid attribute type in shader source: '%s'."), *TokenString);
+						continue;
+					}
+
+					AttributeFunctionMatch.TypeStr = AttributeTypeStrings[TypeIndex];
+					AttributeFunctionMatch.FuncStr = TokenString.Mid(DelimiterIndex + 1, Token.Range.Len() - AttributeFunctionMatch.TypeStr.Len() - DelimiterIndex - 1);
 					AttributeFunctionMatch.MatchBeginning = Token.Range.BeginIndex;
 					bLookingForAttributeFunctionMatch = true;
 
-					if (AttributeFunctionMatch.FuncStr == PCGCustomHLSLKernel::AttributeFunctionSetKeyword)
+					if (AttributeFunctionMatch.FuncStr == PCGCustomHLSLKernel::AttributeFunctionSetKeyword || AttributeFunctionMatch.FuncStr == PCGCustomHLSLKernel::AttributeFunctionAtomicAddKeyword)
 					{
 						InOutParsedSource.InitializedOutputPins.Add(AttributeFunctionMatch.PinStr);
 					}
 				}
 				else if (CopyElementKeywords.Contains(TokenString))
 				{
-					const int32 FirstDelimiterIndex = TokenString.Find("_");
-					const int32 SecondDelimiterIndex = TokenString.Find("_", ESearchCase::IgnoreCase, ESearchDir::FromEnd);
+					const int32 FirstDelimiterIndex = TokenString.Find(TEXT("_"));
+					const int32 SecondDelimiterIndex = TokenString.Find(TEXT("_"), ESearchCase::IgnoreCase, ESearchDir::FromEnd);
 					
 					FString TargetPin = TokenString.Left(FirstDelimiterIndex);
 					FString SourcePin = TokenString.RightChop(SecondDelimiterIndex + 1);
@@ -1007,14 +1027,14 @@ void UPCGCustomHLSLKernel::ParseShaderSource()
 				}
 				else if (InitializeKeywords.Contains(TokenString))
 				{
-					const int32 DelimiterIndex = TokenString.Find("_");
+					const int32 DelimiterIndex = TokenString.Find(TEXT("_"));
 					FString PinStr = TokenString.Left(DelimiterIndex);
 
 					InOutParsedSource.InitializedOutputPins.Add(MoveTemp(PinStr));
 				}
 				else if (PointPropertySetterKeywords.Contains(TokenString))
 				{
-					const int32 DelimiterIndex = TokenString.Find("_");
+					const int32 DelimiterIndex = TokenString.Find(TEXT("_"));
 					const FString PinStr = TokenString.Left(DelimiterIndex);
 					const FString PropertyStr = TokenString.RightChop(DelimiterIndex + 1);
 					
@@ -1810,7 +1830,7 @@ bool UPCGCustomHLSLKernel::AreAttributesValid(FPCGContext* InContext, FText* Out
 					return FText::FromString(ParsedFunction.PinLabel + TEXT("_") + ParsedFunction.FunctionName + TypeStr);
 				};
 
-				if (FunctionName == PCGCustomHLSLKernel::AttributeFunctionSetKeyword)
+				if (FunctionName == PCGCustomHLSLKernel::AttributeFunctionSetKeyword || FunctionName == PCGCustomHLSLKernel::AttributeFunctionAtomicAddKeyword)
 				{
 					if (const TSharedPtr<const FPCGDataCollectionDesc>* PinDescPtr = OutputPinDescs.Find(PinLabel))
 					{
diff --git a/Engine/Plugins/PCG/Source/PCG/Private/Compute/Elements/PCGCustomHLSLKernel.h b/Engine/Plugins/PCG/Source/PCG/Private/Compute/Elements/PCGCustomHLSLKernel.h
--- a/Engine/Plugins/PCG/Source/PCG/Private/Compute/Elements/PCGCustomHLSLKernel.h
+++ b/Engine/Plugins/PCG/Source/PCG/Private/Compute/Elements/PCGCustomHLSLKernel.h
@@ -129,6 +129,9 @@ public:
 	virtual void GetInputPins(TArray<FPCGPinProperties>& OutPins) const override;
 	virtual void GetOutputPins(TArray<FPCGPinPropertiesGPU>& OutPins) const override;
 
+	void SetOutputPinRequiresZeroInitialization(FName InOutputPinLabel) { PinsRequiringZeroInitialization.AddUnique(InOutputPinLabel); }
+	virtual bool DoesOutputPinRequireZeroInitialization(FName InOutputPinLabel) const override { return PinsRequiringZeroInitialization.Contains(InOutputPinLabel); }
+
 protected:
 #if WITH_EDITOR
 	virtual void InitializeInternal() override;
@@ -177,4 +180,8 @@ protected:
 	/** Properties that need to be fully allocated per pin. Note, we use an integer instead of EPCGPointNativeProperties directly because bitflag enums don't serialize correctly. */
 	UPROPERTY()
 	TMap<FName, /*EPCGPointNativeProperties*/int32> PinAllocatedProperties;
+
+	/** Pins for which the output buffer should be statically initialized to 0 before execution. */
+	UPROPERTY()
+	TArray<FName> PinsRequiringZeroInitialization;
 };
diff --git a/Engine/Plugins/PCG/Source/PCG/Public/Compute/PCGPinPropertiesGPU.h b/Engine/Plugins/PCG/Source/PCG/Public/Compute/PCGPinPropertiesGPU.h
--- a/Engine/Plugins/PCG/Source/PCG/Public/Compute/PCGPinPropertiesGPU.h
+++ b/Engine/Plugins/PCG/Source/PCG/Public/Compute/PCGPinPropertiesGPU.h
@@ -13,8 +13,8 @@
 UENUM()
 enum class EPCGPinInitMode : uint8
 {
-	FromInputPins UMETA(Tooltip = "Initialize from one or more input pins."),
-	Custom UMETA(Tooltip = "Initialize from fixed values."),
+	FromInputPins UMETA(Tooltip = "Initialize structure of data (data count, element count, attribute definitions, etc) from one or more input pins."),
+	Custom UMETA(Tooltip = "Initialize structure of data (data count, element count, attribute definitions, etc) from fixed values."),
 };
 
 /** Method for computing data count. */
@@ -107,6 +107,10 @@ struct FPCGPinPropertiesGPUStruct
 	UPROPERTY(BlueprintReadWrite, EditAnywhere, Category = Settings, meta = (EditCondition = "bAllowEditInitMode && (InitializationMode == EPCGPinInitMode::Custom || ElementCountMode == EPCGElementCountMode::FromInputData)", EditConditionHides, ClampMin = 1))
 	int ElementCountMultiplier = 1;
 
+	/** Statically initialize data to 0 from CPU before execution. */
+	UPROPERTY(BlueprintReadWrite, EditAnywhere, Category = Settings, meta = (EditCondition = "bShowInitializeToZeroSetting", EditConditionHides, HideEditConditionToggle))
+	bool bInitializeToZero = false;
+
 	/** How to inherit attribute names, types, and values. */
 	UPROPERTY(BlueprintReadWrite, EditAnywhere, Category = Settings, meta = (EditCondition = "bAllowEditInitMode && InitializationMode == EPCGPinInitMode::FromInputPins && !bShowTexturePinSettings", EditConditionHides))
 	EPCGAttributeInheritanceMode AttributeInheritanceMode = EPCGAttributeInheritanceMode::CopyAttributeSetup;
@@ -128,6 +132,9 @@ struct FPCGPinPropertiesGPUStruct
 
 	UPROPERTY(Transient)
 	bool bShowTexturePinSettings = false;
+
+	UPROPERTY(Transient)
+	bool bShowInitializeToZeroSetting = false;
 #endif // WITH_EDITORONLY_DATA
 };
 
