[PCG] Allow required execution dependency pin on the generic get user parameter node

#rnx

diff --git a/Engine/Plugins/PCG/Source/PCG/Public/Elements/PCGUserParameterGet.h b/Engine/Plugins/PCG/Source/PCG/Public/Elements/PCGUserParameterGet.h
--- a/Engine/Plugins/PCG/Source/PCG/Public/Elements/PCGUserParameterGet.h
+++ b/Engine/Plugins/PCG/Source/PCG/Public/Elements/PCGUserParameterGet.h
@@ -120,7 +120,6 @@ protected:
 	virtual TArray<FPCGPinProperties> InputPinProperties() const override { return {}; }
 	virtual TArray<FPCGPinProperties> OutputPinProperties() const override;
 	virtual FPCGElementPtr CreateElement() const override;
-	virtual bool HasExecutionDependencyPin() const override { return false; }
 
 #if WITH_EDITOR
 	virtual FName GetDefaultNodeName() const override { return FName("GetGenericGraphParameter"); }
