[PCG] Replaced get transforms copy by a hierarchy-aware call, until it is fixed.

diff --git a/Engine/Plugins/PCG/Source/PCG/Private/Elements/PCGCollapsePoints.cpp b/Engine/Plugins/PCG/Source/PCG/Private/Elements/PCGCollapsePoints.cpp
--- a/Engine/Plugins/PCG/Source/PCG/Private/Elements/PCGCollapsePoints.cpp
+++ b/Engine/Plugins/PCG/Source/PCG/Private/Elements/PCGCollapsePoints.cpp
@@ -10,6 +10,7 @@
 #include "SpatialAlgo/PCGOctreeQueries.h"
 
 #include "Algo/MaxElement.h"
+#include "Algo/Transform.h"
 
 #include UE_INLINE_GENERATED_CPP_BY_NAME(PCGCollapsePoints)
 
@@ -307,11 +308,15 @@ PRAGMA_DISABLE_DEPRECATION_WARNINGS
 		OutState.SourceData = Cast<UPCGPointData>(PointData);
 PRAGMA_ENABLE_DEPRECATION_WARNINGS
 
+		const int32 NumPoints = PointData->GetNumPoints();
+
 		// Get points, octree will be rebuilt later
-		OutState.PointTransforms = PointData->GetTransformsCopy();
+		OutState.PointTransforms.Reserve(NumPoints);
+
+		const TConstPCGValueRange<FTransform> ReadValueRange = PointData->GetConstTransformValueRange();
+		Algo::Transform(ReadValueRange, OutState.PointTransforms, [](const FTransform& Transform) { return Transform; });
 
 		// Get weights
-		const int32 NumPoints = PointData->GetNumPoints();
 		OutState.Weights.SetNumUninitialized(NumPoints);
 
 		ensure(OutState.PointTransforms.Num() == OutState.Weights.Num());
