[PCG] PIE: Load as preview fix for managed components

diff --git a/Engine/Plugins/PCG/Source/PCG/Private/PCGComponent.cpp b/Engine/Plugins/PCG/Source/PCG/Private/PCGComponent.cpp
--- a/Engine/Plugins/PCG/Source/PCG/Private/PCGComponent.cpp
+++ b/Engine/Plugins/PCG/Source/PCG/Private/PCGComponent.cpp
@@ -4194,6 +4194,12 @@ void UPCGComponent::MarkResourcesAsTransientOnLoad()
 		}
 	}
 
+	// Force duplicate owner for PIE so that transient changes are reflected in game
+	if (AActor* Owner = GetOwner())
+	{
+		Owner->bForceDuplicateInPIE = true;
+	}
+
 	LoadedPreviewResources = GeneratedResources;
 }
 
diff --git a/Engine/Plugins/PCG/Source/PCG/Private/PCGManagedResource.cpp b/Engine/Plugins/PCG/Source/PCG/Private/PCGManagedResource.cpp
--- a/Engine/Plugins/PCG/Source/PCG/Private/PCGManagedResource.cpp
+++ b/Engine/Plugins/PCG/Source/PCG/Private/PCGManagedResource.cpp
@@ -384,6 +384,15 @@ void UPCGManagedComponent::HideComponent()
 	if (GeneratedComponent.IsValid())
 	{
 		GeneratedComponent->UnregisterComponent();
+
+		// If transient on load, hide for PIE
+		if (bMarkedTransientOnLoad)
+		{
+			if (USceneComponent* SceneComponent = Cast<USceneComponent>(GeneratedComponent.Get()))
+			{
+				SceneComponent->bHiddenInGame = true;
+			}
+		}
 	}
 }
 
@@ -392,6 +401,15 @@ void UPCGManagedComponentList::HideComponent(int32 ComponentIndex)
 	if (GeneratedComponents[ComponentIndex].IsValid())
 	{
 		GeneratedComponents[ComponentIndex]->UnregisterComponent();
+
+		// If transient on load, hide for PIE
+		if (bMarkedTransientOnLoad)
+		{
+			if (USceneComponent* SceneComponent = Cast<USceneComponent>(GeneratedComponents[ComponentIndex].Get()))
+			{
+				SceneComponent->bHiddenInGame = true;
+			}
+		}
 	}
 }
 #endif // WITH_EDITOR
