[PCG] Combined Submission for Data View type

CLs: 44349862, 44569695, 44569728, 44569766, 44569794, 45852129

diff --git a/Engine/Plugins/PCG/Source/PCG/Private/Data/DataView/PCGDataView.cpp b/Engine/Plugins/PCG/Source/PCG/Private/Data/DataView/PCGDataView.cpp
new file mode 100644
--- /dev/null
+++ b/Engine/Plugins/PCG/Source/PCG/Private/Data/DataView/PCGDataView.cpp
@@ -0,0 +1,65 @@
+// Copyright Epic Games, Inc. All Rights Reserved.
+
+#include "Data/DataView/PCGDataView.h"
+
+#include "Data/DataView/PCGDataViewInterface.h"
+#include "Data/DataView/PCGDataViewNativePropertySelectors.h"
+
+namespace PCGDataView::Constants
+{
+	const FGuid FDataViewVersion::GUID = FGuid(0x04E74488, 0x4BAC8717, 0xBBB18694, 0x39F8F3CF);
+	FCustomVersionRegistration GRegisterPCGDataViewCustomVersion(FDataViewVersion::GUID, FDataViewVersion::LatestVersion, FDataViewVersion::FriendlyName);
+}
+
+bool FPCGDataView::IsValid() const
+{
+	return ViewedData && (Selection.bAllAttributes || !Selection.Attributes.IsEmpty());
+}
+
+TArray<FPCGMetadataDomainID> FPCGDataView::GetAllMetadataDomainIDs() const
+{
+	if (const UPCGData* Data = Cast<UPCGData>(ViewedData))
+	{
+		if (Selection.bAllAttributes)
+		{
+			return Data->GetAllSupportedMetadataDomainIDs();
+		}
+		else
+		{
+			TArray<FPCGMetadataDomainID> SupportedDomains;
+			for (const FPCGAttributePropertySelector& Selector : Selection.Attributes)
+			{
+				if (FPCGMetadataDomainID DomainID = Data->GetMetadataDomainIDFromSelector(Selector); DomainID.IsValid())
+				{
+					// Translate default domain ID to actual domain ID.
+					SupportedDomains.AddUnique(DomainID.IsDefault() ? Data->GetDefaultMetadataDomainID() : DomainID);
+				}
+			}
+
+			return SupportedDomains;
+		}
+	}
+
+	return {};
+}
+
+TArray<FPCGAttributePropertySelector> FPCGDataView::GetResolvedSelection() const
+{
+	if (Selection.bAllAttributes)
+	{
+		TArray<FPCGAttributePropertySelector> SelectedAttributes;
+
+		// Collect all properties defined by the data type
+		const FPCGDataViewRegistry& DataViewRegistry = FPCGModule::GetConstPCGDataViewRegistry();
+		if (const IPCGDataViewPropertySelector* PropertySelection = DataViewRegistry.GetPropertySelector(*this))
+		{
+			SelectedAttributes.Append(PropertySelection->GetSelection(*this));
+		}
+
+		PCGDataView::Helpers::AppendAllAttributeSelectors(ViewedData, SelectedAttributes);
+
+		return SelectedAttributes;
+	}
+
+	return Selection.Attributes;
+}
diff --git a/Engine/Plugins/PCG/Source/PCG/Private/Data/DataView/PCGDataViewData.cpp b/Engine/Plugins/PCG/Source/PCG/Private/Data/DataView/PCGDataViewData.cpp
new file mode 100644
--- /dev/null
+++ b/Engine/Plugins/PCG/Source/PCG/Private/Data/DataView/PCGDataViewData.cpp
@@ -0,0 +1,127 @@
+// Copyright Epic Games, Inc. All Rights Reserved.
+
+#include "Data/DataView/PCGDataViewData.h"
+
+#include "PCGParamData.h"
+#include "Data/PCGSplineData.h"
+#include "Data/DataView/PCGDataViewNativeJsonConverters.h"
+#include "Elements/PCGConvertToDataView.h"
+
+#include UE_INLINE_GENERATED_CPP_BY_NAME(PCGDataViewData)
+
+#define LOCTEXT_NAMESPACE "PCGDataViewData"
+
+PCG_DEFINE_TYPE_INFO(FPCGDataTypeInfoDataView, UPCGDataViewData)
+
+bool FPCGDataTypeInfoDataView::SupportsConversionFrom(const FPCGDataTypeIdentifier& InputType, const FPCGDataTypeIdentifier& ThisType, TSubclassOf<UPCGSettings>* OptionalOutConversionSettings, FText* OptionalOutCompatibilityMessage) const
+{
+	// @todo_pcg: Start with supporting Spatial and Param, but can be expanded to any PCG Data type.
+	if (InputType.IsChildOf(FPCGDataTypeInfoSpatial::AsId()) ||
+		InputType.IsSameType(FPCGDataTypeInfoParam::AsId()))
+	{
+		if (OptionalOutConversionSettings)
+		{
+			*OptionalOutConversionSettings = UPCGConvertToDataViewSettings::StaticClass();
+		}
+
+		return true;
+	}
+	else
+	{
+		if (OptionalOutCompatibilityMessage)
+		{
+			*OptionalOutCompatibilityMessage = FText::Format(LOCTEXT("TypeCompatibilityNotSupported", "The '{0}' type is not yet supported by Data Views."), InputType.ToDisplayText());
+		}
+
+		return false;
+	}
+}
+
+void UPCGDataViewData::Initialize(FPCGDataView&& InDataView)
+{
+	DataView = InDataView;
+}
+
+const FPCGDataView& UPCGDataViewData::GetDataView() const
+{
+	return DataView;
+}
+
+FPCGMetadataDomainID UPCGDataViewData::GetMetadataDomainIDFromSelector(const FPCGAttributePropertySelector& InSelector) const
+{
+	if (!ensure(DataView.ViewedData))
+	{
+		return {};
+	}
+
+	return DataView.ViewedData->GetMetadataDomainIDFromSelector(InSelector);
+}
+
+FPCGMetadataDomainID UPCGDataViewData::GetDefaultMetadataDomainID() const
+{
+	if (!ensure(DataView.ViewedData))
+	{
+		return PCGMetadataDomainID::Default;
+	}
+
+	TArray<FPCGMetadataDomainID> SupportedDomainIDs = GetAllSupportedMetadataDomainIDs();
+	if (SupportedDomainIDs.IsEmpty())
+	{
+		return {};
+	}
+
+	const FPCGMetadataDomainID DefaultDomainID = DataView.ViewedData->GetDefaultMetadataDomainID();
+
+	FPCGMetadataDomainID DomainID{};
+	DomainID = SupportedDomainIDs.Contains(DefaultDomainID) ? DefaultDomainID : SupportedDomainIDs[0];
+
+	return DomainID;
+}
+
+TArray<FPCGMetadataDomainID> UPCGDataViewData::GetAllSupportedMetadataDomainIDs() const
+{
+	if (!ensure(DataView.ViewedData))
+	{
+		return {};
+	}
+
+	if (DataView.Selection.bAllAttributes)
+	{
+		return DataView.ViewedData->GetAllSupportedMetadataDomainIDs();
+	}
+	else
+	{
+		TArray<FPCGMetadataDomainID> SupportedDomains;
+		for (const FPCGAttributePropertySelector& Selector : DataView.Selection.Attributes)
+		{
+			if (FPCGMetadataDomainID DomainID = DataView.ViewedData->GetMetadataDomainIDFromSelector(Selector); DomainID.IsValid())
+			{
+				SupportedDomains.AddUnique(DomainID.IsDefault() ? DataView.ViewedData->GetDefaultMetadataDomainID() : DomainID);
+			}
+		}
+
+		return SupportedDomains;
+	}
+}
+
+const UPCGMetadata* UPCGDataViewData::ConstMetadata() const
+{
+	if (!ensure(DataView.ViewedData))
+	{
+		return nullptr;
+	}
+
+	return DataView.ViewedData->ConstMetadata();
+}
+
+TArray<FPCGAttributePropertySelector> UPCGDataViewData::GetSelectedAttributes() const
+{
+	if (!ensure(DataView.ViewedData))
+	{
+		return {};
+	}
+
+	return DataView.GetResolvedSelection();
+}
+
+#undef LOCTEXT_NAMESPACE
diff --git a/Engine/Plugins/PCG/Source/PCG/Private/Data/DataView/PCGDataViewInterface.cpp b/Engine/Plugins/PCG/Source/PCG/Private/Data/DataView/PCGDataViewInterface.cpp
new file mode 100644
--- /dev/null
+++ b/Engine/Plugins/PCG/Source/PCG/Private/Data/DataView/PCGDataViewInterface.cpp
@@ -0,0 +1,27 @@
+// Copyright Epic Games, Inc. All Rights Reserved.
+
+#include "Data/DataView/PCGDataViewInterface.h"
+
+#include "Misc/FileHelper.h"
+
+#define LOCTEXT_NAMESPACE "PCGDataViewInterface"
+
+TValueOrError<void, FText> UPCGDataViewConverterBase::SerializeToFile(const FPCGDataView& InDataView, const FInstancedStruct& Parameters, const FString& FilePath) const
+{
+	TValueOrError<FString, FText> Result = SerializeToString(InDataView, Parameters);
+
+	if (Result.HasError())
+	{
+		return MakeError(Result.StealError());
+	}
+
+	if (!FFileHelper::SaveStringToFile(Result.GetValue(), *FilePath))
+	{
+		FText ErrorText = FText::Format(LOCTEXT("InvalidFilePath", "Could not create a valid file with path: {0}"), FText::FromStringView(FilePath));
+		return MakeError(MoveTemp(ErrorText));
+	}
+
+	return MakeValue();
+}
+
+#undef LOCTEXT_NAMESPACE
diff --git a/Engine/Plugins/PCG/Source/PCG/Private/Data/DataView/PCGDataViewNativeJsonConverters.cpp b/Engine/Plugins/PCG/Source/PCG/Private/Data/DataView/PCGDataViewNativeJsonConverters.cpp
new file mode 100644
--- /dev/null
+++ b/Engine/Plugins/PCG/Source/PCG/Private/Data/DataView/PCGDataViewNativeJsonConverters.cpp
@@ -0,0 +1,114 @@
+// Copyright Epic Games, Inc. All Rights Reserved.
+
+#include "Data/DataView/PCGDataViewNativeJsonConverters.h"
+
+#include "Data/DataView/PCGDataView.h"
+#include "Data/DataView/PCGDataViewNativePropertySelectors.h"
+
+#define LOCTEXT_NAMESPACE "PCGDataViewNativeJsonConverters"
+
+TValueOrError<FInstancedStruct, FText> UPCGDataViewJsonConverter::SerializeToTargetStruct(const FPCGDataView& InDataView, const FInstancedStruct& Parameters) const
+{
+	using namespace PCG::IO;
+
+	TSharedPtr<FJsonObject> RootJsonObject = MakeShared<FJsonObject>();
+
+	// Overridable
+	BuildJsonHeader(InDataView, RootJsonObject);
+
+	const TArray<FPCGAttributePropertySelector> Selection = InDataView.GetResolvedSelection();
+	if (!Selection.IsEmpty())
+	{
+		const UPCGData* ViewedData = InDataView.ViewedData;
+
+		if (!ensure(ViewedData))
+		{
+			return MakeError(LOCTEXT("InvalidViewedData", "Viewed data was invalid"));
+		}
+
+		// @todo_pcg: Currently only support Data Domain and Elements domain, but should be extended eventually to handle any domain.
+
+		if (ViewedData->IsSupportedMetadataDomainID(PCGMetadataDomainID::Data))
+		{
+			// Serialize Data Domain attributes first, immediately after the header.
+			const bool bDataIsDefaultDomain = InDataView.ViewedData->GetDefaultMetadataDomainID() == PCGMetadataDomainID::Data;
+			const TArray<FPCGAttributePropertySelector> DataDomainAttributes = Selection.FilterByPredicate([ViewedData, bDataIsDefaultDomain](const FPCGAttributePropertySelector& Selector)
+			{
+				const FPCGMetadataDomainID DomainID = ViewedData->GetMetadataDomainIDFromSelector(Selector);
+				return (DomainID == PCGMetadataDomainID::Data) || (bDataIsDefaultDomain && DomainID == PCGMetadataDomainID::Default);
+			});
+			Json::Helpers::AppendSelectionByAttribute(RootJsonObject, ViewedData, DataDomainAttributes);
+		}
+
+		// If this data supports the element domain, catalogue the elements and serialize them based on the provided layout.
+		if (InDataView.ViewedData->IsSupportedMetadataDomainID(PCGMetadataDomainID::Elements))
+		{
+			const bool bElementsAreDefaultDomain = InDataView.ViewedData->GetDefaultMetadataDomainID() == PCGMetadataDomainID::Elements;
+			const TArray<FPCGAttributePropertySelector> ElementsDomainAttributes = Selection.FilterByPredicate([&ViewedData = InDataView.ViewedData, bElementsAreDefaultDomain](const FPCGAttributePropertySelector& Selector)
+			{
+				const FPCGMetadataDomainID DomainID = ViewedData->GetMetadataDomainIDFromSelector(Selector);
+				return (DomainID == PCGMetadataDomainID::Elements) || (bElementsAreDefaultDomain && DomainID == PCGMetadataDomainID::Default);
+			});
+
+			if (InDataView.AttributeLayout == EPCGDataViewAttributeLayout::ByElement)
+			{
+				Json::Helpers::AppendSelectionByElement(RootJsonObject, InDataView.ViewedData, ElementsDomainAttributes);
+			}
+			else if (InDataView.AttributeLayout == EPCGDataViewAttributeLayout::ByAttribute)
+			{
+				Json::Helpers::AppendSelectionByAttribute(RootJsonObject, InDataView.ViewedData, ElementsDomainAttributes);
+			}
+		}
+	}
+
+	FInstancedStruct InstancedStruct;
+	InstancedStruct.InitializeAs<FPCGDataViewJsonOutput>(MoveTemp(RootJsonObject));
+	return MakeValue(InstancedStruct);
+}
+
+TValueOrError<FString, FText> UPCGDataViewJsonConverter::SerializeToString(const FPCGDataView& InDataView, const FInstancedStruct& Parameters) const
+{
+	TValueOrError<FInstancedStruct, FText> Result = SerializeToTargetStruct(InDataView, Parameters);
+	if (Result.HasError())
+	{
+		return MakeError(Result.StealError());
+	}
+
+	const FInstancedStruct DataStruct = Result.GetValue();
+	if (!ensure(DataStruct.GetScriptStruct() == FPCGDataViewJsonOutput::StaticStruct()))
+	{
+		return MakeError(LOCTEXT("DataStructTypesMismatch", "Struct types between result and FPCGDataViewJsonOutput."));
+	}
+
+	if (DataStruct.IsValid() && Parameters.IsValid())
+	{
+		const FPCGDataViewJsonOutput& JsonData = DataStruct.Get<FPCGDataViewJsonOutput>();
+		if (!ensure(Parameters.GetScriptStruct() == FPCGDataViewJsonParameters::StaticStruct()))
+		{
+			return MakeError(LOCTEXT("ParamStructTypesMismatch", "Struct types between parameters and FPCGDataViewJsonParameters."));
+		}
+
+		const FPCGDataViewJsonParameters& JsonParams = Parameters.Get<FPCGDataViewJsonParameters>();
+		return MakeValue(PCG::IO::Json::Helpers::ToJsonString(JsonData.JsonObject, JsonParams.bPrettyJson));
+	}
+
+	return MakeError(LOCTEXT("DataInvalid", "The Data View data or parameters were invalid."));
+}
+
+TValueOrError<UPCGData*, FText> UPCGDataViewJsonConverter::ConstructDataFromTarget(const FInstancedStruct& InTargetStruct, UObject* Outer) const
+{
+	/** @todo_pcg: Construct a PCG Data from a Json object. Revisit serialization options.
+	 * Property Bag Option: Serialize the bag description first as a schema and the payload as a Json string
+	 * Use the bag description to construct the property bag, and then de-serialize Json->UObject on the payload string
+	 */
+	return MakeError(LOCTEXT("DataConstructionNotYetImplemented", "Data construction from Json not yet implemented."));
+}
+
+void UPCGDataViewJsonConverter::BuildJsonHeader(const FPCGDataView& InDataView, TSharedPtr<FJsonObject>& InOutJsonObject) const
+{
+	using namespace PCG::IO;
+	// Append data header. Add it directly to the root object
+	Json::Helpers::AppendVersionToHeader(InOutJsonObject, PCGDataView::Constants::FDataViewVersion::LatestVersion);
+}
+
+#undef LOCTEXT_NAMESPACE
diff --git a/Engine/Plugins/PCG/Source/PCG/Private/Data/DataView/PCGDataViewNativePropertySelectors.cpp b/Engine/Plugins/PCG/Source/PCG/Private/Data/DataView/PCGDataViewNativePropertySelectors.cpp
new file mode 100644
--- /dev/null
+++ b/Engine/Plugins/PCG/Source/PCG/Private/Data/DataView/PCGDataViewNativePropertySelectors.cpp
@@ -0,0 +1,91 @@
+// Copyright Epic Games, Inc. All Rights Reserved.
+
+#include "Data/DataView/PCGDataViewNativePropertySelectors.h"
+
+#include "Data/PCGSplineStruct.h"
+#include "Data/DataView/PCGDataView.h"
+
+namespace PCGDataView
+{
+	namespace Constants
+	{
+		/** @todo_pcg: It would be ideal if each data could handle this individually, but must also account for
+		 * the distinction of only adding non-overlapping properties if there are redundancies. For now, pick
+		 * by hand.
+		 */
+		// A selection of all PCG Point properties that avoid redundancies, i.e. $Transform and $Position.
+		static constexpr EPCGPointProperties AllPointProperties[]
+		{
+			EPCGPointProperties::Transform,
+			EPCGPointProperties::Density,
+			EPCGPointProperties::BoundsMin,
+			EPCGPointProperties::BoundsMax,
+			EPCGPointProperties::Color,
+			EPCGPointProperties::Steepness,
+			EPCGPointProperties::Seed
+		};
+
+		// A selection of all Spline properties.
+		static constexpr EPCGSplineStructProperties AllSplineProperties[]
+		{
+			EPCGSplineStructProperties::Transform,
+			EPCGSplineStructProperties::ArriveTangent,
+			EPCGSplineStructProperties::LeaveTangent,
+			EPCGSplineStructProperties::InterpType,
+			EPCGSplineStructProperties::LocalTransform
+		};
+	} // namespace Constants
+
+	namespace Helpers
+	{
+		void AppendAllAttributeSelectors(const UPCGData* InData, TArray<FPCGAttributePropertySelector>& InOutSelectors)
+		{
+			check(InData);
+			check(InData->ConstMetadata());
+
+			TArray<FPCGAttributeIdentifier> AttributeNames;
+			TArray<EPCGMetadataTypes> AttributeTypes;
+			InData->ConstMetadata()->GetAllAttributes(AttributeNames, AttributeTypes);
+
+			for (int32 Index = 0; Index < AttributeNames.Num(); ++Index)
+			{
+				const bool bIsDefaultDomain = AttributeNames[Index].MetadataDomain.IsDefault();
+				InOutSelectors.Add(FPCGAttributePropertySelector::CreateAttributeSelector(AttributeNames[Index].Name, bIsDefaultDomain ? NAME_None : AttributeNames[Index].MetadataDomain.DebugName));
+			}
+		}
+
+		void AppendAllPointPropertySelectors(TArray<FPCGAttributePropertySelector>& InOutSelectors)
+		{
+			for (const EPCGPointProperties Property : Constants::AllPointProperties)
+			{
+				InOutSelectors.Add(FPCGAttributePropertySelector::CreatePointPropertySelector(Property));
+			}
+		}
+
+		void AppendAllSplinePropertySelectors(TArray<FPCGAttributePropertySelector>& InOutSelectors)
+		{
+			if (const UEnum* SplinePropertiesEnum = StaticEnum<EPCGSplineStructProperties>())
+			{
+				for (const EPCGSplineStructProperties Attribute : Constants::AllSplineProperties)
+				{
+					const FName DisplayName = FName(SplinePropertiesEnum->GetNameStringByValue(static_cast<int32>(Attribute)));
+					InOutSelectors.Add(FPCGAttributePropertySelector::CreatePropertySelector(DisplayName));
+				}
+			}
+		}
+	} // namespace Helpers
+} // namespace PCGDataView
+
+TArray<FPCGAttributePropertySelector> FPCGDataViewBasePointDataPropertySelector::GetSelection(const FPCGDataView& InDataView) const
+{
+	TArray<FPCGAttributePropertySelector> Selectors;
+	PCGDataView::Helpers::AppendAllPointPropertySelectors(Selectors);
+	return Selectors;
+}
+
+TArray<FPCGAttributePropertySelector> FPCGDataViewSplinePropertySelector::GetSelection(const FPCGDataView& InDataView) const
+{
+	TArray<FPCGAttributePropertySelector> Selectors;
+	PCGDataView::Helpers::AppendAllSplinePropertySelectors(Selectors);
+	return Selectors;
+}
diff --git a/Engine/Plugins/PCG/Source/PCG/Private/Data/Registry/PCGDataViewRegistry.cpp b/Engine/Plugins/PCG/Source/PCG/Private/Data/Registry/PCGDataViewRegistry.cpp
new file mode 100644
--- /dev/null
+++ b/Engine/Plugins/PCG/Source/PCG/Private/Data/Registry/PCGDataViewRegistry.cpp
@@ -0,0 +1,101 @@
+// Copyright Epic Games, Inc. All Rights Reserved.
+
+#include "Data/Registry/PCGDataViewRegistry.h"
+
+#include "Data/DataView/PCGDataView.h"
+#include "Data/DataView/PCGDataViewInterface.h"
+
+void FPCGDataViewRegistry::RegisterPropertySelector(const TSubclassOf<UPCGData> InDataClass, TUniquePtr<const IPCGDataViewPropertySelector> PropertySelector)
+{
+	if (ensureMsgf(!PropertySelectorRegistry.Find(InDataClass), TEXT("Cannot register multiple property selectors of the same type.")))
+	{
+		PropertySelectorRegistry.Add(InDataClass, MoveTemp(PropertySelector));
+	}
+}
+
+void FPCGDataViewRegistry::UnregisterPropertySelector(const TSubclassOf<UPCGData> InDataClass)
+{
+	if (ensureMsgf(PropertySelectorRegistry.Find(InDataClass), TEXT("Attempted to unregister a data class that was not registered.")))
+	{
+		PropertySelectorRegistry.Remove(InDataClass);
+	}
+}
+
+void FPCGDataViewRegistry::RegisterConverter(const TSubclassOf<UPCGData> InDataClass, const TSubclassOf<UPCGDataViewConverterBase> InConverterClass)
+{
+	if (TArray<TSubclassOf<UPCGDataViewConverterBase>>* ClassRegistry = ConverterRegistry.Find(InDataClass))
+	{
+		const int32 ConverterIndex = ClassRegistry->Find(InConverterClass);
+		if (ensureMsgf(ConverterIndex == INDEX_NONE, TEXT("Cannot register multiple data view converters of the same type.")))
+		{
+			ConverterRegistry[InDataClass].Emplace(InConverterClass);
+		}
+	}
+	else
+	{
+		ConverterRegistry.Add(InDataClass, {InConverterClass});
+	}
+}
+
+void FPCGDataViewRegistry::UnregisterConverter(const TSubclassOf<UPCGData> InDataClass, const TSubclassOf<UPCGDataViewConverterBase> InConverterClass)
+{
+	if (ensureMsgf(ConverterRegistry.Find(InDataClass), TEXT("Attempted to unregister a data class that was not registered.")))
+	{
+		if (ensureMsgf(INDEX_NONE != ConverterRegistry[InDataClass].Find(InConverterClass), TEXT("Attempted to unregister a converter that was not registered.")))
+		{
+			ConverterRegistry[InDataClass].Remove(InConverterClass);
+			if (ConverterRegistry[InDataClass].IsEmpty())
+			{
+				ConverterRegistry.Remove(InDataClass);
+			}
+		}
+	}
+}
+
+const IPCGDataViewPropertySelector* FPCGDataViewRegistry::GetPropertySelector(const FPCGDataView& InDataView) const
+{
+	if (InDataView.IsValid())
+	{
+		// Climb the class hierarchy.
+		UClass* DataClass = InDataView.ViewedData->GetClass();
+		while (DataClass)
+		{
+			if (const TUniquePtr<const IPCGDataViewPropertySelector>* PropertySelectorPtr = PropertySelectorRegistry.Find(DataClass); PropertySelectorPtr && PropertySelectorPtr->IsValid())
+			{
+				return PropertySelectorPtr->Get();
+			}
+			else
+			{
+				DataClass = DataClass->GetSuperClass();
+			}
+		}
+	}
+
+	return nullptr;
+}
+
+const UPCGDataViewConverterBase* FPCGDataViewRegistry::GetConverter(const FPCGDataView& InDataView, TSubclassOf<UPCGDataViewConverterBase> InConverterClass) const
+{
+	if (InDataView.IsValid())
+	{
+		// Climb the class hierarchy.
+		UClass* DataClass = InDataView.ViewedData->GetClass();
+		while (DataClass && DataClass != UPCGData::StaticClass()) // @todo_pcg: The base converter should work for UPCGData
+		{
+			if (const TArray<TSubclassOf<UPCGDataViewConverterBase>>* ConverterClassArray = ConverterRegistry.Find(DataClass))
+			{
+				for (TSubclassOf<UPCGDataViewConverterBase> ConverterClass : *ConverterClassArray)
+				{
+					if (ConverterClass == InConverterClass)
+					{
+						return GetDefault<UPCGDataViewConverterBase>(ConverterClass);
+					}
+				}
+			}
+
+			DataClass = DataClass->GetSuperClass();
+		}
+	}
+
+	return nullptr;
+}
diff --git a/Engine/Plugins/PCG/Source/PCG/Private/Elements/IO/PCGExportSelectedAttributes.cpp b/Engine/Plugins/PCG/Source/PCG/Private/Elements/IO/PCGExportSelectedAttributes.cpp
--- a/Engine/Plugins/PCG/Source/PCG/Private/Elements/IO/PCGExportSelectedAttributes.cpp
+++ b/Engine/Plugins/PCG/Source/PCG/Private/Elements/IO/PCGExportSelectedAttributes.cpp
@@ -186,11 +186,11 @@ bool FPCGExportSelectedAttributesElement::ExecuteInternal(FPCGContext* Context)
 	FArchive* Archive = nullptr;
 	TSharedPtr<FJsonObject> RootJsonObject = nullptr;
 
-	using namespace PCG::IO::Constants::Attribute;
+	using namespace PCG::IO::Constants;
 	// Initialize the serialization object needed.
 	if (Settings->Format == EPCGExportAttributesFormat::Binary)
 	{
-		FString FinalPath = *Helpers::String::MergeFileName(OutputDirectory, Settings->FileName, TEXT(".bin"));
+		FString FinalPath = *Helpers::String::MergeFileName(OutputDirectory, Settings->FileName, PCG::IO::Binary::Constants::Extension);
 		Archive = IFileManager::Get().CreateFileWriter(*FinalPath);
 		if (!Archive)
 		{
@@ -214,8 +214,9 @@ bool FPCGExportSelectedAttributesElement::ExecuteInternal(FPCGContext* Context)
 
 		// Set up the Json header to include custom versions
 		TSharedPtr<FJsonObject> JsonHeaderObject = MakeShared<FJsonObject>();
-		PCG::IO::Json::Helpers::AppendHeader(JsonHeaderObject, FCustomExportVersion::LatestVersion, Settings->bAddCustomDataVersion ? Settings->CustomVersion : -1);
-		RootJsonObject->SetObjectField(LOCTEXT("JsonHeaderField", "Header").ToString(), std::move(JsonHeaderObject));
+		PCG::IO::Json::Helpers::AppendVersionToHeader(JsonHeaderObject, FCustomExportVersion::LatestVersion);
+		JsonHeaderObject->SetNumberField(TEXT("custom_data_version"), Settings->bAddCustomDataVersion ? Settings->CustomVersion : -1);
+		RootJsonObject->SetObjectField(TEXT("Header"), std::move(JsonHeaderObject));
 	}
 
 	const UPCGData* InputData = Inputs[0].Data;
diff --git a/Engine/Plugins/PCG/Source/PCG/Private/Elements/IO/PCGSaveDataViewElement.cpp b/Engine/Plugins/PCG/Source/PCG/Private/Elements/IO/PCGSaveDataViewElement.cpp
new file mode 100644
--- /dev/null
+++ b/Engine/Plugins/PCG/Source/PCG/Private/Elements/IO/PCGSaveDataViewElement.cpp
@@ -0,0 +1,228 @@
+// Copyright Epic Games, Inc. All Rights Reserved.
+
+#include "Elements/IO/PCGSaveDataViewElement.h"
+
+#include "Data/DataView/PCGDataViewData.h"
+#include "Helpers/PCGSettingsHelpers.h"
+
+#include "HAL/PlatformFileManager.h"
+#include "Misc/FileHelper.h"
+
+#if WITH_EDITOR
+#include "DesktopPlatformModule.h"
+#include "IContentBrowserSingleton.h"
+#endif // WITH_EDITOR
+
+#include UE_INLINE_GENERATED_CPP_BY_NAME(PCGSaveDataViewElement)
+
+#define LOCTEXT_NAMESPACE "PCGSaveDataViewElement"
+
+UPCGSaveDataViewSettings::UPCGSaveDataViewSettings()
+{
+	ConverterParameters.InitializeAs(FPCGDataViewJsonParameters::StaticStruct());
+}
+
+#if WITH_EDITOR
+EPCGChangeType UPCGSaveDataViewSettings::GetChangeTypeForProperty(const FName& InPropertyName) const
+{
+	EPCGChangeType ChangeType = Super::GetChangeTypeForProperty(InPropertyName);
+
+	if (InPropertyName == GET_MEMBER_NAME_CHECKED(UPCGSaveDataViewSettings, ConverterClass))
+	{
+		ChangeType |= EPCGChangeType::Structural;
+	}
+
+	return ChangeType;
+}
+
+void UPCGSaveDataViewSettings::PostEditChangeProperty(FPropertyChangedEvent& PropertyChangedEvent)
+{
+	Super::PostEditChangeProperty(PropertyChangedEvent);
+
+	if (PropertyChangedEvent.GetPropertyName() == GET_MEMBER_NAME_CHECKED(UPCGSaveDataViewSettings, ConverterClass))
+	{
+		if (const UPCGDataViewConverterBase* ConverterCDO = ConverterClass ? ConverterClass->GetDefaultObject<UPCGDataViewConverterBase>() : nullptr)
+		{
+			ConverterParameters.InitializeAs(ConverterCDO->GetParameterStruct());
+		}
+		else
+		{
+			ConverterParameters.Reset();
+		}
+	}
+}
+#endif // WITH_EDITOR
+
+TArray<FPCGPinProperties> UPCGSaveDataViewSettings::InputPinProperties() const
+{
+	TArray<FPCGPinProperties> PinProperties;
+
+	// @todo_pcg: Allow iterating on inputs and optionally save paths as inputs N:N.
+	PinProperties.Emplace_GetRef(PCGPinConstants::DefaultInputLabel, FPCGDataTypeInfoDataView::AsId(), /*bInAllowMultipleConnections=*/false, /*bAllowMultipleData=*/false).SetRequiredPin();
+
+	return PinProperties;
+}
+
+TArray<FPCGPinProperties> UPCGSaveDataViewSettings::OutputPinProperties() const
+{
+	TArray<FPCGPinProperties> PinProperties;
+	FPCGPinProperties& DepPin = PinProperties.Emplace_GetRef(PCGPinConstants::DefaultExecutionDependencyLabel, EPCGDataType::Any);
+#if WITH_EDITOR
+	DepPin.Tooltip = PCGPinConstants::Tooltips::ExecutionDependencyTooltip;
+#endif // WITH_EDITOR
+	DepPin.Usage = EPCGPinUsage::DependencyOnly;
+
+	return PinProperties;
+}
+
+FPCGElementPtr UPCGSaveDataViewSettings::CreateElement() const
+{
+	return MakeShared<FPCGSaveDataViewElement>();
+}
+
+#if WITH_EDITOR
+TArray<FPCGSettingsOverridableParam> UPCGSaveDataViewSettings::GatherOverridableParams() const
+{
+	TArray<FPCGSettingsOverridableParam> OverridableParams = Super::GatherOverridableParams();
+
+	if (ConverterParameters.IsValid())
+	{
+		const PCGSettingsHelpers::FPCGGetAllOverridableParamsConfig Config;
+		return PCGSettingsHelpers::GetAllOverridableParams(ConverterParameters.GetScriptStruct(), Config);
+	}
+
+	return {};
+}
+#endif // WITH_EDITOR
+
+void UPCGSaveDataViewSettings::FixingOverridableParamPropertyClass(FPCGSettingsOverridableParam& Param) const
+{
+	bool bFound = false;
+
+	if (!Param.PropertiesNames.IsEmpty())
+	{
+		const UScriptStruct* ScriptStruct = ConverterParameters.GetScriptStruct();
+		if (ScriptStruct && ScriptStruct->FindPropertyByName(Param.PropertiesNames[0]))
+		{
+			Param.PropertyClass = ScriptStruct;
+			bFound = true;
+		}
+	}
+
+	if (!bFound)
+	{
+		Super::FixingOverridableParamPropertyClass(Param);
+	}
+}
+
+FStructView FPCGSaveDataViewExecutionContext::GetExternalStructContainerForOverridableParam(const FPCGSettingsOverridableParam& InParam)
+{
+	UPCGSaveDataViewSettings* Settings = GetMutableInputSettings<UPCGSaveDataViewSettings>();
+	if (Settings && Settings->ConverterParameters.IsValid() && !InParam.PropertiesNames.IsEmpty() && Settings->ConverterParameters.GetScriptStruct()->FindPropertyByName(InParam.PropertiesNames[0]))
+	{
+		return Settings->ConverterParameters;
+	}
+	else
+	{
+		return nullptr;
+	}
+}
+
+FPCGContext* FPCGSaveDataViewElement::CreateContext()
+{
+	return new FPCGSaveDataViewExecutionContext();
+}
+
+bool FPCGSaveDataViewElement::ExecuteInternal(FPCGContext* InContext) const
+{
+	using namespace PCG::IO;
+	TRACE_CPUPROFILER_EVENT_SCOPE(FPCGSaveDataViewElement::Execute);
+
+	// Since this generates data on disk, only allow execution on editor approved platforms.
+#if PLATFORM_WINDOWS || PLATFORM_LINUX || PLATFORM_MAC
+	const UPCGSaveDataViewSettings* Settings = InContext->GetInputSettings<UPCGSaveDataViewSettings>();
+	check(Settings);
+
+	const TArray<FPCGTaggedData> Inputs = InContext->InputData.GetInputsByPin(PCGPinConstants::DefaultInputLabel);
+	if (Inputs.IsEmpty())
+	{
+		return true;
+	}
+
+	// @todo_pcg: Allow iterating on inputs and optionally save paths as inputs N:N. For now, just one input.
+	if (Inputs.Num() > 1)
+	{
+		PCGLog::InputOutput::LogFirstInputOnlyWarning(PCGPinConstants::DefaultInputLabel, InContext);
+	}
+
+	const FPCGTaggedData& Input = Inputs[0];
+
+	const UPCGDataViewData* DataViewData = Cast<UPCGDataViewData>(Input.Data);
+	if (!DataViewData)
+	{
+		PCGLog::InputOutput::LogInvalidInputDataError(InContext);
+		return true;
+	}
+
+	bool bOverwriteFile = Settings->bOverwriteExistingFile;
+
+	FString FilePath = Settings->FilePath.FilePath;
+#if WITH_EDITOR
+	if (Settings->FilePath.FilePath.IsEmpty() && Settings->bOpenDialogueOnEmptyPath)
+	{
+		TArray<FString> OutFiles;
+		IDesktopPlatform* DesktopPlatform = FDesktopPlatformModule::Get();
+		if (DesktopPlatform->SaveFileDialog(
+			nullptr,
+			LOCTEXT("OpenFileDialogue", "Choose output file path").ToString(),
+			FPaths::GameUserDeveloperDir(),
+			Settings->FilePath.FilePath.IsEmpty() ? FPaths::SetExtension(InContext->GetTaskName(), Json::Constants::Extension) : Settings->FilePath.FilePath,
+			Json::Constants::Extension,
+			EFileDialogFlags::None,
+			OutFiles))
+		{
+			check(!OutFiles.IsEmpty())
+			FilePath = OutFiles[0];
+			// Explicit decision to choose this file. Overwrite is now okay.
+			bOverwriteFile = true;
+		}
+	}
+#endif // WITH_EDITOR
+
+	if (FilePath.IsEmpty())
+	{
+		PCGLog::LogErrorOnGraph(FText::Format(LOCTEXT("InvalidFilePathError", "Invalid or empty export file path: {0}"), FilePath.IsEmpty() ? LOCTEXT("Empty", "Empty") : FText::FromString(FilePath)), InContext);
+		return true;
+	}
+
+	if (!bOverwriteFile && FPlatformFileManager::Get().GetPlatformFile().FileExists(*FilePath))
+	{
+		PCGLog::LogErrorOnGraph(LOCTEXT("NoOverwriteFile", "File would be overwritten. Enable 'Overwrite Existing File' to bypass this check."));
+		return true;
+	}
+
+	const FPCGDataView DataView = DataViewData->GetDataView();
+	if (!DataView.IsValid())
+	{
+		PCGLog::LogWarningOnGraph(LOCTEXT("InvalidDataView", "Data View selection is empty or internal data is no longer invalid."));
+		return true;
+	}
+
+	if (Settings->ConverterClass)
+	{
+		const UPCGDataViewConverterBase* ConverterCDO = Settings->ConverterClass->GetDefaultObject<UPCGDataViewConverterBase>();
+		TValueOrError<void, FText> Result = ConverterCDO->SerializeToFile(DataView, Settings->ConverterParameters, FilePath);
+		if (Result.HasError())
+		{
+			PCGLog::LogErrorOnGraph(Result.GetError(), InContext);
+			return true;
+		}
+	}
+
+#else  // PLATFORM_WINDOWS || PLATFORM_LINUX || PLATFORM_MAC - This is not Windows, Linux, or Mac
+	UE_LOG(LogPCG, Verbose, TEXT("The 'Save Data View' node has been disabled on this platform."));
+#endif // PLATFORM_WINDOWS || PLATFORM_LINUX || PLATFORM_MAC
+	return true;
+}
+
+#undef LOCTEXT_NAMESPACE
diff --git a/Engine/Plugins/PCG/Source/PCG/Private/Elements/PCGConvertToDataView.cpp b/Engine/Plugins/PCG/Source/PCG/Private/Elements/PCGConvertToDataView.cpp
new file mode 100644
--- /dev/null
+++ b/Engine/Plugins/PCG/Source/PCG/Private/Elements/PCGConvertToDataView.cpp
@@ -0,0 +1,102 @@
+// Copyright Epic Games, Inc. All Rights Reserved.
+
+#include "Elements/PCGConvertToDataView.h"
+
+#include "PCGContext.h"
+#include "PCGParamData.h"
+#include "Data/PCGSplineData.h"
+#include "Data/DataView/PCGDataViewData.h"
+#include "Metadata/Accessors/PCGAttributeAccessorHelpers.h"
+
+#include UE_INLINE_GENERATED_CPP_BY_NAME(PCGConvertToDataView)
+
+#define LOCTEXT_NAMESPACE "PCGConvertToDataViewElement"
+
+#if WITH_EDITOR
+EPCGChangeType UPCGConvertToDataViewSettings::GetChangeTypeForProperty(const FName& InPropertyName) const
+{
+	EPCGChangeType ChangeType = Super::GetChangeTypeForProperty(InPropertyName);
+	if (InPropertyName == GET_MEMBER_NAME_CHECKED(UPCGConvertToDataViewSettings, bSelectionAsInput))
+	{
+		ChangeType |= EPCGChangeType::Structural;
+	}
+
+	return ChangeType;
+}
+
+bool UPCGConvertToDataViewSettings::ShouldDrawNodeCompact() const
+{
+	return !bSelectionAsInput;
+}
+#endif // WITH_EDITOR
+
+TArray<FPCGPinProperties> UPCGConvertToDataViewSettings::InputPinProperties() const
+{
+	TArray<FPCGPinProperties> PinProperties;
+	PinProperties.Emplace_GetRef(PCGPinConstants::DefaultInputLabel, FPCGDataTypeInfoSpatial::AsId()).SetRequiredPin();
+
+	return PinProperties;
+}
+
+TArray<FPCGPinProperties> UPCGConvertToDataViewSettings::OutputPinProperties() const
+{
+	TArray<FPCGPinProperties> PinProperties;
+	PinProperties.Emplace(PCGPinConstants::DefaultOutputLabel, FPCGDataTypeInfoDataView::AsId());
+
+	return PinProperties;
+}
+
+FPCGElementPtr UPCGConvertToDataViewSettings::CreateElement() const
+{
+	return MakeShared<FPCGConvertToDataViewElement>();
+}
+
+bool FPCGConvertToDataViewElement::ExecuteInternal(FPCGContext* InContext) const
+{
+	TRACE_CPUPROFILER_EVENT_SCOPE(FPCGConvertToDataViewElement::Execute);
+	check(InContext);
+
+	const UPCGConvertToDataViewSettings* Settings = InContext->GetInputSettings<UPCGConvertToDataViewSettings>();
+
+	// Collect attribute selection from input.
+	TArray<FString> AttributeStrings;
+	if (Settings->bSelectionAsInput)
+	{
+		Settings->SelectionSources.ParseIntoArrayWS(AttributeStrings, TEXT(","));
+	}
+
+	TArray<FPCGTaggedData> Inputs = InContext->InputData.GetInputsByPin(PCGPinConstants::DefaultInputLabel);
+	TArray<FPCGTaggedData>& Outputs = InContext->OutputData.TaggedData;
+
+	for (const FPCGTaggedData& Input : Inputs)
+	{
+		const UPCGMetadata* SourceMetadata = Input.Data->Metadata;
+		check(SourceMetadata);
+
+		FPCGDataView DataView{};
+		DataView.ViewedData = Input.Data;
+		if (Settings->bSelectionAsInput)
+		{
+			DataView.Selection.bAllAttributes = false;
+			for (const FString& AttributeString : AttributeStrings)
+			{
+				DataView.Selection.Attributes.AddUnique(FPCGAttributePropertySelector::CreateSelectorFromString(AttributeString));
+			}
+		}
+		else
+		{
+			DataView.Selection = Settings->Selection;
+		}
+
+		DataView.AttributeLayout = Settings->AttributeLayout;
+
+		UPCGDataViewData* DataViewData = InContext->NewObject_AnyThread<UPCGDataViewData>(InContext);
+		DataViewData->Initialize(MoveTemp(DataView));
+
+		Outputs.Emplace_GetRef(Input).Data = DataViewData;
+	}
+
+	return true;
+}
+
+#undef LOCTEXT_NAMESPACE
diff --git a/Engine/Plugins/PCG/Source/PCG/Private/Helpers/IO/PCGIOHelpers.cpp b/Engine/Plugins/PCG/Source/PCG/Private/Helpers/IO/PCGIOHelpers.cpp
--- a/Engine/Plugins/PCG/Source/PCG/Private/Helpers/IO/PCGIOHelpers.cpp
+++ b/Engine/Plugins/PCG/Source/PCG/Private/Helpers/IO/PCGIOHelpers.cpp
@@ -11,11 +11,8 @@ namespace PCG::IO
 {
 	namespace Constants
 	{
-		namespace Attribute
-		{
-			const FGuid FCustomExportVersion::GUID = FGuid(0x04E74488, 0x4BAC8717, 0xBBB18694, 0x39F8F3CE);
-			FCustomVersionRegistration GRegisterPCGExportSelectedAttributesCustomVersion(FCustomExportVersion::GUID, FCustomExportVersion::LatestVersion, TEXT("PCGExportSelectedAttributes"));
-		}
+		const FGuid FCustomExportVersion::GUID = FGuid(0x04E74488, 0x4BAC8717, 0xBBB18694, 0x39F8F3CE);
+		FCustomVersionRegistration GRegisterPCGExportSelectedAttributesCustomVersion(FCustomExportVersion::GUID, FCustomExportVersion::LatestVersion, TEXT("PCGExportSelectedAttributes"));
 	}
 
 	namespace Helpers::String
diff --git a/Engine/Plugins/PCG/Source/PCG/Private/Helpers/IO/PCGJsonHelpers.cpp b/Engine/Plugins/PCG/Source/PCG/Private/Helpers/IO/PCGJsonHelpers.cpp
--- a/Engine/Plugins/PCG/Source/PCG/Private/Helpers/IO/PCGJsonHelpers.cpp
+++ b/Engine/Plugins/PCG/Source/PCG/Private/Helpers/IO/PCGJsonHelpers.cpp
@@ -2,8 +2,10 @@
 
 #include "Helpers/IO/PCGJsonHelpers.h"
 
+#include "Helpers/PCGHelpers.h"
 #include "Helpers/IO/PCGIOHelpers.h"
 #include "Metadata/Accessors/PCGAttributeAccessorHelpers.h"
+#include "Policies/CondensedJsonPrintPolicy.h"
 
 #include "Serialization/JsonSerializer.h"
 #include "Serialization/JsonWriter.h"
@@ -61,7 +63,7 @@ namespace PCG::IO::Json
 				if constexpr (PCG::Private::MetadataTraits<T>::IsFloatingPoint)
 				{
 					TArray<T> Values;
-					PCGAttributeAccessorHelpers::ExtractAllValues(Entry.Accessor.Get(), Entry.Keys.Get(), Values, EPCGAttributeAccessorFlags::StrictType);
+					PCGAttributeAccessorHelpers::ExtractAllValues(Entry.Accessor.Get(), Entry.Keys.Get(), Values, EPCGAttributeAccessorFlags::AllowBroadcastAndConstructible);
 
 					JsonValues.Reserve(Values.Num());
 					Algo::Transform(Values, JsonValues, [&Entry](const T& Value)
@@ -72,7 +74,7 @@ namespace PCG::IO::Json
 				else // For everything else, automatically broadcast it to string
 				{
 					TArray<FString> Values;
-					PCGAttributeAccessorHelpers::ExtractAllValues(Entry.Accessor.Get(), Entry.Keys.Get(), Values, EPCGAttributeAccessorFlags::StrictType);
+					PCGAttributeAccessorHelpers::ExtractAllValues(Entry.Accessor.Get(), Entry.Keys.Get(), Values, EPCGAttributeAccessorFlags::AllowBroadcastAndConstructible);
 
 					JsonValues.Reserve(Values.Num());
 					Algo::Transform(Values, JsonValues, [](const FString& OutValue)
@@ -93,8 +95,8 @@ namespace PCG::IO::Json
 			for (int32 CacheIndex = 0; CacheIndex < AccessorCache.Num(); ++CacheIndex)
 			{
 				const Accessor::FCacheEntry& Entry = AccessorCache[CacheIndex];
-				FString& SourceName = OutKeyValuesMap[CacheIndex].Get<0>();
-				SourceName = Entry.Selector.ToString();
+				FString& FieldName = OutKeyValuesMap[CacheIndex].Get<0>();
+				FieldName = Entry.Selector.ToString();
 
 				TArray<TSharedPtr<FJsonValue>>& SourceJsonValues = OutKeyValuesMap[CacheIndex].Get<1>();
 				SourceJsonValues = AttributeToJsonValues(Entry);
@@ -141,13 +143,27 @@ namespace PCG::IO::Json
 					ElementJsonObject->SetField(Key, MoveTemp(JsonValues[ElementIndex]));
 				}
 
-				InOutJsonObject->SetObjectField(FString::FromInt(ElementIndex), MoveTemp(ElementJsonObject));
+				InOutJsonObject->SetObjectField(FString::Printf(TEXT("element_%d"), ElementIndex), MoveTemp(ElementJsonObject));
 			}
 
 			// The map is no longer valid, so clean it.
 			InOutKeyValueMap.Empty();
 		}
-	}
+
+		template <typename CharType = TCHAR, typename PrintPolicy = TCondensedJsonPrintPolicy<CharType>>
+		FString ToJsonString(const TSharedPtr<FJsonObject>& JsonObject)
+		{
+			FString OutputString;
+			const TSharedRef<TJsonWriter<CharType, PrintPolicy>> JsonWriter = TJsonWriterFactory<CharType, PrintPolicy>::Create(&OutputString);
+			if (!FJsonSerializer::Serialize(JsonObject.ToSharedRef(), JsonWriter))
+			{
+				PCGLog::LogErrorOnGraph(LOCTEXT("JsonSerializationError", "Serialization of Json data failed."));
+				return {};
+			}
+
+			return OutputString;
+		}
+	} // namespace Private
 
 	namespace Helpers
 	{
@@ -186,31 +202,50 @@ namespace PCG::IO::Json
 			return MakeShared<FJsonValueObject>(JsonTransformObject);
 		}
 
-		void AppendHeader(const TSharedPtr<FJsonObject>& InOutJsonObject, const int32 DataVersion, const int32 CustomDataVersion)
+		void AppendHeader(const TSharedPtr<FJsonObject>& InOutJsonObject, int32 DataVersion, int32 CustomDataVersion)
+		{
+			AppendVersionToHeader(InOutJsonObject, DataVersion);
+		}
+
+		void AppendVersionToHeader(const TSharedPtr<FJsonObject>& InOutJsonObject, const int32 DataVersion)
 		{
 			if (ensure(InOutJsonObject.IsValid()))
 			{
-				InOutJsonObject->SetNumberField(TEXT("data_version"), DataVersion);
+				InOutJsonObject->SetNumberField(Constants::DefaultDataVersionKey, DataVersion);
+			}
+		}
 
-				if (CustomDataVersion >= 0)
-				{
-					InOutJsonObject->SetNumberField(TEXT("custom_data_version"), CustomDataVersion);
-				}
+		void AppendSelectionByAttribute(TSharedPtr<FJsonObject>& InOutJsonObject, const UPCGData* InData, const TArray<FPCGAttributePropertySelector>& Selectors, const FPCGContext* InOptionalContext)
+		{
+			if (!InOutJsonObject || !InData || Selectors.IsEmpty())
+			{
+				return;
 			}
+
+			const Accessor::FCache AccessorCache = Private::CreateAccessorCache(InData, Selectors, InOptionalContext);
+			JsonKeyValuesMap KeyValuesMap;
+			Private::BuildKeyValuesMap(KeyValuesMap, AccessorCache);
+			Private::AppendSelectionByAttribute(InOutJsonObject, KeyValuesMap);
 		}
 
-		FString ToJsonString(const TSharedPtr<FJsonObject>& JsonObject)
+		void AppendSelectionByElement(TSharedPtr<FJsonObject>& InOutJsonObject, const UPCGData* InData, const TArray<FPCGAttributePropertySelector>& Selectors, const FPCGContext* InOptionalContext)
 		{
-			FString OutputString;
-			// @todo_pcg: Condensed by default. Add 'pretty json' as an option.
-			const TSharedRef<TJsonWriter<>> JsonWriter = TJsonWriterFactory<>::Create(&OutputString);
-			if (!FJsonSerializer::Serialize(JsonObject.ToSharedRef(), JsonWriter))
+			if (!InOutJsonObject || !InData || Selectors.IsEmpty())
 			{
-				PCGLog::LogErrorOnGraph(LOCTEXT("JsonSerializationError", "Serialization of Json data failed."));
-				return {};
+				return;
 			}
 
-			return OutputString;
+			const Accessor::FCache AccessorCache = Private::CreateAccessorCache(InData, Selectors, InOptionalContext);
+			JsonKeyValuesMap KeyValuesMap;
+			Private::BuildKeyValuesMap(KeyValuesMap, AccessorCache);
+			Private::AppendSelectionByElement(InOutJsonObject, KeyValuesMap, PCGHelpers::GetNumberOfElements(InData));
+		}
+
+		FString ToJsonString(const TSharedPtr<FJsonObject>& JsonObject, const bool bPrettyJson)
+		{
+			return bPrettyJson
+				? Private::ToJsonString<TCHAR, TPrettyJsonPrintPolicy<TCHAR>>(JsonObject)
+				: Private::ToJsonString<TCHAR, TCondensedJsonPrintPolicy<TCHAR>>(JsonObject);
 		}
 	} // namespace Helpers
 }  // namespace PCG::IO::Json
diff --git a/Engine/Plugins/PCG/Source/PCG/Private/PCGModule.cpp b/Engine/Plugins/PCG/Source/PCG/Private/PCGModule.cpp
--- a/Engine/Plugins/PCG/Source/PCG/Private/PCGModule.cpp
+++ b/Engine/Plugins/PCG/Source/PCG/Private/PCGModule.cpp
@@ -7,12 +7,16 @@
 #include "PCGDefaultExecutionSource.h"
 #include "PCGElement.h"
 #include "PCGGraph.h"
+#include "PCGParamData.h"
 #include "Compute/PCGComputeKernel.h"
 #include "Data/PCGBasePointData.h"
 #include "Data/PCGPointData.h"
 #include "Data/PCGPolyLineData.h"
 #include "Data/PCGSplineData.h"
 #include "Data/PCGPolygon2DData.h"
+#include "Data/DataView/PCGDataViewInterface.h"
+#include "Data/DataView/PCGDataViewNativeJsonConverters.h"
+#include "Data/DataView/PCGDataViewNativePropertySelectors.h"
 
 #include "Engine/Engine.h"
 #include "Interfaces/IPluginManager.h"
@@ -178,6 +182,13 @@ void FPCGModule::StartupModule()
 	AttributeAccessorFactory.RegisterMethods<UPCGPolygon2DData>(UPCGPolygon2DData::GetPolygon2DAccessorMethods());
 	AttributeAccessorFactory.RegisterMethods<UPCGSplineData>(UPCGSplineData::GetSplineAccessorMethods());
 
+	DataViewRegistry.RegisterPropertySelector(UPCGData::StaticClass(), MakeUnique<FPCGDataViewPropertySelector>());
+	DataViewRegistry.RegisterPropertySelector(UPCGBasePointData::StaticClass(), MakeUnique<FPCGDataViewBasePointDataPropertySelector>());
+	DataViewRegistry.RegisterPropertySelector(UPCGSplineData::StaticClass(), MakeUnique<FPCGDataViewSplinePropertySelector>());
+
+	DataViewRegistry.RegisterConverter(UPCGParamData::StaticClass(), UPCGDataViewJsonConverter::StaticClass());
+	DataViewRegistry.RegisterConverter(UPCGSpatialData::StaticClass(), UPCGDataViewJsonConverter::StaticClass());
+
 #if WITH_EDITOR
 	ObjectHashFactory.RegisterObjectHashContextFactory(UPCGGraphInterface::StaticClass(), FPCGOnCreateObjectHashContext::CreateStatic(&FPCGGraphHashContext::MakeInstance));
 	ObjectHashFactory.RegisterObjectHashContextFactory(UPCGSettingsInterface::StaticClass(), FPCGOnCreateObjectHashContext::CreateStatic(&FPCGSettingsHashContext::MakeInstance));
@@ -203,6 +214,13 @@ void FPCGModule::ShutdownModule()
 
 void FPCGModule::PreExit()
 {
+	DataViewRegistry.UnregisterConverter(UPCGSpatialData::StaticClass(), UPCGDataViewJsonConverter::StaticClass());
+	DataViewRegistry.UnregisterConverter(UPCGParamData::StaticClass(), UPCGDataViewJsonConverter::StaticClass());
+
+	DataViewRegistry.UnregisterPropertySelector(UPCGSplineData::StaticClass());
+	DataViewRegistry.UnregisterPropertySelector(UPCGBasePointData::StaticClass());
+	DataViewRegistry.UnregisterPropertySelector(UPCGData::StaticClass());
+
 	// Unregistering accessor methods
 	AttributeAccessorFactory.UnregisterMethods<UPCGSplineData>();
 	AttributeAccessorFactory.UnregisterMethods<UPCGPolygon2DData>();
diff --git a/Engine/Plugins/PCG/Source/PCG/Public/Data/DataView/PCGDataView.h b/Engine/Plugins/PCG/Source/PCG/Public/Data/DataView/PCGDataView.h
new file mode 100644
--- /dev/null
+++ b/Engine/Plugins/PCG/Source/PCG/Public/Data/DataView/PCGDataView.h
@@ -0,0 +1,103 @@
+// Copyright Epic Games, Inc. All Rights Reserved.
+
+#pragma once
+
+#include "PCGData.h"
+#include "Helpers/IO/PCGIOHelpers.h"
+#include "Metadata/PCGAttributePropertySelector.h"
+
+#include "Serialization/CustomVersion.h"
+#include "Templates/SubclassOf.h"
+
+#include "PCGDataView.generated.h"
+
+class UPCGDataViewBase;
+struct FPCGAttributePropertySelector;
+struct FPCGContext;
+
+namespace PCGDataView::Constants
+{
+	// Custom PCG version for DataView data export
+	struct FDataViewVersion
+	{
+		static constexpr TCHAR FriendlyName[] = TEXT("PCGDataView");
+
+		enum Type
+		{
+			InitialVersion = 0,
+
+			// New versions can be added above this line
+			VersionPlusOne,
+			LatestVersion = VersionPlusOne - 1
+		};
+
+		static FName GetFriendlyName() { return FName(FriendlyName); }
+		const static FGuid GUID;
+	};
+} // namespace PCGDataView::Constants
+
+UENUM()
+enum class EPCGDataViewConversion : uint32
+{
+	Json UMETA(Tooltip = "Convert the data view into Json format. See 'EPCGDataViewAttributeLayout' for how properties and attributes can be serialized."),
+
+	// @todo_pcg: Binary (Archive)
+};
+
+UENUM()
+enum class EPCGDataViewAttributeLayout : uint8
+{
+	ByElement UMETA(Tooltip = "Use the elements as the main data object. Each element contains all its attributes."),
+	ByAttribute UMETA(Tooltip = "Use the attributes as the main data object. Each attribute grouping will contain the element values in sequential order.")
+};
+
+USTRUCT(MinimalAPI, BlueprintType, Blueprintable)
+struct FPCGDataViewSelection
+{
+	GENERATED_BODY()
+
+	// @todo_pcg: Can EPCGMetadataFilterMode be extended to include an "AllAttributes" mode?
+	UPROPERTY(EditAnywhere, Category = Selection)
+	bool bAllAttributes = true;
+
+	// @todo_pcg: Extend to use EPCGMetadataFilterMode for more complex selections
+	UPROPERTY(EditAnywhere, Category = Selection, meta = (EditCondition = "!bAllAttributes", EditConditionHides))
+	TArray<FPCGAttributePropertySelector> Attributes;
+};
+
+/** This wrapper class maintains a reference to unowned "viewed" data and also contains a "selection" of
+ * metadata. This can be used with a variety of operations--usually exporting data--where a PCG Data type
+ * needs to be abstracted and/or serialized into a more generic type or refined to an explicit selection
+ * of properties or attributes on the data, to make it more accessible for processing.
+ *
+ * Ex. Convert the $Position property of each point in a PCG Point data into Json key:value objects. 
+ */
+USTRUCT(MinimalAPI, BlueprintType, Blueprintable)
+struct FPCGDataView
+{
+	GENERATED_BODY()
+
+	FPCGDataView() = default;
+	virtual ~FPCGDataView() = default;
+
+	// The data exists and the selection is not empty.
+	PCG_API virtual bool IsValid() const;
+
+	// Aggregate all domain IDs from the current selection
+	PCG_API TArray<FPCGMetadataDomainID> GetAllMetadataDomainIDs() const;
+
+	// Get the final selection. @todo_pcg: Could be cached.
+	PCG_API TArray<FPCGAttributePropertySelector> GetResolvedSelection() const;
+
+	// The PCG Data this will provide a 'view' onto.
+	UPROPERTY(VisibleAnywhere, Category = Data)
+	TObjectPtr<const UPCGData> ViewedData = nullptr;
+
+	// An optionally filtered collection of attributes and properties to perform an operation on, such as data export.
+	UPROPERTY(VisibleAnywhere, Category = Data, meta = (ShowOnlyInnerProperties))
+	FPCGDataViewSelection Selection;
+
+	// How the attribute data should be exported, assuming the data supports the elements domain.
+	UPROPERTY(VisibleAnywhere, Category = Data)
+	EPCGDataViewAttributeLayout AttributeLayout = EPCGDataViewAttributeLayout::ByElement;
+};
diff --git a/Engine/Plugins/PCG/Source/PCG/Public/Data/DataView/PCGDataViewData.h b/Engine/Plugins/PCG/Source/PCG/Public/Data/DataView/PCGDataViewData.h
new file mode 100644
--- /dev/null
+++ b/Engine/Plugins/PCG/Source/PCG/Public/Data/DataView/PCGDataViewData.h
@@ -0,0 +1,74 @@
+// Copyright Epic Games, Inc. All Rights Reserved.
+
+#pragma once
+
+#include "PCGData.h"
+#include "PCGDataViewInterface.h"
+#include "PCGDataView.h"
+
+#include "PCGDataViewData.generated.h"
+
+// The Data View type info and description
+USTRUCT()
+struct FPCGDataTypeInfoDataView : public FPCGDataTypeInfo
+{
+	GENERATED_BODY()
+
+	PCG_DECLARE_TYPE_INFO()
+
+	virtual bool SupportsConversionFrom(const FPCGDataTypeIdentifier& InputType, const FPCGDataTypeIdentifier& ThisType, TSubclassOf<UPCGSettings>* OptionalOutConversionSettings = nullptr, FText* OptionalOutCompatibilityMessage = nullptr) const override;
+};
+
+/** Acts as a type-erasing abstraction wrapper around other PCG Data types, as well as allowing for a limited
+ * selection of attributes or properties to be defined for performing an operation on later. I.e. Specified
+ * attributes or properties can be selectively serialized into another format, supported by implemented interfaces.
+ *
+ * Ex. A PCG Data View Data could point to a PCG Point Data, where only the $Position property is selected.
+ * Because the point data implements the IPCGDataViewJson interface, the selection of Positions for each point
+ * in the data may be serialized to/from a Json object/string form.
+ *
+ * Implementation Note: UPCGDataViewData should implement all available Data View interfaces to enforce and
+ * maintain parity with consumer types. Calling an interface method on the Data View Data should essentially
+ * forward the call to the underlying data, if the interface is supported.
+ */
+UCLASS(MinimalAPI, BlueprintType, ClassGroup = (Procedural))
+class UPCGDataViewData : public UPCGData
+{
+	GENERATED_BODY()
+
+	// Used to cache the selector/accessor for ease of data retrieval into the data
+	struct FAccessorCacheEntry
+	{
+		FPCGAttributePropertySelector Selector;
+		TUniquePtr<const IPCGAttributeAccessor> Accessor = nullptr;
+		TUniquePtr<const IPCGAttributeAccessorKeys> Keys = nullptr;
+	};
+
+public:
+	PCG_ASSIGN_TYPE_INFO(FPCGDataTypeInfoDataView);
+
+	// Set up the data view upon creation.
+	void Initialize(FPCGDataView&& InDataView);
+
+	//~ Begin UPCGData interface
+	virtual const FPCGDataView& GetDataView() const;
+
+	// Forwarding helper to get a domain ID from a selector, to the internal data.
+	virtual FPCGMetadataDomainID GetMetadataDomainIDFromSelector(const FPCGAttributePropertySelector& InSelector) const override;
+	// Forwarding helper to get the default domain ID, to the internal data.
+	virtual FPCGMetadataDomainID GetDefaultMetadataDomainID() const override;
+	// Forwarding helper to get all supported domain IDs, to the internal data.
+	virtual TArray<FPCGMetadataDomainID> GetAllSupportedMetadataDomainIDs() const override;
+	// Forwarding helper to the underlying metadata.
+	virtual const UPCGMetadata* ConstMetadata() const override;
+	//~ End UPCGData interface
+
+	// The selection is stored lazily. Get Selected Attributes will create the array of selectors. @todo_pcg: Look into caching results.
+	TArray<FPCGAttributePropertySelector> GetSelectedAttributes() const;
+
+	// @todo_pcg: Serialize to/from Archive (Binary)
+
+private:
+	UPROPERTY()
+	FPCGDataView DataView;
+};
diff --git a/Engine/Plugins/PCG/Source/PCG/Public/Data/DataView/PCGDataViewInterface.h b/Engine/Plugins/PCG/Source/PCG/Public/Data/DataView/PCGDataViewInterface.h
new file mode 100644
--- /dev/null
+++ b/Engine/Plugins/PCG/Source/PCG/Public/Data/DataView/PCGDataViewInterface.h
@@ -0,0 +1,63 @@
+// Copyright Epic Games, Inc. All Rights Reserved.
+
+#pragma once
+
+#include "Helpers/IO/PCGIOHelpers.h"
+#include "Helpers/IO/PCGJsonHelpers.h"
+#include "Templates/ValueOrError.h"
+
+#include "PCGDataViewInterface.generated.h"
+
+class FJsonObject;
+class UPCGData;
+struct FPCGDataView;
+
+// @todo_pcg: Enable when loading from serialized data view.
+// UENUM(BlueprintType)
+// enum class EPCGDataViewVersionMismatchPolicy : uint8
+// {
+// 	Ignore = 0u,
+// 	Log,
+// 	Warn,
+// 	Error
+// };
+
+/** A class can inherit from this interface to define a process for selecting properties of a PCG Data */
+class IPCGDataViewPropertySelector
+{
+public:
+	virtual ~IPCGDataViewPropertySelector() = default;
+
+	/** Allows data types to set custom property selectors. */
+	virtual TArray<FPCGAttributePropertySelector> GetSelection(const FPCGDataView& InDataView) const = 0;
+};
+
+/** The native base property selector for PCG. Covers spatial and any base types */
+class FPCGDataViewPropertySelector : public IPCGDataViewPropertySelector
+{
+	virtual TArray<FPCGAttributePropertySelector> GetSelection(const FPCGDataView& InDataView) const override { return {}; }
+};
+
+UCLASS(Abstract)
+class UPCGDataViewConverterBase : public UObject
+{
+	GENERATED_BODY()
+
+public:
+	virtual bool CanConvertToString() const { return false; }
+
+	/** Override to serialize the Data View's underlying data and selection into an FInstancedStruct matching GetTargetStruct(). */
+	virtual TValueOrError<FInstancedStruct, FText> SerializeToTargetStruct(const FPCGDataView& InDataView, const FInstancedStruct& Parameters) const PURE_VIRTUAL(UPCGDataViewConverterBase::SerializeToTargetStruct, { return MakeError(FText::GetEmpty()); })
+
+	virtual TValueOrError<FString, FText> SerializeToString(const FPCGDataView& InDataView, const FInstancedStruct& Parameters) const PURE_VIRTUAL(UPCGDataViewConverterBase::SerializeToString, { return MakeError(FText::GetEmpty()); })
+
+	/** Construct a new PCG Data subclass from an FInstancedStruct matching GetTargetStruct(). */
+	virtual TValueOrError<UPCGData*, FText> ConstructDataFromTarget(const FInstancedStruct& InTargetStruct, UObject* Outer = nullptr) const PURE_VIRTUAL(UPCGDataViewConverterBase::ConstructDataFromTarget, { return MakeError(FText::GetEmpty()); })
+	/** The custom UScriptStruct that contains the data after conversion in the target data format. */
+	virtual UScriptStruct* GetTargetStruct() const PURE_VIRTUAL(UPCGDataViewConverterBase::GetTargetStruct, { return nullptr; })
+	/** The custom UScriptStruct that contains the user defined parameters for the conversion. */
+	virtual UScriptStruct* GetParameterStruct() const PURE_VIRTUAL(UPCGDataViewConverterBase::GetParameterStruct, { return nullptr; })
+
+	/** Will leverage SerializeToString and save the result to the provided file path. */
+	TValueOrError<void, FText> SerializeToFile(const FPCGDataView& InDataView, const FInstancedStruct& Parameters, const FString& FilePath) const;
+};
diff --git a/Engine/Plugins/PCG/Source/PCG/Public/Data/DataView/PCGDataViewNativeJsonConverters.h b/Engine/Plugins/PCG/Source/PCG/Public/Data/DataView/PCGDataViewNativeJsonConverters.h
new file mode 100644
--- /dev/null
+++ b/Engine/Plugins/PCG/Source/PCG/Public/Data/DataView/PCGDataViewNativeJsonConverters.h
@@ -0,0 +1,44 @@
+// Copyright Epic Games, Inc. All Rights Reserved.
+
+#pragma once
+
+#include "Data/DataView/PCGDataViewInterface.h"
+
+#include "PCGDataViewNativeJsonConverters.generated.h"
+
+class FJsonObject;
+
+USTRUCT()
+struct FPCGDataViewJsonParameters
+{
+	GENERATED_BODY()
+
+	UPROPERTY(EditAnywhere, Category = "Json")
+	bool bPrettyJson = false;
+};
+
+/** Holds the root Json object for a PCG Data <-> Json conversion. @todo_pcg: consider a serializable option. */
+USTRUCT()
+struct FPCGDataViewJsonOutput
+{
+	GENERATED_BODY()
+
+	TSharedPtr<FJsonObject> JsonObject = nullptr;
+};
+
+UCLASS()
+class UPCGDataViewJsonConverter : public UPCGDataViewConverterBase
+{
+	GENERATED_BODY()
+
+public:
+	virtual TValueOrError<FInstancedStruct, FText> SerializeToTargetStruct(const FPCGDataView& InDataView, const FInstancedStruct& Parameters) const override;
+	virtual TValueOrError<FString, FText> SerializeToString(const FPCGDataView& InDataView, const FInstancedStruct& Parameters) const override;
+	virtual TValueOrError<UPCGData*, FText> ConstructDataFromTarget(const FInstancedStruct& InTargetStruct, UObject* Outer = nullptr) const override;
+	virtual UScriptStruct* GetTargetStruct() const override { return FPCGDataViewJsonOutput::StaticStruct(); }
+	virtual UScriptStruct* GetParameterStruct() const override { return FPCGDataViewJsonParameters::StaticStruct(); }
+
+protected:
+	/** Override to create a custom Json header. */
+	virtual void BuildJsonHeader(const FPCGDataView& InDataView, TSharedPtr<FJsonObject>& InOutJsonObject) const;
+};
diff --git a/Engine/Plugins/PCG/Source/PCG/Public/Data/DataView/PCGDataViewNativePropertySelectors.h b/Engine/Plugins/PCG/Source/PCG/Public/Data/DataView/PCGDataViewNativePropertySelectors.h
new file mode 100644
--- /dev/null
+++ b/Engine/Plugins/PCG/Source/PCG/Public/Data/DataView/PCGDataViewNativePropertySelectors.h
@@ -0,0 +1,33 @@
+// Copyright Epic Games, Inc. All Rights Reserved.
+
+#pragma once
+
+#include "Data/DataView/PCGDataViewInterface.h"
+
+namespace PCGDataView::Helpers
+{
+	// Append all non-property attributes to the provided array of selectors.
+	void AppendAllAttributeSelectors(const UPCGData* InData, TArray<FPCGAttributePropertySelector>& InOutSelectors);
+
+	// Append all relevant point properties avoiding redundancy, i.e. Transform already includes Positions, Rotation, Scale
+	void AppendAllPointPropertySelectors(TArray<FPCGAttributePropertySelector>& InOutSelectors);
+
+	// Append all relevant spline properties to the provided array of selectors.
+	void AppendAllSplinePropertySelectors(TArray<FPCGAttributePropertySelector>& InOutSelectors);
+} // namespace PCGDataView::Helpers
+
+class FPCGDataViewBasePointDataPropertySelector : public FPCGDataViewPropertySelector
+{
+public:
+	// ~Begin IPCGDataViewPropertySelector interface
+	virtual TArray<FPCGAttributePropertySelector> GetSelection(const FPCGDataView& InDataView) const override;
+	// ~End IPCGDataViewPropertySelector interface
+};
+
+class FPCGDataViewSplinePropertySelector : public FPCGDataViewPropertySelector
+{
+public:
+	// ~Begin IPCGDataViewPropertySelector interface
+	virtual TArray<FPCGAttributePropertySelector> GetSelection(const FPCGDataView& InDataView) const override;
+	// ~End IPCGDataViewPropertySelector interface
+};
diff --git a/Engine/Plugins/PCG/Source/PCG/Public/Data/Registry/PCGDataViewRegistry.h b/Engine/Plugins/PCG/Source/PCG/Public/Data/Registry/PCGDataViewRegistry.h
new file mode 100644
--- /dev/null
+++ b/Engine/Plugins/PCG/Source/PCG/Public/Data/Registry/PCGDataViewRegistry.h
@@ -0,0 +1,36 @@
+// Copyright Epic Games, Inc. All Rights Reserved.
+
+#pragma once
+
+#include "Templates/SubclassOf.h"
+
+class UPCGDataViewConverterBase;
+class IPCGDataViewPropertySelector;
+class UPCGData;
+struct FPCGDataView;
+
+class FPCGDataViewRegistry
+{
+public:
+	FPCGDataViewRegistry() = default;
+	FPCGDataViewRegistry(const FPCGDataViewRegistry&) = delete;
+	FPCGDataViewRegistry(FPCGDataViewRegistry&&) = default;
+	FPCGDataViewRegistry& operator=(const FPCGDataViewRegistry&) = delete;
+	FPCGDataViewRegistry& operator=(FPCGDataViewRegistry&&) = default;
+	virtual ~FPCGDataViewRegistry() = default;
+
+	PCG_API void RegisterPropertySelector(const TSubclassOf<UPCGData> InDataClass, TUniquePtr<const IPCGDataViewPropertySelector> PropertySelector);
+	PCG_API void UnregisterPropertySelector(const TSubclassOf<UPCGData> InDataClass);
+
+	PCG_API void RegisterConverter(const TSubclassOf<UPCGData> InDataClass, TSubclassOf<UPCGDataViewConverterBase> InConverterClass);
+	PCG_API void UnregisterConverter(const TSubclassOf<UPCGData> InDataClass, const TSubclassOf<UPCGDataViewConverterBase> InConverterClass);
+
+	PCG_API const IPCGDataViewPropertySelector* GetPropertySelector(const FPCGDataView& InDataView) const;
+	PCG_API const UPCGDataViewConverterBase* GetConverter(const FPCGDataView& InDataView, TSubclassOf<UPCGDataViewConverterBase> InConverterClass) const;
+
+private:
+	TMap<const TSubclassOf<UPCGData>, TUniquePtr<const IPCGDataViewPropertySelector>> PropertySelectorRegistry;
+
+	using FConverterRegistryType = TMap<const TSubclassOf<UPCGData>, TArray<TSubclassOf<UPCGDataViewConverterBase>>>;
+	FConverterRegistryType ConverterRegistry;
+};
diff --git a/Engine/Plugins/PCG/Source/PCG/Public/Elements/IO/PCGSaveDataViewElement.h b/Engine/Plugins/PCG/Source/PCG/Public/Elements/IO/PCGSaveDataViewElement.h
new file mode 100644
--- /dev/null
+++ b/Engine/Plugins/PCG/Source/PCG/Public/Elements/IO/PCGSaveDataViewElement.h
@@ -0,0 +1,83 @@
+// Copyright Epic Games, Inc. All Rights Reserved.
+
+#pragma once
+
+#include "PCGContext.h"
+#include "PCGSettings.h"
+#include "Data/DataView/PCGDataView.h"
+#include "Data/DataView/PCGDataViewNativeJsonConverters.h"
+
+#include "Serialization/CustomVersion.h"
+
+#include "PCGSaveDataViewElement.generated.h"
+
+/** Serialize a PCG Data via a Data View into a target format, i.e. Json or Binary. */
+UCLASS(BlueprintType, MinimalAPI, ClassGroup = (Procedural))
+class UPCGSaveDataViewSettings : public UPCGSettings
+{
+	GENERATED_BODY()
+
+public:
+	UPCGSaveDataViewSettings();
+
+	//~ Begin UPCGSettings interface
+#if WITH_EDITOR
+	virtual FName GetDefaultNodeName() const override { return FName(TEXT("PCGSaveDataView")); }
+	virtual FText GetDefaultNodeTitle() const override { return NSLOCTEXT("PCGSaveDataViewElement", "NodeTitle", "Save Data View"); }
+	virtual FText GetNodeTooltipText() const override { return NSLOCTEXT("PCGSaveDataViewElement", "NodeTooltip", "Export the Data View selection to a selected file."); }
+	virtual EPCGSettingsType GetType() const override { return EPCGSettingsType::InputOutput; }
+	virtual bool IsEditorOnly() const override { return true; }
+	virtual EPCGChangeType GetChangeTypeForProperty(const FName& InPropertyName) const override;
+	virtual void PostEditChangeProperty(FPropertyChangedEvent& PropertyChangedEvent) override;
+#endif // WITH_EDITOR
+
+protected:
+	virtual TArray<FPCGPinProperties> InputPinProperties() const override;
+	virtual TArray<FPCGPinProperties> OutputPinProperties() const override;
+
+	virtual FPCGElementPtr CreateElement() const override;
+#if WITH_EDITOR
+	virtual TArray<FPCGSettingsOverridableParam> GatherOverridableParams() const override;
+#endif // WITH_EDITOR
+	virtual void FixingOverridableParamPropertyClass(FPCGSettingsOverridableParam& Param) const override;
+	//~ End UPCGSettings interface
+
+public:
+	UPROPERTY(EditAnywhere, BlueprintReadWrite, Category = Settings, NoClear)
+	TSubclassOf<UPCGDataViewConverterBase> ConverterClass = UPCGDataViewJsonConverter::StaticClass();
+
+	UPROPERTY(EditAnywhere, BlueprintReadWrite, Category = Settings, meta = (ShowOnlyInnerProperties, StructTypeConst))
+	FInstancedStruct ConverterParameters;
+
+	/** The file path to save the data. */
+	UPROPERTY(EditAnywhere, BlueprintReadWrite, Category = Settings, meta = (PCG_Overridable))
+	FFilePath FilePath;
+
+#if WITH_EDITORONLY_DATA
+	/** Open a file dialogue on empty path. */
+	UPROPERTY(EditAnywhere, BlueprintReadWrite, Category = Settings, meta = (PCG_Overridable))
+	bool bOpenDialogueOnEmptyPath = false;
+#endif // WITH_EDITORONLY_DATA
+
+	/** Overwrite the existing file, if it exists. */
+	UPROPERTY(EditAnywhere, BlueprintReadWrite, Category = Settings, meta = (PCG_Overridable))
+	bool bOverwriteExistingFile = false;
+};
+
+class FPCGSaveDataViewExecutionContext : public FPCGContext
+{
+protected:
+	virtual FStructView GetExternalStructContainerForOverridableParam(const FPCGSettingsOverridableParam& InParam) override;
+};
+
+class FPCGSaveDataViewElement final : public IPCGElement
+{
+public:
+	virtual bool CanExecuteOnlyOnMainThread(FPCGContext* Context) const override { return true; }
+	// @todo_pcg: Checksum the file?
+	virtual bool IsCacheable(const UPCGSettings* InSettings) const override { return false; }
+
+protected:
+	virtual FPCGContext* CreateContext() override;
+	virtual bool ExecuteInternal(FPCGContext* InContext) const override;
+};
diff --git a/Engine/Plugins/PCG/Source/PCG/Public/Elements/PCGConvertToDataView.h b/Engine/Plugins/PCG/Source/PCG/Public/Elements/PCGConvertToDataView.h
new file mode 100644
--- /dev/null
+++ b/Engine/Plugins/PCG/Source/PCG/Public/Elements/PCGConvertToDataView.h
@@ -0,0 +1,60 @@
+// Copyright Epic Games, Inc. All Rights Reserved.
+
+#pragma once
+
+#include "PCGSettings.h"
+#include "Data/DataView/PCGDataView.h"
+
+#include "PCGConvertToDataView.generated.h"
+
+/** Conversion node from any PCG Data Type to a Data View Data.
+ * Constructs a view onto the PCG Data, allowing future operations to be performed on a selection of
+ * internal data or metadata.
+ */
+UCLASS(MinimalAPI, BlueprintType, ClassGroup = (Procedural))
+class UPCGConvertToDataViewSettings : public UPCGSettings
+{
+	GENERATED_BODY()
+
+public:
+	//~Begin UPCGSettings interface
+#if WITH_EDITOR
+	virtual FName GetDefaultNodeName() const override { return FName(TEXT("ConvertToDataView")); }
+	virtual FText GetDefaultNodeTitle() const override { return NSLOCTEXT("PCGConvertToDataViewElement", "NodeTitle", "To Data View"); }
+	virtual FText GetNodeTooltipText() const override { return NSLOCTEXT("PCGConvertToDataViewElement", "NodeTooltip", "Construct a view onto the PCG Data, which allows for extraction or conversion of selected internal data."); }
+	virtual EPCGSettingsType GetType() const override { return EPCGSettingsType::Generic; }
+	virtual EPCGChangeType GetChangeTypeForProperty(const FName& InPropertyName) const override;
+	virtual bool ShouldDrawNodeCompact() const override;
+#endif // WITH_EDITOR
+	virtual bool HasExecutionDependencyPin() const override { return false; }
+
+protected:
+	virtual TArray<FPCGPinProperties> InputPinProperties() const override;
+	virtual TArray<FPCGPinProperties> OutputPinProperties() const override;
+	virtual FPCGElementPtr CreateElement() const override;
+	//~End UPCGSettings interface
+
+public:
+	/** Provide the selection as an attribute set of selector strings directly into a secondary input. */
+	UPROPERTY(EditAnywhere, BlueprintReadWrite, Category = "Settings")
+	bool bSelectionAsInput = false;
+
+	/** The selection sources as a comma-separated string. E.g. "$Index, $Position, CustomTangent", etc. */
+	UPROPERTY(EditAnywhere, BlueprintReadWrite, Category = "Settings", meta = (EditCondition = "bSelectionAsInput", EditConditionHides, PCG_Overridable))
+	FString SelectionSources;
+
+	/** Manually select attributes or properties for this data view. */
+	UPROPERTY(EditAnywhere, BlueprintReadWrite, Category = "Settings", meta = (EditCondition = "!bSelectionAsInput", EditConditionHides))
+	FPCGDataViewSelection Selection;
+
+	/** How the attribute data should be exported, assuming the data supports the elements domain. */
+	UPROPERTY(EditAnywhere, BlueprintReadWrite, Category = "Settings", meta = (PCG_Overridable))
+	EPCGDataViewAttributeLayout AttributeLayout = EPCGDataViewAttributeLayout::ByElement;
+};
+
+class FPCGConvertToDataViewElement : public IPCGElement
+{
+	virtual bool SupportsBasePointDataInputs(FPCGContext* InContext) const override { return true; }
+	virtual EPCGElementExecutionLoopMode ExecutionLoopMode(const UPCGSettings* Settings) const override { return EPCGElementExecutionLoopMode::SinglePrimaryPin; }
+	virtual bool ExecuteInternal(FPCGContext* InContext) const override;
+};
diff --git a/Engine/Plugins/PCG/Source/PCG/Public/Helpers/IO/PCGIOHelpers.h b/Engine/Plugins/PCG/Source/PCG/Public/Helpers/IO/PCGIOHelpers.h
--- a/Engine/Plugins/PCG/Source/PCG/Public/Helpers/IO/PCGIOHelpers.h
+++ b/Engine/Plugins/PCG/Source/PCG/Public/Helpers/IO/PCGIOHelpers.h
@@ -10,28 +10,33 @@ namespace PCG::IO
 {
 	namespace Constants
 	{
-		namespace Attribute
+		// Custom PCG version for attribute export
+		struct FCustomExportVersion
 		{
-			// Custom PCG version for attribute export
-			struct FCustomExportVersion
+			enum Type
 			{
-				enum Type
-				{
-					InitialVersion = 0,
-
-					// New versions can be added above this line
-					VersionPlusOne,
-					LatestVersion = VersionPlusOne - 1
-				};
+				InitialVersion = 0,
 
-				static FName GetFriendlyName() { return FName(TEXT("PCG::IO::AttributeExport")); }
-				const static FGuid GUID;
+				// New versions can be added above this line
+				VersionPlusOne,
+				LatestVersion = VersionPlusOne - 1
 			};
-		} // namespace Attribute
 
-		static constexpr int32 InlineAllocationCount = 16;
+			static FName GetFriendlyName() { return FName(TEXT("PCG::IO::AttributeExport")); }
+			const static FGuid GUID;
+		};
+
+		namespace Attribute
+		{
+			struct UE_DEPRECATED(5.8, "Use PCG::IO::Constants::FCustomExportVersion") FCustomExportVersion : Constants::FCustomExportVersion {};
+		}
 	} // namespace Constants
 
+	namespace Binary::Constants
+	{
+		static const FString Extension = TEXT(".bin");
+	}
+
 	namespace Helpers::String
 	{
 		void ToPrecisionString(const double Value, const int32 Precision, FString& OutString);
@@ -39,6 +44,11 @@ namespace PCG::IO
 
 	namespace Accessor
 	{
+		namespace Constants
+		{
+			static constexpr int32 CacheInlineAllocationCount = 16;
+		}
+
 		struct FCacheEntry
 		{
 			FPCGAttributePropertySelector Selector;
@@ -46,6 +56,6 @@ namespace PCG::IO
 			TUniquePtr<const IPCGAttributeAccessorKeys> Keys = nullptr;
 		};
 
-		using FCache = TArray<FCacheEntry, TInlineAllocator<Constants::InlineAllocationCount>>;
+		using FCache = TArray<FCacheEntry, TInlineAllocator<Constants::CacheInlineAllocationCount>>;
 	} // namespace Accessor
 } // namespace PCG::IO
diff --git a/Engine/Plugins/PCG/Source/PCG/Public/Helpers/IO/PCGJsonHelpers.h b/Engine/Plugins/PCG/Source/PCG/Public/Helpers/IO/PCGJsonHelpers.h
--- a/Engine/Plugins/PCG/Source/PCG/Public/Helpers/IO/PCGJsonHelpers.h
+++ b/Engine/Plugins/PCG/Source/PCG/Public/Helpers/IO/PCGJsonHelpers.h
@@ -2,6 +2,7 @@
 
 #pragma once
 
+#include "PCGData.h"
 #include "Metadata/PCGMetadataAttributeTraits.h"
 
 #include "Dom/JsonValue.h"
@@ -18,6 +19,12 @@ namespace PCG::IO::Json
 	template <typename T>
 	concept CFloatingPoint = std::is_floating_point_v<std::remove_cv_t<T>>;
 
+	namespace Constants
+	{
+		static const FString Extension = TEXT(".json");
+		static const FString DefaultDataVersionKey = TEXT("data_version");
+	}
+
 	namespace Helpers
 	{
 		// Converts a floating point (float, double, FVector, FRotator, etc) type to a Json value.
@@ -68,12 +75,21 @@ namespace PCG::IO::Json
 			}
 		}
 
-		// Append PCG Json data version header object to a Json object
+		UE_DEPRECATED(5.8, "AppendHeader is deprecated, please use AppendVersionToHeader instead.")
 		void AppendHeader(const TSharedPtr<FJsonObject>& InOutJsonObject, int32 DataVersion, int32 CustomDataVersion = -1);
 
+		// Append PCG Json data version header object to a Json object
+		void AppendVersionToHeader(const TSharedPtr<FJsonObject>& InOutJsonObject, int32 DataVersion);
+
+		// Append the selected properties/attributes to the provided Json object on an attribute-wise axis.
+		void AppendSelectionByAttribute(TSharedPtr<FJsonObject>& InOutJsonObject, const UPCGData* InData, const TArray<FPCGAttributePropertySelector>& Selectors, const FPCGContext* InOptionalContext = nullptr);
+
+		// Append the selected properties/attributes to the provided Json object on an element-wise axis.
+		void AppendSelectionByElement(TSharedPtr<FJsonObject>& InOutJsonObject, const UPCGData* InData, const TArray<FPCGAttributePropertySelector>& Selectors, const FPCGContext* InOptionalContext = nullptr);
+
 		// Helper to convert a Json object into a string
-		FString ToJsonString(const TSharedPtr<FJsonObject>& JsonObject);
+		FString ToJsonString(const TSharedPtr<FJsonObject>& JsonObject, bool bPrettyJson = false);
 
 		// @todo_pcg: Add a fast track for AppendAllAttributes
-	}
-}
+	} // namespace Helpers
+} // namespace PCG::IO::Json
diff --git a/Engine/Plugins/PCG/Source/PCG/Public/PCGContext.h b/Engine/Plugins/PCG/Source/PCG/Public/PCGContext.h
--- a/Engine/Plugins/PCG/Source/PCG/Public/PCGContext.h
+++ b/Engine/Plugins/PCG/Source/PCG/Public/PCGContext.h
@@ -186,6 +186,14 @@ PRAGMA_ENABLE_DEPRECATION_WARNINGS
 		return SettingsWithOverride ? Cast<SettingsType>(SettingsWithOverride) : GetOriginalSettings<SettingsType>();
 	}
 
+	// Return the mutable settings casted in the wanted type.
+	// If there is no override, it returns null.
+	template<typename SettingsType>
+	SettingsType* GetMutableInputSettings()
+	{
+		return SettingsWithOverride ? Cast<SettingsType>(SettingsWithOverride) : nullptr;
+	}
+
 	UE_API FPCGTaskId GetGraphExecutionTaskId() const;
 
 	UE_API FString GetTaskName() const;
diff --git a/Engine/Plugins/PCG/Source/PCG/Public/PCGModule.h b/Engine/Plugins/PCG/Source/PCG/Public/PCGModule.h
--- a/Engine/Plugins/PCG/Source/PCG/Public/PCGModule.h
+++ b/Engine/Plugins/PCG/Source/PCG/Public/PCGModule.h
@@ -5,6 +5,7 @@
 #include "PCGDataVisualizationRegistry.h"
 #include "PCGGraphExecutionRegistry.h"
 #include "Data/Registry/PCGDataTypeRegistry.h"
+#include "Data/Registry/PCGDataViewRegistry.h"
 #include "Data/Registry/PCGGetDataFunctionRegistry.h"
 #include "Hash/PCGObjectHash.h"
 #include "Metadata/Accessors/PCGAttributeAccessorFactory.h"
@@ -59,6 +60,9 @@ public:
 	static const FPCGGraphExecutionRegistry& GetConstGraphExecutionRegistry() { return GetPCGModuleChecked().GraphExecutionRegistry; }
 	static FPCGGraphExecutionRegistry& GetGraphExecutionRegistry() { return GetPCGModuleChecked().GraphExecutionRegistry; }
 
+	PCG_API static const FPCGDataViewRegistry& GetConstPCGDataViewRegistry() { return GetPCGModuleChecked().DataViewRegistry; }
+	PCG_API static FPCGDataViewRegistry& GetMutablePCGDataViewRegistry() { return GetPCGModuleChecked().DataViewRegistry; }
+
 #if WITH_EDITOR
 	static const FPCGObjectHashFactory& GetConstObjectHashFactory() { return GetPCGModuleChecked().ObjectHashFactory; }
 	static FPCGObjectHashFactory& GetMutableObjectHashFactory() { return GetPCGModuleChecked().ObjectHashFactory; }
@@ -84,6 +88,7 @@ private:
 	FPCGAttributeAccessorFactory AttributeAccessorFactory;
 	FPCGDataTypeRegistry DataTypeRegistry;
 	FPCGGraphExecutionRegistry GraphExecutionRegistry;
+	FPCGDataViewRegistry DataViewRegistry;
 
 #if WITH_EDITOR
 	FPCGObjectHashFactory ObjectHashFactory;
@@ -111,4 +116,4 @@ private:
 
 	static FPCGGraphChangedDelegate OnGraphChangedDelegate;
 #endif // WITH_EDITOR
-};
\ No newline at end of file
+};
diff --git a/Engine/Plugins/PCG/Source/PCGEditor/Private/DataVisualizations/PCGDataViewVisualization.cpp b/Engine/Plugins/PCG/Source/PCGEditor/Private/DataVisualizations/PCGDataViewVisualization.cpp
new file mode 100644
--- /dev/null
+++ b/Engine/Plugins/PCG/Source/PCGEditor/Private/DataVisualizations/PCGDataViewVisualization.cpp
@@ -0,0 +1,124 @@
+// Copyright Epic Games, Inc. All Rights Reserved.
+
+#include "DataVisualizations/PCGDataViewVisualization.h"
+
+#include "Data/DataView/PCGDataView.h"
+#include "Data/DataView/PCGDataViewData.h"
+#include "DataVisualizations/PCGDataVisualizationHelpers.h"
+
+void FPCGDataViewVisualization::ExecuteDebugDisplay(FPCGContext* Context, const UPCGSettingsInterface* SettingsInterface, const UPCGData* Data, AActor* TargetActor) const
+{
+	/* Do nothing. */
+}
+
+// @todo_pcg: It would be nice to preset the filter to only show selected attributes, even with custom visualizer.
+FPCGTableVisualizerInfo FPCGDataViewVisualization::GetTableVisualizerInfoWithDomain(const UPCGData* Data, const FPCGMetadataDomainID& DomainID) const
+{
+	const UPCGDataViewData* DataViewData = Cast<UPCGDataViewData>(Data);
+	if (!DataViewData)
+	{
+		return FPCGTableVisualizerInfo{};
+	}
+
+	// Pull the visualization of the underlying type
+	const FPCGDataView& DataView = DataViewData->GetDataView();
+	if (const UPCGData* ViewedData = DataView.ViewedData)
+	{
+		const FPCGDataVisualizationRegistry& DataVisRegistry = FPCGModule::GetConstPCGDataVisualizationRegistry();
+
+		if (const IPCGDataVisualization* DataVisualization = DataVisRegistry.GetDataVisualization(ViewedData->GetClass()))
+		{
+			return DataVisualization->GetTableVisualizerInfoWithDomain(ViewedData, DomainID);
+		}
+
+		FPCGTableVisualizerInfo Info;
+		Info.Data = ViewedData;
+
+		if (DataView.Selection.bAllAttributes)
+		{
+			PCGDataVisualizationHelpers::CreateDefaultMetadataColumnInfos(ViewedData, DomainID);
+			return Info;
+		}
+
+		// Any data beyond here has no custom visualizer and isn't selecting all attributes. Add attributes one at a time.
+		const UPCGMetadata* Metadata = ViewedData->ConstMetadata();
+		check(Metadata);
+
+		const FPCGMetadataDomain* MetadataDomain = Metadata ? Metadata->GetConstMetadataDomain(DomainID) : nullptr;
+		if (!MetadataDomain)
+		{
+			return FPCGTableVisualizerInfo{};
+		}
+
+		// Sort the index to the front
+		TArray<FPCGAttributePropertySelector> Attributes = DataView.Selection.Attributes;
+		FPCGAttributePropertySelector IndexSelector = FPCGAttributePropertySelector::CreateExtraPropertySelector(EPCGExtraProperties::Index);
+		const int32 IndexIndex = Attributes.IndexOfByPredicate([&IndexSelector](const FPCGAttributePropertySelector& Selector)
+		{
+			return Selector == IndexSelector;
+		});
+
+		if (IndexIndex > 0)
+		{
+			Attributes.Swap(0, IndexIndex);
+		}
+
+		for (const FPCGAttributePropertySelector& Selector : Attributes)
+		{
+			const FName AttributeName = Selector.IsBasicAttribute() ? Selector.GetAttributeName() : Selector.GetPropertyName();
+			FPCGMetadataDomainID ID = ViewedData->GetMetadataDomainIDFromSelector(Selector);
+			if (MetadataDomain->HasAttribute(AttributeName) && ID == DomainID)
+			{
+				auto AddInfo = [ViewedData, &Info]<typename T>(T, const FPCGAttributePropertySelector& InSelector)
+				{
+					PCGDataVisualizationHelpers::AddTypedColumnInfo<T>(Info, ViewedData, InSelector);
+				};
+
+				const FPCGMetadataAttributeBase* AttributeBase = MetadataDomain->GetConstAttribute(AttributeName);
+				PCGMetadataAttribute::CallbackWithRightType(AttributeBase->GetTypeId(), AddInfo, Selector);
+			}
+		}
+
+		if (!Info.ColumnInfos.IsEmpty())
+		{
+			// Sort by $Index, if it exists. Otherwise, sort by the first column. 
+			const bool bHasIndex = Info.ColumnInfos.ContainsByPredicate([&IndexSelector](const FPCGTableVisualizerColumnInfo& ColumnInfo)
+			{
+				return ColumnInfo.Id == IndexSelector.GetPropertyName();
+			});
+
+			if (bHasIndex)
+			{
+				// Attributes were sorted, so this should be guaranteed
+				ensure(Info.ColumnInfos[0].Id == IndexSelector.GetExtraNames()[0]);
+				Info.SortingColumn = Info.ColumnInfos[0].Id;
+			}
+			else // Create an index column
+			{
+				ViewedData->SetDomainFromDomainID(DomainID, IndexSelector);
+				PCGDataVisualizationHelpers::AddColumnInfo(Info, ViewedData, IndexSelector);
+				Info.SortingColumn = Info.ColumnInfos.Last().Id;
+			}
+		}
+
+		return Info;
+	}
+
+	return FPCGTableVisualizerInfo{};
+}
+
+FString FPCGDataViewVisualization::GetDomainDisplayNameForInspection(const UPCGData* Data, const FPCGMetadataDomainID& DomainID) const
+{
+	if (const UPCGDataViewData* DataViewData = Cast<UPCGDataViewData>(Data))
+	{
+		const FPCGDataView& DataView = DataViewData->GetDataView();
+		if (const UPCGData* ViewedData = Cast<UPCGData>(DataView.ViewedData))
+		{
+			FPCGAttributePropertySelector Selector;
+			ViewedData->SetDomainFromDomainID(DomainID, Selector);
+			return Selector.GetDomainName().ToString();
+		}
+	}
+
+	return {};
+}
diff --git a/Engine/Plugins/PCG/Source/PCGEditor/Private/PCGEditorModule.cpp b/Engine/Plugins/PCG/Source/PCGEditor/Private/PCGEditorModule.cpp
--- a/Engine/Plugins/PCG/Source/PCGEditor/Private/PCGEditorModule.cpp
+++ b/Engine/Plugins/PCG/Source/PCGEditor/Private/PCGEditorModule.cpp
@@ -70,6 +70,8 @@
 #include "PropertyEditorModule.h"
 #include "ScopedTransaction.h"
 #include "ToolMenus.h"
+#include "Data/DataView/PCGDataViewData.h"
+#include "DataVisualizations/PCGDataViewVisualization.h"
 #include "Details/PCGRaycastFilterRuleCollectionDetails.h"
 #include "Editor/SceneOutliner/Public/ISceneOutliner.h"
 #include "EditorMode/PCGEdMode.h"
@@ -861,6 +863,7 @@ void FPCGEditorModule::RegisterPCGDataVisualizations()
 	DataVisRegistry.InternalRegistry.Add(UPCGBaseTextureData::StaticClass(), MakeUnique<const FPCGBaseTextureDataVisualization>());
 	DataVisRegistry.InternalRegistry.Add(UPCGPolygon2DData::StaticClass(), MakeUnique<const FPCGPolygon2DDataVisualization>());
 	DataVisRegistry.InternalRegistry.Add(UPCGPolygon2DInteriorSurfaceData::StaticClass(), MakeUnique<const FPCGPolygon2DDataVisualization>());
+	DataVisRegistry.InternalRegistry.Add(UPCGDataViewData::StaticClass(), MakeUnique<const FPCGDataViewVisualization>());
 }
 
 void FPCGEditorModule::UnregisterPCGDataVisualizations()
@@ -900,6 +903,7 @@ void FPCGEditorModule::RegisterPinColorAndIcons()
 	InRegistry.RegisterPinColorFunction(FPCGDataTypeInfoDynamicMesh::AsId(), [](const FPCGDataTypeIdentifier&) { return GetDefault<UPCGEditorSettings>()->DynamicMeshPinColor; });
 	InRegistry.RegisterPinColorFunction(FPCGDataTypeInfoSpatial::AsId(), [](const FPCGDataTypeIdentifier&) { return GetDefault<UPCGEditorSettings>()->SpatialDataPinColor; });
 	InRegistry.RegisterPinColorFunction(FPCGDataTypeInfoStaticMeshResource::AsId(), [](const FPCGDataTypeIdentifier&) { return GetDefault<UPCGEditorSettings>()->StaticMeshResourcePinColor; });
+	InRegistry.RegisterPinColorFunction(FPCGDataTypeInfoDataView::AsId(), [](const FPCGDataTypeIdentifier&) { return GetDefault<UPCGEditorSettings>()->DataViewPinColor; });
 	InRegistry.RegisterPinColorFunction(FPCGDataTypeInfoOther::AsId(), [](const FPCGDataTypeIdentifier&) { return GetDefault<UPCGEditorSettings>()->UnknownDataPinColor; });
 	InRegistry.RegisterPinColorFunction(FPCGDataTypeInfo::AsId(), [](const FPCGDataTypeIdentifier&) { return GetDefault<UPCGEditorSettings>()->DefaultPinColor; });
 
@@ -985,6 +989,7 @@ void FPCGEditorModule::UnregisterPinColorAndIcons()
 	InRegistry.UnregisterPinColorFunction(FPCGDataTypeInfoConcrete::AsId());
 	InRegistry.UnregisterPinColorFunction(FPCGDataTypeInfoPoint::AsId());
 	InRegistry.UnregisterPinColorFunction(FPCGDataTypeInfoParam::AsId());
+	InRegistry.UnregisterPinColorFunction(FPCGDataTypeInfoDataView::AsId());
 
 	// PinIcons. Just have catch all as most of it will be the same icons.
 	InRegistry.UnregisterPinIconsFunction(FPCGDataTypeInfo::AsId());
diff --git a/Engine/Plugins/PCG/Source/PCGEditor/Private/PCGEditorSettings.cpp b/Engine/Plugins/PCG/Source/PCGEditor/Private/PCGEditorSettings.cpp
--- a/Engine/Plugins/PCG/Source/PCGEditor/Private/PCGEditorSettings.cpp
+++ b/Engine/Plugins/PCG/Source/PCGEditor/Private/PCGEditorSettings.cpp
@@ -48,6 +48,7 @@ UPCGEditorSettings::UPCGEditorSettings(const FObjectInitializer& ObjectInitializ
 	LandscapeDataPinColor = FLinearColor(0.66f, 0.66f, 0.07f);
 	BaseTextureDataPinColor = FLinearColor(0.81f, 0.28f, 0.22f);
 	TextureDataPinColor = FLinearColor(0.79f, 0.08f, 0.01f);
+	DataViewPinColor = FLinearColor(0.69f, 0.92f, 0.23f);
 	RenderTargetDataPinColor = FLinearColor(0.8f, 0.18f, 0.12f);
 	VirtualTextureDataPinColor = FLinearColor(0.8f, 0.48f, 0.22f);
 	VolumeDataPinColor = FLinearColor(0.79f, 0.06f, 0.5f);
diff --git a/Engine/Plugins/PCG/Source/PCGEditor/Public/DataVisualizations/PCGDataViewVisualization.h b/Engine/Plugins/PCG/Source/PCGEditor/Public/DataVisualizations/PCGDataViewVisualization.h
new file mode 100644
--- /dev/null
+++ b/Engine/Plugins/PCG/Source/PCGEditor/Public/DataVisualizations/PCGDataViewVisualization.h
@@ -0,0 +1,18 @@
+// Copyright Epic Games, Inc. All Rights Reserved.
+
+#pragma once
+
+#include "PCGDataVisualization.h"
+
+class UPCGData;
+
+class FPCGDataViewVisualization : public IPCGDataVisualization
+{
+public:
+	// ~Begin IPCGDataVisualization interface
+	virtual void ExecuteDebugDisplay(FPCGContext* Context, const UPCGSettingsInterface* SettingsInterface, const UPCGData* Data, AActor* TargetActor) const override;
+	virtual ~FPCGDataViewVisualization() override = default;
+	virtual FPCGTableVisualizerInfo GetTableVisualizerInfoWithDomain(const UPCGData* Data, const FPCGMetadataDomainID& DomainID) const override;
+	virtual FString GetDomainDisplayNameForInspection(const UPCGData* Data, const FPCGMetadataDomainID& DomainID) const override;
+	// ~End IPCGDataVisualization interface
+};
diff --git a/Engine/Plugins/PCG/Source/PCGEditor/Public/PCGEditorSettings.h b/Engine/Plugins/PCG/Source/PCGEditor/Public/PCGEditorSettings.h
--- a/Engine/Plugins/PCG/Source/PCGEditor/Public/PCGEditorSettings.h
+++ b/Engine/Plugins/PCG/Source/PCGEditor/Public/PCGEditorSettings.h
@@ -201,6 +201,10 @@ public:
 	UPROPERTY(EditAnywhere, config, Category = Node, meta = (HideAlphaChannel))
 	FLinearColor TextureDataPinColor;
 
+	/** Color used for data pins of type Data View */
+	UPROPERTY(EditAnywhere, config, Category = Node, meta = (HideAlphaChannel))
+	FLinearColor DataViewPinColor;
+
 	/** Color used for data pins of type Render Target */
 	UPROPERTY(EditAnywhere, config, Category = Node, meta = (HideAlphaChannel))
 	FLinearColor RenderTargetDataPinColor;
