[Backout] - CL47861923
#rnx
Original CL Desc
-----------------------------------------------------------------
[PCG] Combined Submission for Data View type

CLs: 44349862, 44569695, 44569728, 44569766, 44569794, 45852129

diff --git a/Engine/Plugins/PCG/Source/PCG/Private/Data/DataView/PCGDataView.cpp b/Engine/Plugins/PCG/Source/PCG/Private/Data/DataView/PCGDataView.cpp
deleted file mode 100644
diff --git a/Engine/Plugins/PCG/Source/PCG/Private/Data/DataView/PCGDataViewData.cpp b/Engine/Plugins/PCG/Source/PCG/Private/Data/DataView/PCGDataViewData.cpp
deleted file mode 100644
diff --git a/Engine/Plugins/PCG/Source/PCG/Private/Data/DataView/PCGDataViewInterface.cpp b/Engine/Plugins/PCG/Source/PCG/Private/Data/DataView/PCGDataViewInterface.cpp
deleted file mode 100644
diff --git a/Engine/Plugins/PCG/Source/PCG/Private/Data/DataView/PCGDataViewNativeJsonConverters.cpp b/Engine/Plugins/PCG/Source/PCG/Private/Data/DataView/PCGDataViewNativeJsonConverters.cpp
deleted file mode 100644
diff --git a/Engine/Plugins/PCG/Source/PCG/Private/Data/DataView/PCGDataViewNativePropertySelectors.cpp b/Engine/Plugins/PCG/Source/PCG/Private/Data/DataView/PCGDataViewNativePropertySelectors.cpp
deleted file mode 100644
diff --git a/Engine/Plugins/PCG/Source/PCG/Private/Data/Registry/PCGDataViewRegistry.cpp b/Engine/Plugins/PCG/Source/PCG/Private/Data/Registry/PCGDataViewRegistry.cpp
deleted file mode 100644
diff --git a/Engine/Plugins/PCG/Source/PCG/Private/Elements/IO/PCGExportSelectedAttributes.cpp b/Engine/Plugins/PCG/Source/PCG/Private/Elements/IO/PCGExportSelectedAttributes.cpp
--- a/Engine/Plugins/PCG/Source/PCG/Private/Elements/IO/PCGExportSelectedAttributes.cpp
+++ b/Engine/Plugins/PCG/Source/PCG/Private/Elements/IO/PCGExportSelectedAttributes.cpp
@@ -186,11 +186,11 @@ bool FPCGExportSelectedAttributesElement::ExecuteInternal(FPCGContext* Context)
 	FArchive* Archive = nullptr;
 	TSharedPtr<FJsonObject> RootJsonObject = nullptr;
 
-	using namespace PCG::IO::Constants;
+	using namespace PCG::IO::Constants::Attribute;
 	// Initialize the serialization object needed.
 	if (Settings->Format == EPCGExportAttributesFormat::Binary)
 	{
-		FString FinalPath = *Helpers::String::MergeFileName(OutputDirectory, Settings->FileName, PCG::IO::Binary::Constants::Extension);
+		FString FinalPath = *Helpers::String::MergeFileName(OutputDirectory, Settings->FileName, TEXT(".bin"));
 		Archive = IFileManager::Get().CreateFileWriter(*FinalPath);
 		if (!Archive)
 		{
@@ -214,9 +214,8 @@ bool FPCGExportSelectedAttributesElement::ExecuteInternal(FPCGContext* Context)
 
 		// Set up the Json header to include custom versions
 		TSharedPtr<FJsonObject> JsonHeaderObject = MakeShared<FJsonObject>();
-		PCG::IO::Json::Helpers::AppendVersionToHeader(JsonHeaderObject, FCustomExportVersion::LatestVersion);
-		JsonHeaderObject->SetNumberField(TEXT("custom_data_version"), Settings->bAddCustomDataVersion ? Settings->CustomVersion : -1);
-		RootJsonObject->SetObjectField(TEXT("Header"), std::move(JsonHeaderObject));
+		PCG::IO::Json::Helpers::AppendHeader(JsonHeaderObject, FCustomExportVersion::LatestVersion, Settings->bAddCustomDataVersion ? Settings->CustomVersion : -1);
+		RootJsonObject->SetObjectField(LOCTEXT("JsonHeaderField", "Header").ToString(), std::move(JsonHeaderObject));
 	}
 
 	const UPCGData* InputData = Inputs[0].Data;
diff --git a/Engine/Plugins/PCG/Source/PCG/Private/Elements/IO/PCGSaveDataViewElement.cpp b/Engine/Plugins/PCG/Source/PCG/Private/Elements/IO/PCGSaveDataViewElement.cpp
deleted file mode 100644
diff --git a/Engine/Plugins/PCG/Source/PCG/Private/Elements/PCGConvertToDataView.cpp b/Engine/Plugins/PCG/Source/PCG/Private/Elements/PCGConvertToDataView.cpp
deleted file mode 100644
diff --git a/Engine/Plugins/PCG/Source/PCG/Private/Helpers/IO/PCGIOHelpers.cpp b/Engine/Plugins/PCG/Source/PCG/Private/Helpers/IO/PCGIOHelpers.cpp
--- a/Engine/Plugins/PCG/Source/PCG/Private/Helpers/IO/PCGIOHelpers.cpp
+++ b/Engine/Plugins/PCG/Source/PCG/Private/Helpers/IO/PCGIOHelpers.cpp
@@ -11,8 +11,11 @@ namespace PCG::IO
 {
 	namespace Constants
 	{
-		const FGuid FCustomExportVersion::GUID = FGuid(0x04E74488, 0x4BAC8717, 0xBBB18694, 0x39F8F3CE);
-		FCustomVersionRegistration GRegisterPCGExportSelectedAttributesCustomVersion(FCustomExportVersion::GUID, FCustomExportVersion::LatestVersion, TEXT("PCGExportSelectedAttributes"));
+		namespace Attribute
+		{
+			const FGuid FCustomExportVersion::GUID = FGuid(0x04E74488, 0x4BAC8717, 0xBBB18694, 0x39F8F3CE);
+			FCustomVersionRegistration GRegisterPCGExportSelectedAttributesCustomVersion(FCustomExportVersion::GUID, FCustomExportVersion::LatestVersion, TEXT("PCGExportSelectedAttributes"));
+		}
 	}
 
 	namespace Helpers::String
diff --git a/Engine/Plugins/PCG/Source/PCG/Private/Helpers/IO/PCGJsonHelpers.cpp b/Engine/Plugins/PCG/Source/PCG/Private/Helpers/IO/PCGJsonHelpers.cpp
--- a/Engine/Plugins/PCG/Source/PCG/Private/Helpers/IO/PCGJsonHelpers.cpp
+++ b/Engine/Plugins/PCG/Source/PCG/Private/Helpers/IO/PCGJsonHelpers.cpp
@@ -2,10 +2,8 @@
 
 #include "Helpers/IO/PCGJsonHelpers.h"
 
-#include "Helpers/PCGHelpers.h"
 #include "Helpers/IO/PCGIOHelpers.h"
 #include "Metadata/Accessors/PCGAttributeAccessorHelpers.h"
-#include "Policies/CondensedJsonPrintPolicy.h"
 
 #include "Serialization/JsonSerializer.h"
 #include "Serialization/JsonWriter.h"
@@ -63,7 +61,7 @@ namespace PCG::IO::Json
 				if constexpr (PCG::Private::MetadataTraits<T>::IsFloatingPoint)
 				{
 					TArray<T> Values;
-					PCGAttributeAccessorHelpers::ExtractAllValues(Entry.Accessor.Get(), Entry.Keys.Get(), Values, EPCGAttributeAccessorFlags::AllowBroadcastAndConstructible);
+					PCGAttributeAccessorHelpers::ExtractAllValues(Entry.Accessor.Get(), Entry.Keys.Get(), Values, EPCGAttributeAccessorFlags::StrictType);
 
 					JsonValues.Reserve(Values.Num());
 					Algo::Transform(Values, JsonValues, [&Entry](const T& Value)
@@ -74,7 +72,7 @@ namespace PCG::IO::Json
 				else // For everything else, automatically broadcast it to string
 				{
 					TArray<FString> Values;
-					PCGAttributeAccessorHelpers::ExtractAllValues(Entry.Accessor.Get(), Entry.Keys.Get(), Values, EPCGAttributeAccessorFlags::AllowBroadcastAndConstructible);
+					PCGAttributeAccessorHelpers::ExtractAllValues(Entry.Accessor.Get(), Entry.Keys.Get(), Values, EPCGAttributeAccessorFlags::StrictType);
 
 					JsonValues.Reserve(Values.Num());
 					Algo::Transform(Values, JsonValues, [](const FString& OutValue)
@@ -95,8 +93,8 @@ namespace PCG::IO::Json
 			for (int32 CacheIndex = 0; CacheIndex < AccessorCache.Num(); ++CacheIndex)
 			{
 				const Accessor::FCacheEntry& Entry = AccessorCache[CacheIndex];
-				FString& FieldName = OutKeyValuesMap[CacheIndex].Get<0>();
-				FieldName = Entry.Selector.ToString();
+				FString& SourceName = OutKeyValuesMap[CacheIndex].Get<0>();
+				SourceName = Entry.Selector.ToString();
 
 				TArray<TSharedPtr<FJsonValue>>& SourceJsonValues = OutKeyValuesMap[CacheIndex].Get<1>();
 				SourceJsonValues = AttributeToJsonValues(Entry);
@@ -143,27 +141,13 @@ namespace PCG::IO::Json
 					ElementJsonObject->SetField(Key, MoveTemp(JsonValues[ElementIndex]));
 				}
 
-				InOutJsonObject->SetObjectField(FString::Printf(TEXT("element_%d"), ElementIndex), MoveTemp(ElementJsonObject));
+				InOutJsonObject->SetObjectField(FString::FromInt(ElementIndex), MoveTemp(ElementJsonObject));
 			}
 
 			// The map is no longer valid, so clean it.
 			InOutKeyValueMap.Empty();
 		}
-
-		template <typename CharType = TCHAR, typename PrintPolicy = TCondensedJsonPrintPolicy<CharType>>
-		FString ToJsonString(const TSharedPtr<FJsonObject>& JsonObject)
-		{
-			FString OutputString;
-			const TSharedRef<TJsonWriter<CharType, PrintPolicy>> JsonWriter = TJsonWriterFactory<CharType, PrintPolicy>::Create(&OutputString);
-			if (!FJsonSerializer::Serialize(JsonObject.ToSharedRef(), JsonWriter))
-			{
-				PCGLog::LogErrorOnGraph(LOCTEXT("JsonSerializationError", "Serialization of Json data failed."));
-				return {};
-			}
-
-			return OutputString;
-		}
-	} // namespace Private
+	}
 
 	namespace Helpers
 	{
@@ -202,50 +186,31 @@ namespace PCG::IO::Json
 			return MakeShared<FJsonValueObject>(JsonTransformObject);
 		}
 
-		void AppendHeader(const TSharedPtr<FJsonObject>& InOutJsonObject, int32 DataVersion, int32 CustomDataVersion)
-		{
-			AppendVersionToHeader(InOutJsonObject, DataVersion);
-		}
-
-		void AppendVersionToHeader(const TSharedPtr<FJsonObject>& InOutJsonObject, const int32 DataVersion)
+		void AppendHeader(const TSharedPtr<FJsonObject>& InOutJsonObject, const int32 DataVersion, const int32 CustomDataVersion)
 		{
 			if (ensure(InOutJsonObject.IsValid()))
 			{
-				InOutJsonObject->SetNumberField(Constants::DefaultDataVersionKey, DataVersion);
-			}
-		}
+				InOutJsonObject->SetNumberField(TEXT("data_version"), DataVersion);
 
-		void AppendSelectionByAttribute(TSharedPtr<FJsonObject>& InOutJsonObject, const UPCGData* InData, const TArray<FPCGAttributePropertySelector>& Selectors, const FPCGContext* InOptionalContext)
-		{
-			if (!InOutJsonObject || !InData || Selectors.IsEmpty())
-			{
-				return;
+				if (CustomDataVersion >= 0)
+				{
+					InOutJsonObject->SetNumberField(TEXT("custom_data_version"), CustomDataVersion);
+				}
 			}
-
-			const Accessor::FCache AccessorCache = Private::CreateAccessorCache(InData, Selectors, InOptionalContext);
-			JsonKeyValuesMap KeyValuesMap;
-			Private::BuildKeyValuesMap(KeyValuesMap, AccessorCache);
-			Private::AppendSelectionByAttribute(InOutJsonObject, KeyValuesMap);
 		}
 
-		void AppendSelectionByElement(TSharedPtr<FJsonObject>& InOutJsonObject, const UPCGData* InData, const TArray<FPCGAttributePropertySelector>& Selectors, const FPCGContext* InOptionalContext)
+		FString ToJsonString(const TSharedPtr<FJsonObject>& JsonObject)
 		{
-			if (!InOutJsonObject || !InData || Selectors.IsEmpty())
+			FString OutputString;
+			// @todo_pcg: Condensed by default. Add 'pretty json' as an option.
+			const TSharedRef<TJsonWriter<>> JsonWriter = TJsonWriterFactory<>::Create(&OutputString);
+			if (!FJsonSerializer::Serialize(JsonObject.ToSharedRef(), JsonWriter))
 			{
-				return;
+				PCGLog::LogErrorOnGraph(LOCTEXT("JsonSerializationError", "Serialization of Json data failed."));
+				return {};
 			}
 
-			const Accessor::FCache AccessorCache = Private::CreateAccessorCache(InData, Selectors, InOptionalContext);
-			JsonKeyValuesMap KeyValuesMap;
-			Private::BuildKeyValuesMap(KeyValuesMap, AccessorCache);
-			Private::AppendSelectionByElement(InOutJsonObject, KeyValuesMap, PCGHelpers::GetNumberOfElements(InData));
-		}
-
-		FString ToJsonString(const TSharedPtr<FJsonObject>& JsonObject, const bool bPrettyJson)
-		{
-			return bPrettyJson
-				? Private::ToJsonString<TCHAR, TPrettyJsonPrintPolicy<TCHAR>>(JsonObject)
-				: Private::ToJsonString<TCHAR, TCondensedJsonPrintPolicy<TCHAR>>(JsonObject);
+			return OutputString;
 		}
 	} // namespace Helpers
 }  // namespace PCG::IO::Json
diff --git a/Engine/Plugins/PCG/Source/PCG/Private/PCGModule.cpp b/Engine/Plugins/PCG/Source/PCG/Private/PCGModule.cpp
--- a/Engine/Plugins/PCG/Source/PCG/Private/PCGModule.cpp
+++ b/Engine/Plugins/PCG/Source/PCG/Private/PCGModule.cpp
@@ -7,16 +7,12 @@
 #include "PCGDefaultExecutionSource.h"
 #include "PCGElement.h"
 #include "PCGGraph.h"
-#include "PCGParamData.h"
 #include "Compute/PCGComputeKernel.h"
 #include "Data/PCGBasePointData.h"
 #include "Data/PCGPointData.h"
 #include "Data/PCGPolyLineData.h"
 #include "Data/PCGSplineData.h"
 #include "Data/PCGPolygon2DData.h"
-#include "Data/DataView/PCGDataViewInterface.h"
-#include "Data/DataView/PCGDataViewNativeJsonConverters.h"
-#include "Data/DataView/PCGDataViewNativePropertySelectors.h"
 
 #include "Engine/Engine.h"
 #include "Interfaces/IPluginManager.h"
@@ -182,13 +178,6 @@ void FPCGModule::StartupModule()
 	AttributeAccessorFactory.RegisterMethods<UPCGPolygon2DData>(UPCGPolygon2DData::GetPolygon2DAccessorMethods());
 	AttributeAccessorFactory.RegisterMethods<UPCGSplineData>(UPCGSplineData::GetSplineAccessorMethods());
 
-	DataViewRegistry.RegisterPropertySelector(UPCGData::StaticClass(), MakeUnique<FPCGDataViewPropertySelector>());
-	DataViewRegistry.RegisterPropertySelector(UPCGBasePointData::StaticClass(), MakeUnique<FPCGDataViewBasePointDataPropertySelector>());
-	DataViewRegistry.RegisterPropertySelector(UPCGSplineData::StaticClass(), MakeUnique<FPCGDataViewSplinePropertySelector>());
-
-	DataViewRegistry.RegisterConverter(UPCGParamData::StaticClass(), UPCGDataViewJsonConverter::StaticClass());
-	DataViewRegistry.RegisterConverter(UPCGSpatialData::StaticClass(), UPCGDataViewJsonConverter::StaticClass());
-
 #if WITH_EDITOR
 	ObjectHashFactory.RegisterObjectHashContextFactory(UPCGGraphInterface::StaticClass(), FPCGOnCreateObjectHashContext::CreateStatic(&FPCGGraphHashContext::MakeInstance));
 	ObjectHashFactory.RegisterObjectHashContextFactory(UPCGSettingsInterface::StaticClass(), FPCGOnCreateObjectHashContext::CreateStatic(&FPCGSettingsHashContext::MakeInstance));
@@ -214,13 +203,6 @@ void FPCGModule::ShutdownModule()
 
 void FPCGModule::PreExit()
 {
-	DataViewRegistry.UnregisterConverter(UPCGSpatialData::StaticClass(), UPCGDataViewJsonConverter::StaticClass());
-	DataViewRegistry.UnregisterConverter(UPCGParamData::StaticClass(), UPCGDataViewJsonConverter::StaticClass());
-
-	DataViewRegistry.UnregisterPropertySelector(UPCGSplineData::StaticClass());
-	DataViewRegistry.UnregisterPropertySelector(UPCGBasePointData::StaticClass());
-	DataViewRegistry.UnregisterPropertySelector(UPCGData::StaticClass());
-
 	// Unregistering accessor methods
 	AttributeAccessorFactory.UnregisterMethods<UPCGSplineData>();
 	AttributeAccessorFactory.UnregisterMethods<UPCGPolygon2DData>();
diff --git a/Engine/Plugins/PCG/Source/PCG/Public/Data/DataView/PCGDataView.h b/Engine/Plugins/PCG/Source/PCG/Public/Data/DataView/PCGDataView.h
deleted file mode 100644
diff --git a/Engine/Plugins/PCG/Source/PCG/Public/Data/DataView/PCGDataViewData.h b/Engine/Plugins/PCG/Source/PCG/Public/Data/DataView/PCGDataViewData.h
deleted file mode 100644
diff --git a/Engine/Plugins/PCG/Source/PCG/Public/Data/DataView/PCGDataViewInterface.h b/Engine/Plugins/PCG/Source/PCG/Public/Data/DataView/PCGDataViewInterface.h
deleted file mode 100644
diff --git a/Engine/Plugins/PCG/Source/PCG/Public/Data/DataView/PCGDataViewNativeJsonConverters.h b/Engine/Plugins/PCG/Source/PCG/Public/Data/DataView/PCGDataViewNativeJsonConverters.h
deleted file mode 100644
diff --git a/Engine/Plugins/PCG/Source/PCG/Public/Data/DataView/PCGDataViewNativePropertySelectors.h b/Engine/Plugins/PCG/Source/PCG/Public/Data/DataView/PCGDataViewNativePropertySelectors.h
deleted file mode 100644
diff --git a/Engine/Plugins/PCG/Source/PCG/Public/Data/Registry/PCGDataViewRegistry.h b/Engine/Plugins/PCG/Source/PCG/Public/Data/Registry/PCGDataViewRegistry.h
deleted file mode 100644
diff --git a/Engine/Plugins/PCG/Source/PCG/Public/Elements/IO/PCGSaveDataViewElement.h b/Engine/Plugins/PCG/Source/PCG/Public/Elements/IO/PCGSaveDataViewElement.h
deleted file mode 100644
diff --git a/Engine/Plugins/PCG/Source/PCG/Public/Elements/PCGConvertToDataView.h b/Engine/Plugins/PCG/Source/PCG/Public/Elements/PCGConvertToDataView.h
deleted file mode 100644
diff --git a/Engine/Plugins/PCG/Source/PCG/Public/Helpers/IO/PCGIOHelpers.h b/Engine/Plugins/PCG/Source/PCG/Public/Helpers/IO/PCGIOHelpers.h
--- a/Engine/Plugins/PCG/Source/PCG/Public/Helpers/IO/PCGIOHelpers.h
+++ b/Engine/Plugins/PCG/Source/PCG/Public/Helpers/IO/PCGIOHelpers.h
@@ -10,33 +10,28 @@ namespace PCG::IO
 {
 	namespace Constants
 	{
-		// Custom PCG version for attribute export
-		struct FCustomExportVersion
+		namespace Attribute
 		{
-			enum Type
+			// Custom PCG version for attribute export
+			struct FCustomExportVersion
 			{
-				InitialVersion = 0,
+				enum Type
+				{
+					InitialVersion = 0,
 
-				// New versions can be added above this line
-				VersionPlusOne,
-				LatestVersion = VersionPlusOne - 1
-			};
+					// New versions can be added above this line
+					VersionPlusOne,
+					LatestVersion = VersionPlusOne - 1
+				};
 
-			static FName GetFriendlyName() { return FName(TEXT("PCG::IO::AttributeExport")); }
-			const static FGuid GUID;
-		};
+				static FName GetFriendlyName() { return FName(TEXT("PCG::IO::AttributeExport")); }
+				const static FGuid GUID;
+			};
+		} // namespace Attribute
 
-		namespace Attribute
-		{
-			struct UE_DEPRECATED(5.8, "Use PCG::IO::Constants::FCustomExportVersion") FCustomExportVersion : Constants::FCustomExportVersion {};
-		}
+		static constexpr int32 InlineAllocationCount = 16;
 	} // namespace Constants
 
-	namespace Binary::Constants
-	{
-		static const FString Extension = TEXT(".bin");
-	}
-
 	namespace Helpers::String
 	{
 		void ToPrecisionString(const double Value, const int32 Precision, FString& OutString);
@@ -44,11 +39,6 @@ namespace PCG::IO
 
 	namespace Accessor
 	{
-		namespace Constants
-		{
-			static constexpr int32 CacheInlineAllocationCount = 16;
-		}
-
 		struct FCacheEntry
 		{
 			FPCGAttributePropertySelector Selector;
@@ -56,6 +46,6 @@ namespace PCG::IO
 			TUniquePtr<const IPCGAttributeAccessorKeys> Keys = nullptr;
 		};
 
-		using FCache = TArray<FCacheEntry, TInlineAllocator<Constants::CacheInlineAllocationCount>>;
+		using FCache = TArray<FCacheEntry, TInlineAllocator<Constants::InlineAllocationCount>>;
 	} // namespace Accessor
 } // namespace PCG::IO
diff --git a/Engine/Plugins/PCG/Source/PCG/Public/Helpers/IO/PCGJsonHelpers.h b/Engine/Plugins/PCG/Source/PCG/Public/Helpers/IO/PCGJsonHelpers.h
--- a/Engine/Plugins/PCG/Source/PCG/Public/Helpers/IO/PCGJsonHelpers.h
+++ b/Engine/Plugins/PCG/Source/PCG/Public/Helpers/IO/PCGJsonHelpers.h
@@ -2,7 +2,6 @@
 
 #pragma once
 
-#include "PCGData.h"
 #include "Metadata/PCGMetadataAttributeTraits.h"
 
 #include "Dom/JsonValue.h"
@@ -19,12 +18,6 @@ namespace PCG::IO::Json
 	template <typename T>
 	concept CFloatingPoint = std::is_floating_point_v<std::remove_cv_t<T>>;
 
-	namespace Constants
-	{
-		static const FString Extension = TEXT(".json");
-		static const FString DefaultDataVersionKey = TEXT("data_version");
-	}
-
 	namespace Helpers
 	{
 		// Converts a floating point (float, double, FVector, FRotator, etc) type to a Json value.
@@ -75,21 +68,12 @@ namespace PCG::IO::Json
 			}
 		}
 
-		UE_DEPRECATED(5.8, "AppendHeader is deprecated, please use AppendVersionToHeader instead.")
-		void AppendHeader(const TSharedPtr<FJsonObject>& InOutJsonObject, int32 DataVersion, int32 CustomDataVersion = -1);
-
 		// Append PCG Json data version header object to a Json object
-		void AppendVersionToHeader(const TSharedPtr<FJsonObject>& InOutJsonObject, int32 DataVersion);
-
-		// Append the selected properties/attributes to the provided Json object on an attribute-wise axis.
-		void AppendSelectionByAttribute(TSharedPtr<FJsonObject>& InOutJsonObject, const UPCGData* InData, const TArray<FPCGAttributePropertySelector>& Selectors, const FPCGContext* InOptionalContext = nullptr);
-
-		// Append the selected properties/attributes to the provided Json object on an element-wise axis.
-		void AppendSelectionByElement(TSharedPtr<FJsonObject>& InOutJsonObject, const UPCGData* InData, const TArray<FPCGAttributePropertySelector>& Selectors, const FPCGContext* InOptionalContext = nullptr);
+		void AppendHeader(const TSharedPtr<FJsonObject>& InOutJsonObject, int32 DataVersion, int32 CustomDataVersion = -1);
 
 		// Helper to convert a Json object into a string
-		FString ToJsonString(const TSharedPtr<FJsonObject>& JsonObject, bool bPrettyJson = false);
+		FString ToJsonString(const TSharedPtr<FJsonObject>& JsonObject);
 
 		// @todo_pcg: Add a fast track for AppendAllAttributes
-	} // namespace Helpers
-} // namespace PCG::IO::Json
+	}
+}
diff --git a/Engine/Plugins/PCG/Source/PCG/Public/PCGContext.h b/Engine/Plugins/PCG/Source/PCG/Public/PCGContext.h
--- a/Engine/Plugins/PCG/Source/PCG/Public/PCGContext.h
+++ b/Engine/Plugins/PCG/Source/PCG/Public/PCGContext.h
@@ -186,14 +186,6 @@ PRAGMA_ENABLE_DEPRECATION_WARNINGS
 		return SettingsWithOverride ? Cast<SettingsType>(SettingsWithOverride) : GetOriginalSettings<SettingsType>();
 	}
 
-	// Return the mutable settings casted in the wanted type.
-	// If there is no override, it returns null.
-	template<typename SettingsType>
-	SettingsType* GetMutableInputSettings()
-	{
-		return SettingsWithOverride ? Cast<SettingsType>(SettingsWithOverride) : nullptr;
-	}
-
 	UE_API FPCGTaskId GetGraphExecutionTaskId() const;
 
 	UE_API FString GetTaskName() const;
diff --git a/Engine/Plugins/PCG/Source/PCG/Public/PCGModule.h b/Engine/Plugins/PCG/Source/PCG/Public/PCGModule.h
--- a/Engine/Plugins/PCG/Source/PCG/Public/PCGModule.h
+++ b/Engine/Plugins/PCG/Source/PCG/Public/PCGModule.h
@@ -5,7 +5,6 @@
 #include "PCGDataVisualizationRegistry.h"
 #include "PCGGraphExecutionRegistry.h"
 #include "Data/Registry/PCGDataTypeRegistry.h"
-#include "Data/Registry/PCGDataViewRegistry.h"
 #include "Data/Registry/PCGGetDataFunctionRegistry.h"
 #include "Hash/PCGObjectHash.h"
 #include "Metadata/Accessors/PCGAttributeAccessorFactory.h"
@@ -60,9 +59,6 @@ public:
 	static const FPCGGraphExecutionRegistry& GetConstGraphExecutionRegistry() { return GetPCGModuleChecked().GraphExecutionRegistry; }
 	static FPCGGraphExecutionRegistry& GetGraphExecutionRegistry() { return GetPCGModuleChecked().GraphExecutionRegistry; }
 
-	PCG_API static const FPCGDataViewRegistry& GetConstPCGDataViewRegistry() { return GetPCGModuleChecked().DataViewRegistry; }
-	PCG_API static FPCGDataViewRegistry& GetMutablePCGDataViewRegistry() { return GetPCGModuleChecked().DataViewRegistry; }
-
 #if WITH_EDITOR
 	static const FPCGObjectHashFactory& GetConstObjectHashFactory() { return GetPCGModuleChecked().ObjectHashFactory; }
 	static FPCGObjectHashFactory& GetMutableObjectHashFactory() { return GetPCGModuleChecked().ObjectHashFactory; }
@@ -88,7 +84,6 @@ private:
 	FPCGAttributeAccessorFactory AttributeAccessorFactory;
 	FPCGDataTypeRegistry DataTypeRegistry;
 	FPCGGraphExecutionRegistry GraphExecutionRegistry;
-	FPCGDataViewRegistry DataViewRegistry;
 
 #if WITH_EDITOR
 	FPCGObjectHashFactory ObjectHashFactory;
@@ -116,4 +111,4 @@ private:
 
 	static FPCGGraphChangedDelegate OnGraphChangedDelegate;
 #endif // WITH_EDITOR
-};
+};
\ No newline at end of file
diff --git a/Engine/Plugins/PCG/Source/PCGEditor/Private/DataVisualizations/PCGDataViewVisualization.cpp b/Engine/Plugins/PCG/Source/PCGEditor/Private/DataVisualizations/PCGDataViewVisualization.cpp
deleted file mode 100644
diff --git a/Engine/Plugins/PCG/Source/PCGEditor/Private/PCGEditorModule.cpp b/Engine/Plugins/PCG/Source/PCGEditor/Private/PCGEditorModule.cpp
--- a/Engine/Plugins/PCG/Source/PCGEditor/Private/PCGEditorModule.cpp
+++ b/Engine/Plugins/PCG/Source/PCGEditor/Private/PCGEditorModule.cpp
@@ -70,8 +70,6 @@
 #include "PropertyEditorModule.h"
 #include "ScopedTransaction.h"
 #include "ToolMenus.h"
-#include "Data/DataView/PCGDataViewData.h"
-#include "DataVisualizations/PCGDataViewVisualization.h"
 #include "Details/PCGRaycastFilterRuleCollectionDetails.h"
 #include "Editor/SceneOutliner/Public/ISceneOutliner.h"
 #include "EditorMode/PCGEdMode.h"
@@ -863,7 +861,6 @@ void FPCGEditorModule::RegisterPCGDataVisualizations()
 	DataVisRegistry.InternalRegistry.Add(UPCGBaseTextureData::StaticClass(), MakeUnique<const FPCGBaseTextureDataVisualization>());
 	DataVisRegistry.InternalRegistry.Add(UPCGPolygon2DData::StaticClass(), MakeUnique<const FPCGPolygon2DDataVisualization>());
 	DataVisRegistry.InternalRegistry.Add(UPCGPolygon2DInteriorSurfaceData::StaticClass(), MakeUnique<const FPCGPolygon2DDataVisualization>());
-	DataVisRegistry.InternalRegistry.Add(UPCGDataViewData::StaticClass(), MakeUnique<const FPCGDataViewVisualization>());
 }
 
 void FPCGEditorModule::UnregisterPCGDataVisualizations()
@@ -903,7 +900,6 @@ void FPCGEditorModule::RegisterPinColorAndIcons()
 	InRegistry.RegisterPinColorFunction(FPCGDataTypeInfoDynamicMesh::AsId(), [](const FPCGDataTypeIdentifier&) { return GetDefault<UPCGEditorSettings>()->DynamicMeshPinColor; });
 	InRegistry.RegisterPinColorFunction(FPCGDataTypeInfoSpatial::AsId(), [](const FPCGDataTypeIdentifier&) { return GetDefault<UPCGEditorSettings>()->SpatialDataPinColor; });
 	InRegistry.RegisterPinColorFunction(FPCGDataTypeInfoStaticMeshResource::AsId(), [](const FPCGDataTypeIdentifier&) { return GetDefault<UPCGEditorSettings>()->StaticMeshResourcePinColor; });
-	InRegistry.RegisterPinColorFunction(FPCGDataTypeInfoDataView::AsId(), [](const FPCGDataTypeIdentifier&) { return GetDefault<UPCGEditorSettings>()->DataViewPinColor; });
 	InRegistry.RegisterPinColorFunction(FPCGDataTypeInfoOther::AsId(), [](const FPCGDataTypeIdentifier&) { return GetDefault<UPCGEditorSettings>()->UnknownDataPinColor; });
 	InRegistry.RegisterPinColorFunction(FPCGDataTypeInfo::AsId(), [](const FPCGDataTypeIdentifier&) { return GetDefault<UPCGEditorSettings>()->DefaultPinColor; });
 
@@ -989,7 +985,6 @@ void FPCGEditorModule::UnregisterPinColorAndIcons()
 	InRegistry.UnregisterPinColorFunction(FPCGDataTypeInfoConcrete::AsId());
 	InRegistry.UnregisterPinColorFunction(FPCGDataTypeInfoPoint::AsId());
 	InRegistry.UnregisterPinColorFunction(FPCGDataTypeInfoParam::AsId());
-	InRegistry.UnregisterPinColorFunction(FPCGDataTypeInfoDataView::AsId());
 
 	// PinIcons. Just have catch all as most of it will be the same icons.
 	InRegistry.UnregisterPinIconsFunction(FPCGDataTypeInfo::AsId());
diff --git a/Engine/Plugins/PCG/Source/PCGEditor/Private/PCGEditorSettings.cpp b/Engine/Plugins/PCG/Source/PCGEditor/Private/PCGEditorSettings.cpp
--- a/Engine/Plugins/PCG/Source/PCGEditor/Private/PCGEditorSettings.cpp
+++ b/Engine/Plugins/PCG/Source/PCGEditor/Private/PCGEditorSettings.cpp
@@ -48,7 +48,6 @@ UPCGEditorSettings::UPCGEditorSettings(const FObjectInitializer& ObjectInitializ
 	LandscapeDataPinColor = FLinearColor(0.66f, 0.66f, 0.07f);
 	BaseTextureDataPinColor = FLinearColor(0.81f, 0.28f, 0.22f);
 	TextureDataPinColor = FLinearColor(0.79f, 0.08f, 0.01f);
-	DataViewPinColor = FLinearColor(0.69f, 0.92f, 0.23f);
 	RenderTargetDataPinColor = FLinearColor(0.8f, 0.18f, 0.12f);
 	VirtualTextureDataPinColor = FLinearColor(0.8f, 0.48f, 0.22f);
 	VolumeDataPinColor = FLinearColor(0.79f, 0.06f, 0.5f);
diff --git a/Engine/Plugins/PCG/Source/PCGEditor/Public/DataVisualizations/PCGDataViewVisualization.h b/Engine/Plugins/PCG/Source/PCGEditor/Public/DataVisualizations/PCGDataViewVisualization.h
deleted file mode 100644
diff --git a/Engine/Plugins/PCG/Source/PCGEditor/Public/PCGEditorSettings.h b/Engine/Plugins/PCG/Source/PCGEditor/Public/PCGEditorSettings.h
--- a/Engine/Plugins/PCG/Source/PCGEditor/Public/PCGEditorSettings.h
+++ b/Engine/Plugins/PCG/Source/PCGEditor/Public/PCGEditorSettings.h
@@ -201,10 +201,6 @@ public:
 	UPROPERTY(EditAnywhere, config, Category = Node, meta = (HideAlphaChannel))
 	FLinearColor TextureDataPinColor;
 
-	/** Color used for data pins of type Data View */
-	UPROPERTY(EditAnywhere, config, Category = Node, meta = (HideAlphaChannel))
-	FLinearColor DataViewPinColor;
-
 	/** Color used for data pins of type Render Target */
 	UPROPERTY(EditAnywhere, config, Category = Node, meta = (HideAlphaChannel))
 	FLinearColor RenderTargetDataPinColor;
