[PCG] Add additional title information to Extract Attribute node

Subtitle now appears as `$Transform.Rotation -> MyRotation`, which makes it easier to read graphs at a glance when doing operations on extracted attributes.

#rnx

diff --git a/Engine/Plugins/PCG/Source/PCG/Private/Elements/Metadata/PCGExtractAttribute.cpp b/Engine/Plugins/PCG/Source/PCG/Private/Elements/Metadata/PCGExtractAttribute.cpp
--- a/Engine/Plugins/PCG/Source/PCG/Private/Elements/Metadata/PCGExtractAttribute.cpp
+++ b/Engine/Plugins/PCG/Source/PCG/Private/Elements/Metadata/PCGExtractAttribute.cpp
@@ -111,6 +111,11 @@ FText UPCGExtractAttributeSettings::GetNodeTooltipText() const
 }
 #endif // WITH_EDITOR
 
+FString UPCGExtractAttributeSettings::GetAdditionalTitleInformation() const
+{
+	return FString::Printf(TEXT("%s -> %s"), *InputSource.ToString(), *OutputAttributeName.ToString());
+}
+
 FPCGElementPtr UPCGExtractAttributeSettings::CreateElement() const
 {
 	return MakeShared<FPCGExtractAttributeElement>();
diff --git a/Engine/Plugins/PCG/Source/PCG/Public/Elements/Metadata/PCGExtractAttribute.h b/Engine/Plugins/PCG/Source/PCG/Public/Elements/Metadata/PCGExtractAttribute.h
--- a/Engine/Plugins/PCG/Source/PCG/Public/Elements/Metadata/PCGExtractAttribute.h
+++ b/Engine/Plugins/PCG/Source/PCG/Public/Elements/Metadata/PCGExtractAttribute.h
@@ -66,9 +66,10 @@ public:
 #if WITH_EDITOR
 	virtual FName GetDefaultNodeName() const override;
 	virtual FText GetDefaultNodeTitle() const override;
-    virtual FText GetNodeTooltipText() const override;
+	virtual FText GetNodeTooltipText() const override;
 	virtual EPCGSettingsType GetType() const override { return EPCGSettingsType::Metadata; }
 #endif
+	virtual FString GetAdditionalTitleInformation() const override;
 
 protected:
 	virtual TArray<FPCGPinProperties> InputPinProperties() const override;
@@ -91,7 +92,7 @@ class FPCGExtractAttributeElement : public IPCGElement
 {
 protected:
 	virtual bool ExecuteInternal(FPCGContext* InContext) const override;
-    virtual bool SupportsBasePointDataInputs(FPCGContext* InContext) const override { return true; }
-    virtual EPCGElementExecutionLoopMode ExecutionLoopMode(const UPCGSettings* Settings) const override { return EPCGElementExecutionLoopMode::SinglePrimaryPin; }
+	virtual bool SupportsBasePointDataInputs(FPCGContext* InContext) const override { return true; }
+	virtual EPCGElementExecutionLoopMode ExecutionLoopMode(const UPCGSettings* Settings) const override { return EPCGElementExecutionLoopMode::SinglePrimaryPin; }
 };
 
