[PCG] Add support for overrides on array properties

* Arrays are now extracted as overrides for the default settings, for Blueprints and for Graph Parameters on subgraphs
* Added a temporary accessor, in private, that works specifically for arrays in a very limited way until we have a pass on accessors with array metadata support
* Also supports array as object overrides with the limitation of 1 element per array (like on the Spawn Spline Mesh node, we can override Override Materials with an attribute, which means a single value)
* Rationale for single value support it that it is going to be supported with multi values with array metadata support
* Apply on Object also supports 1 attribute set per elements to apply the overrides, with will allow for multi values arrays with our current implementation.

diff --git a/Engine/Plugins/PCG/Source/PCG/Internal/Metadata/Accessors/PCGAttributeAccessorHelpersInternal.h b/Engine/Plugins/PCG/Source/PCG/Internal/Metadata/Accessors/PCGAttributeAccessorHelpersInternal.h
new file mode 100644
--- /dev/null
+++ b/Engine/Plugins/PCG/Source/PCG/Internal/Metadata/Accessors/PCGAttributeAccessorHelpersInternal.h
@@ -0,0 +1,151 @@
+// Copyright Epic Games, Inc. All Rights Reserved.
+
+#pragma once
+
+#include "Helpers/PCGPropertyHelpers.h"
+#include "Metadata/Accessors/PCGPropertyAccessor.h"
+
+#include "UObject/UnrealType.h"
+
+namespace PCGAttributeAccessorHelpers
+{
+namespace Internal
+{
+	// Empty signature to passthrough types to functors.
+	template <typename T>
+	struct Signature
+	{
+		using Type = T;
+	};
+
+	template <typename Func>
+	decltype(auto) DispatchPropertyTypes(const FProperty* InProperty, Func&& Functor)
+	{
+		// Use the FPCGPropertyPathAccessor (on soft object path) as dummy type to get the functor return type
+		// because this accessor takes a generic FProperty, while others take a more specialized type (like FNumericProperty).
+		using ReturnType = decltype(Functor(Signature<FPCGPropertyPathAccessor<FSoftObjectPath>>{}, InProperty));
+
+		if (!InProperty)
+		{
+			return ReturnType{};
+		}
+
+		if (const FNumericProperty* NumericProperty = CastField<FNumericProperty>(InProperty))
+		{
+			if (NumericProperty->IsFloatingPoint())
+			{
+				// As floating properties are mostly all double in UE, convert float to double attributes.
+				return Functor(Signature<FPCGNumericPropertyAccessor<double>>{}, NumericProperty);
+			}
+			else if (NumericProperty->IsInteger())
+			{
+				// But for int32/int64 we can distinguish between the two. Everything of size 32 or less is will be an int32, 64bits integers will be int64.
+				if (NumericProperty->IsA<FInt64Property>() || NumericProperty->IsA<FUInt64Property>())
+				{
+					return Functor(Signature<FPCGNumericPropertyAccessor<int64>>{}, NumericProperty);
+				}
+				else
+				{
+					return Functor(Signature<FPCGNumericPropertyAccessor<int32>>{}, NumericProperty);
+				}
+			}
+		}
+		else if (const FBoolProperty* BoolProperty = CastField<FBoolProperty>(InProperty))
+		{
+			return Functor(Signature<FPCGPropertyAccessor<bool, FBoolProperty>>{}, BoolProperty);
+		}
+		else if (const FStrProperty* StringProperty = CastField<FStrProperty>(InProperty))
+		{
+			return Functor(Signature<FPCGPropertyAccessor<FString, FStrProperty>>{}, StringProperty);
+		}
+		else if (const FNameProperty* NameProperty = CastField<FNameProperty>(InProperty))
+		{
+			return Functor(Signature<FPCGPropertyAccessor<FName, FNameProperty>>{}, NameProperty);
+		}
+		else if (const FEnumProperty* EnumProperty = CastField<FEnumProperty>(InProperty))
+		{
+			return Functor(Signature<FPCGEnumPropertyAccessor>{}, EnumProperty);
+		}
+		else if (const FSoftClassProperty* SoftClassProperty = CastField<FSoftClassProperty>(InProperty))
+		{
+			return Functor(Signature<FPCGPropertySoftClassPathAccessor>{}, SoftClassProperty);
+		}
+		else if (const FSoftObjectProperty* SoftObjectProperty = CastField<FSoftObjectProperty>(InProperty))
+		{
+			return Functor(Signature<FPCGPropertySoftObjectPathAccessor>{}, SoftObjectProperty);
+		}
+		else if (const FClassProperty* ClassProperty = CastField<FClassProperty>(InProperty))
+		{
+			return Functor(Signature<FPCGPropertyObjectPtrAccessor<FClassProperty>>{}, ClassProperty);
+		}
+		else if (const FObjectProperty* ObjectProperty = CastField<FObjectProperty>(InProperty))
+		{
+			return Functor(Signature<FPCGPropertyObjectPtrAccessor<FObjectProperty>>{}, ObjectProperty);
+		}
+		else if (const FStructProperty* StructProperty = CastField<FStructProperty>(InProperty))
+		{
+			if (StructProperty->Struct == TBaseStructure<FVector>::Get())
+			{
+				return Functor(Signature<FPCGPropertyStructAccessor<FVector>>{}, StructProperty);
+			}
+			else if (StructProperty->Struct == TBaseStructure<FVector4>::Get())
+			{
+				return Functor(Signature<FPCGPropertyStructAccessor<FVector4>>{}, StructProperty);
+			}
+			else if (StructProperty->Struct == TBaseStructure<FQuat>::Get())
+			{
+				return Functor(Signature<FPCGPropertyStructAccessor<FQuat>>{}, StructProperty);
+			}
+			else if (StructProperty->Struct == TBaseStructure<FTransform>::Get())
+			{
+				return Functor(Signature<FPCGPropertyStructAccessor<FTransform>>{}, StructProperty);
+			}
+			else if (StructProperty->Struct == TBaseStructure<FRotator>::Get())
+			{
+				return Functor(Signature<FPCGPropertyStructAccessor<FRotator>>{}, StructProperty);
+			}
+			else if (StructProperty->Struct == TBaseStructure<FVector2D>::Get())
+			{
+				return Functor(Signature<FPCGPropertyStructAccessor<FVector2D>>{}, StructProperty);
+			}
+			else if (StructProperty->Struct == TBaseStructure<FSoftObjectPath>::Get())
+			{
+				return Functor(Signature<FPCGPropertyStructAccessor<FSoftObjectPath>>{}, StructProperty);
+			}
+			else if (StructProperty->Struct == TBaseStructure<FSoftClassPath>::Get())
+			{
+				return Functor(Signature<FPCGPropertyStructAccessor<FSoftClassPath>>{}, StructProperty);
+			}
+			else if (StructProperty->Struct == TBaseStructure<FLinearColor>::Get())
+			{
+				return Functor(Signature<FPCGLinearColorAccessor>{}, StructProperty);
+			}
+			else if (StructProperty->Struct == TBaseStructure<FColor>::Get())
+			{
+				return Functor(Signature<FPCGColorAccessor>{}, StructProperty);
+			}
+			else if (StructProperty->Struct && StructProperty->Struct->IsChildOf(FPCGAttributePropertySelector::StaticStruct()))
+			{
+				return Functor(Signature<FPCGAttributePropertySelectorAccessor>{}, StructProperty);
+			}
+		}
+
+		return ReturnType{};
+	}
+
+	template <typename Func>
+	decltype(auto) DispatchPropertyTypes(const FName InPropertyName, const UStruct* InStruct, Func&& Functor)
+	{
+		if (InStruct)
+		{
+			if (const FProperty* Property = PCGPropertyHelpers::FindPropertyByName(InStruct, InPropertyName))
+			{
+				return DispatchPropertyTypes(Property, std::forward<Func>(Functor));
+			}
+		}
+
+		using ReturnType = decltype(Functor(static_cast<FPCGPropertyPathAccessor<FSoftObjectPath>*>(nullptr), static_cast<const FProperty*>(nullptr)));
+		return ReturnType{};
+	}
+}// namespace Internal
+} // namespace PCGAttributeAccessorHelpers
diff --git a/Engine/Plugins/PCG/Source/PCG/Internal/Metadata/Accessors/PCGContainerAccessor.cpp b/Engine/Plugins/PCG/Source/PCG/Internal/Metadata/Accessors/PCGContainerAccessor.cpp
new file mode 100644
--- /dev/null
+++ b/Engine/Plugins/PCG/Source/PCG/Internal/Metadata/Accessors/PCGContainerAccessor.cpp
@@ -0,0 +1,63 @@
+﻿// Copyright Epic Games, Inc. All Rights Reserved.
+
+#include "Metadata/Accessors/PCGContainerAccessor.h"
+
+#include "Metadata/Accessors/PCGAttributeAccessorHelpersInternal.h"
+
+TUniquePtr<IPCGAttributeAccessor> PCGContainerAccessorHelpers::MakeContainerAccessor(TConstArrayView<const FProperty*> PropertyChain)
+{
+	// Find the container property in the chain. Make sure that there is just a single one.
+	int32 ContainerPropertyIndex = INDEX_NONE;
+	int32 Index = 0;
+	const FArrayProperty* ArrayProperty = nullptr;
+	for (const FProperty* Property : PropertyChain)
+	{
+		check(Property);
+		if (Property->IsA<FArrayProperty>())
+		{
+			if (ContainerPropertyIndex != INDEX_NONE)
+			{
+				UE_LOG(LogPCG, Error, TEXT("[PCGContainerAccessorHelpers::MakeContainerAccessor] Tried to create an accessor on a property chain with multiple containers. Unsupported."))
+				return nullptr;
+			}
+				
+			ContainerPropertyIndex = Index;
+			ArrayProperty = CastField<FArrayProperty>(Property);
+		}
+			
+		++Index;
+	}
+
+	if (ContainerPropertyIndex == INDEX_NONE || ContainerPropertyIndex >= PropertyChain.Num() - 1)
+	{
+		return nullptr;
+	}
+
+	check(ArrayProperty);
+
+	TArray<const FProperty*> ParentProperties(PropertyChain.GetData(), ContainerPropertyIndex + 1);
+	TArray<const FProperty*> ChildProperties(PropertyChain.GetData() + ContainerPropertyIndex + 1, PropertyChain.Num() - ContainerPropertyIndex - 1);
+
+	using namespace PCGAttributeAccessorHelpers::Internal;
+
+	return DispatchPropertyTypes(ChildProperties.Last(),
+		[&ParentProperties, &ChildProperties, ArrayProperty]<typename TypedAccessor, typename PropertyType>(Signature<TypedAccessor> SignatureDummy, const PropertyType* TypedProperty) -> TUniquePtr<IPCGAttributeAccessor>
+			{
+				using UnderlyingType = typename TypedAccessor::Type;
+
+				TUniquePtr<IPCGAttributeAccessor> UnderlyingAccessor = MakeUnique<TypedAccessor>(TypedProperty, MoveTemp(ChildProperties));
+				if (!UnderlyingAccessor)
+				{
+					return nullptr;
+				}
+				else if (ArrayProperty)
+				{
+					return MakeUnique<IPCGArrayWrapperAccessor<UnderlyingType>>(ArrayProperty, MoveTemp(UnderlyingAccessor), MoveTemp(ParentProperties));
+				}
+				else
+				{
+					checkNoEntry();
+					return nullptr;
+				}
+			});
+}
\ No newline at end of file
diff --git a/Engine/Plugins/PCG/Source/PCG/Internal/Metadata/Accessors/PCGContainerAccessor.h b/Engine/Plugins/PCG/Source/PCG/Internal/Metadata/Accessors/PCGContainerAccessor.h
new file mode 100644
--- /dev/null
+++ b/Engine/Plugins/PCG/Source/PCG/Internal/Metadata/Accessors/PCGContainerAccessor.h
@@ -0,0 +1,86 @@
+// Copyright Epic Games, Inc. All Rights Reserved.
+
+#pragma once
+
+#include "PCGCommon.h"
+#include "PCGModule.h"
+#include "Metadata/Accessors/PCGAttributeAccessorKeys.h"
+#include "Metadata/Accessors/PCGPropertyAccessor.h"
+
+#include "UObject/UnrealType.h"
+
+/**
+ * DISCLAIMER:
+ * Container accessors are a first try and temporary to support containers for accessors.
+ * Because of how the Accessor API is working, it will only works to set values in single containers, with no support for multi-containers.
+ * It is meant to set the full container in one go, no in-place editing or appending, meaning that Index > 0 will be rejected.
+ */
+template <typename T>
+class IPCGArrayWrapperAccessor : public IPCGAttributeAccessorT<IPCGArrayWrapperAccessor<T>>, IPCGPropertyChain
+{
+public:
+	using Type = T;
+	using Super = IPCGAttributeAccessorT<IPCGArrayWrapperAccessor<T>>;
+
+	// ExtraProperties are properties to arrive to the ArrayProperty. UnderlyingAccessor is for writing into an element of the array.
+	IPCGArrayWrapperAccessor(const FArrayProperty* InProperty, TUniquePtr<IPCGAttributeAccessor>&& InUnderlyingAccessor, TArray<const FProperty*>&& ExtraProperties = {})
+		: Super(/*bInReadOnly=*/ false)
+		, IPCGPropertyChain(InProperty, MoveTemp(ExtraProperties))
+		, Property(InProperty)
+		, UnderlyingAccessor(MoveTemp(InUnderlyingAccessor))
+	{
+		check(Property && UnderlyingAccessor.IsValid());
+	}
+
+	bool GetRangeImpl(TArrayView<T> OutValues, int32 Index, const IPCGAttributeAccessorKeys& Keys) const
+	{
+		// Unsupported
+		return false;
+	}
+
+	bool SetRangeImpl(TArrayView<const T> InValues, int32 Index, IPCGAttributeAccessorKeys& Keys, EPCGAttributeAccessorFlags Flags)
+	{
+		if (Index > 0)
+		{
+			UE_LOG(LogPCG, Error, TEXT("Trying to use the Array Wrapper Accessor with an index greater than 0, unsupported."));
+			return false;
+		}
+		
+		const int32 NumValues = InValues.Num();
+
+		// Since we "copy" InValues into a single array, we can only work on a single key.
+		TArray<void*, TInlineAllocator<256>> ContainerKeys;
+		PCGPropertyAccessor::GetContainerKeys(0, 1, Keys, ContainerKeys);
+		if (ContainerKeys.IsEmpty())
+		{
+			return false;
+		}
+
+		// Update the addresses
+		PCGPropertyAccessor::AddressOffset<void*>(GetPropertyChain(), ContainerKeys);
+		
+		FScriptArrayHelper ArrayHelper(Property, ContainerKeys[0]);
+		ArrayHelper.Resize(NumValues);
+
+		// Create subkeys for the forwarding call
+		TArray<void*> Addresses;
+		Addresses.Reserve(InValues.Num());
+		for (int32 i = 0; i < InValues.Num(); ++i)
+		{
+			Addresses.Emplace(ArrayHelper.GetRawPtr(i));
+		}
+
+		FPCGAttributeAccessorKeysGenericPtrs ArrayKeys{Addresses};
+
+		return UnderlyingAccessor->SetRange<T>(InValues, /*Index=*/0, ArrayKeys, Flags);
+	}
+
+private:
+	const FArrayProperty* Property = nullptr;
+	TUniquePtr<IPCGAttributeAccessor> UnderlyingAccessor;
+};
+
+namespace PCGContainerAccessorHelpers
+{
+	TUniquePtr<IPCGAttributeAccessor> MakeContainerAccessor(TConstArrayView<const FProperty*> PropertyChain);
+}
\ No newline at end of file
diff --git a/Engine/Plugins/PCG/Source/PCG/Private/Elements/PCGApplyOnActor.cpp b/Engine/Plugins/PCG/Source/PCG/Private/Elements/PCGApplyOnActor.cpp
--- a/Engine/Plugins/PCG/Source/PCG/Private/Elements/PCGApplyOnActor.cpp
+++ b/Engine/Plugins/PCG/Source/PCG/Private/Elements/PCGApplyOnActor.cpp
@@ -75,7 +75,7 @@ PRAGMA_ENABLE_DEPRECATION_WARNINGS
 TArray<FPCGPinProperties> UPCGApplyOnActorSettings::InputPinProperties() const
 {
 	TArray<FPCGPinProperties> PinProperties;
-	PinProperties.Emplace(PCGPinConstants::DefaultInputLabel, EPCGDataType::Any, /*bAllowMultipleConnections=*/true, /*bAllowMultipleData=*/true);
+	PinProperties.Emplace_GetRef(PCGPinConstants::DefaultInputLabel, EPCGDataType::Any, /*bAllowMultipleConnections=*/true, /*bAllowMultipleData=*/true).SetRequiredPin();
 	PinProperties.Emplace(PCGApplyOnActorConstants::ObjectPropertyOverridesLabel, EPCGDataType::Any, /*bAllowMultipleConnections=*/true, /*bAllowMultipleData=*/true, PCGApplyOnActorConstants::ObjectPropertyOverridesTooltip);
 	return PinProperties;
 }
@@ -128,19 +128,71 @@ bool FPCGApplyOnObjectElement::ExecuteInternal(FPCGContext* Context) const
 	const bool bPropagateObjectChangeEvent = false;
 #endif // WITH_EDITOR
 
+	const TArray<FPCGTaggedData> InputTaggedData = Context->InputData.GetInputsByPin(PCGPinConstants::DefaultInputLabel);
+	const TArray<FPCGTaggedData> OverridesTaggedData = Context->InputData.GetInputsByPin(PCGApplyOnActorConstants::ObjectPropertyOverridesLabel);
+
+	// Inputs are always forwarded.
+	Context->OutputData.TaggedData = InputTaggedData;
+
+	if (Settings->bUseOneOverrideDataPerElement)
+	{
+		if (OverridesTaggedData.Num() > 1 && OverridesTaggedData.Num() != ThisContext->PathsToObjectsAndDataIndex.Num())
+		{
+			PCGLog::LogErrorOnGraph(FText::Format(LOCTEXT("InconsistentDataCount", "When having one data set per element, it is expected to have the same number of override data than elements, we expected {0} and got {1} data."), ThisContext->PathsToObjectsAndDataIndex.Num(), InputTaggedData.Num()), Context);
+			return true;
+		}
+	}
+	else
+	{
+		if (OverridesTaggedData.Num() > 1 && OverridesTaggedData.Num() != InputTaggedData.Num())
+		{
+			PCGLog::InputOutput::LogInvalidCardinalityError(PCGPinConstants::DefaultInputLabel, PCGApplyOnActorConstants::ObjectPropertyOverridesLabel, Context);
+			return true;
+		}
+	}
+
 	int CurrentPathIndex = 0;
 	while (CurrentPathIndex < ThisContext->PathsToObjectsAndDataIndex.Num())
 	{
 		TargetObjectsAndIndices.Reset();
 
-		int InputIndex = ThisContext->PathsToObjectsAndDataIndex[CurrentPathIndex].Get<1>();
+		int32 InputIndex = ThisContext->PathsToObjectsAndDataIndex[CurrentPathIndex].Get<1>();
+		const int32 StartingIndex = CurrentPathIndex;
 		while (CurrentPathIndex < ThisContext->PathsToObjectsAndDataIndex.Num() && InputIndex == ThisContext->PathsToObjectsAndDataIndex[CurrentPathIndex].Get<1>())
 		{
 			TargetObjectsAndIndices.Emplace(ThisContext->PathsToObjectsAndDataIndex[CurrentPathIndex].Get<0>().ResolveObject(), ThisContext->PathsToObjectsAndDataIndex[CurrentPathIndex].Get<2>());
 			++CurrentPathIndex;
 		}
 
-		PCGObjectPropertyOverrideHelpers::ApplyOverrides(Settings->PropertyOverrideDescriptions, TargetObjectsAndIndices, PCGApplyOnActorConstants::ObjectPropertyOverridesLabel, InputIndex, Context, bPropagateObjectChangeEvent);
+		if (!OverridesTaggedData.IsEmpty())
+		{
+			TArray<const UPCGData*, TInlineAllocator<16>> OverridesData;
+			if (Settings->bUseOneOverrideDataPerElement)
+			{
+				const int32 NumElements = CurrentPathIndex - StartingIndex;
+				OverridesData.Reserve(NumElements);
+				for (int32 i = 0; i < NumElements; ++i)
+				{
+					OverridesData.Add(OverridesTaggedData[(StartingIndex + i) % OverridesTaggedData.Num()].Data);
+				}
+			}
+			else
+			{
+				OverridesData.Add(OverridesTaggedData[InputIndex % OverridesTaggedData.Num()].Data);
+			}
+		
+			PCGObjectPropertyOverrideHelpers::FApplyOverrideParams Params =
+			{
+				.ObjectPropertyOverrideDescriptions = MakeConstArrayView(Settings->PropertyOverrideDescriptions),
+				.TargetObjectToOverrideIndices = MakeConstArrayView(TargetObjectsAndIndices),
+				.OverridesData = MakeConstArrayView(OverridesData),
+				.OptionalContext = Context,
+				.bPropagateEditChangeEvent = bPropagateObjectChangeEvent,
+				.bOneDataSetPerElement = Settings->bUseOneOverrideDataPerElement,
+			};
+
+			PCGObjectPropertyOverrideHelpers::ApplyOverrides(Params);
+		}
 
 		for (const TPair<UObject*, int32>& TargetObjectAndIndex : TargetObjectsAndIndices)
 		{
@@ -157,9 +209,6 @@ bool FPCGApplyOnObjectElement::ExecuteInternal(FPCGContext* Context) const
 		}
 	}
 
-	// Since we have only one output pin, there's no need to update the pins here.
-	Context->OutputData.TaggedData = Context->InputData.GetInputsByPin(PCGPinConstants::DefaultInputLabel);
-
 	return true;
 }
 
diff --git a/Engine/Plugins/PCG/Source/PCG/Private/Elements/PCGExecuteBlueprint.cpp b/Engine/Plugins/PCG/Source/PCG/Private/Elements/PCGExecuteBlueprint.cpp
--- a/Engine/Plugins/PCG/Source/PCG/Private/Elements/PCGExecuteBlueprint.cpp
+++ b/Engine/Plugins/PCG/Source/PCG/Private/Elements/PCGExecuteBlueprint.cpp
@@ -425,6 +425,7 @@ TArray<FPCGSettingsOverridableParam> UPCGBlueprintSettings::GatherOverridablePar
 	{
 		PCGSettingsHelpers::FPCGGetAllOverridableParamsConfig Config;
 		Config.bExcludeSuperProperties = true;
+		Config.bExtractArrays = true;
 		Config.ExcludePropertyFlags = CPF_DisableEditOnInstance | CPF_EditConst | CPF_BlueprintReadOnly;
 		OverridableParams.Append(PCGSettingsHelpers::GetAllOverridableParams(BPClass, Config));
 	}
diff --git a/Engine/Plugins/PCG/Source/PCG/Private/Helpers/PCGPropertyHelpers.cpp b/Engine/Plugins/PCG/Source/PCG/Private/Helpers/PCGPropertyHelpers.cpp
--- a/Engine/Plugins/PCG/Source/PCG/Private/Helpers/PCGPropertyHelpers.cpp
+++ b/Engine/Plugins/PCG/Source/PCG/Private/Helpers/PCGPropertyHelpers.cpp
@@ -65,18 +65,7 @@ namespace PCGPropertyHelpers
 	{
 		check(CurrentClass);
 
-		const FProperty* Property = nullptr;
-		// Try to get the property. If it is coming from a user struct, we need to iterate on all properties because the property name is mangled.
-		// There is also a difference between runtime and editor name, as property name in editor can contain invalid characters (like spaces).
-		// So if the current name is invalid, also try with sanitized name.
-		if (const UUserDefinedStruct* UserDefinedStruct = Cast<UUserDefinedStruct>(CurrentClass))
-		{
-			Property = FindPropertyInUserDefinedStruct(UserDefinedStruct, CurrentName);
-		}
-		else
-		{
-			Property = FindFProperty<FProperty>(CurrentClass, CurrentName);
-		}
+		const FProperty* Property = PCGPropertyHelpers::FindPropertyByName(CurrentClass, CurrentName);
 
 		if (!Property)
 		{
@@ -305,7 +294,7 @@ UPCGParamData* PCGPropertyHelpers::ExtractPropertyAsAttributeSet(const PCGProper
 	bool bValidOperation = true;
 	bool bIgnoreOutputAttributeName = (PropertySelectors.Num() > 1);
 
-	for(const FPCGAttributePropertySelector& PropertySelector : PropertySelectors)
+	for (const FPCGAttributePropertySelector& PropertySelector : PropertySelectors)
 	{
 		TArray<const void*> Containers = { Parameters.Container };
 		const FProperty* Property = nullptr;
@@ -388,16 +377,9 @@ UPCGParamData* PCGPropertyHelpers::ExtractPropertyAsAttributeSet(const PCGProper
 
 			for (const FPCGSettingsOverridableParam& Param : AllChildProperties)
 			{
-				if (ensure(!Param.PropertiesNames.IsEmpty()))
+				if (ensure(!Param.PropertiesNames.IsEmpty()) && ensure(!Param.Properties.IsEmpty()) && ensure(Param.Properties[0]))
 				{
-					const FName ChildPropertyName = Param.PropertiesNames[0];
-					if (const FProperty* ChildProperty = UnderlyingClass->FindPropertyByName(ChildPropertyName))
-					{
-						// We use authored name as attribute name to avoid issue with noisy property names, like in UUserDefinedStructs, where some random number is appended to the property name.
-						// By default, it will just return the property name anyway.
-						FString AuthoredName = UnderlyingClass->GetAuthoredNameForField(ChildProperty);
-						ExtractableProperties.Emplace(std::move(AuthoredName), ChildProperty);
-					}
+					ExtractableProperties.Emplace(Param.PropertiesNames[0].ToString(), Param.Properties[0]);
 				}
 			}
 		}
@@ -553,6 +535,22 @@ UPCGParamData* PCGPropertyHelpers::ExtractPropertyAsAttributeSet(const PCGProper
 	return bValidOperation ? ParamData : nullptr;
 }
 
+const FProperty* PCGPropertyHelpers::FindPropertyByName(const UStruct* InStruct, const FName InName)
+{
+	if (!InStruct)
+	{
+		return nullptr;
+	}
+	else if (const UUserDefinedStruct* UserStruct = Cast<UUserDefinedStruct>(InStruct))
+	{
+		return FindPropertyInUserDefinedStruct(UserStruct, InName); 
+	}
+	else
+	{
+		return FindFProperty<FProperty>(InStruct, InName);
+	}
+}
+
 const FProperty* PCGPropertyHelpers::FindPropertyInUserDefinedStruct(const UUserDefinedStruct* InStruct, const FName InName)
 {
 	if (!InStruct)
diff --git a/Engine/Plugins/PCG/Source/PCG/Private/Helpers/PCGSettingsHelpers.cpp b/Engine/Plugins/PCG/Source/PCG/Private/Helpers/PCGSettingsHelpers.cpp
--- a/Engine/Plugins/PCG/Source/PCG/Private/Helpers/PCGSettingsHelpers.cpp
+++ b/Engine/Plugins/PCG/Source/PCG/Private/Helpers/PCGSettingsHelpers.cpp
@@ -309,27 +309,27 @@ namespace PCGSettingsHelpers
 
 #if WITH_EDITOR
 				auto HasMetadata = [Property](const TArray<FName>& InValues, bool bIsConjunction) -> bool
+				{
+					bool bFound = false;
+					for (const FName& Metadata : InValues)
 					{
-						bool bFound = false;
-						for (const FName& Metadata : InValues)
+						if (Property->HasMetaData(Metadata) && !bIsConjunction)
 						{
-							if (Property->HasMetaData(Metadata) && !bIsConjunction)
-							{
-								bFound = true;
-								if (!bIsConjunction)
-								{
-									break;
-								}
-							}
-							else if (bIsConjunction)
+							bFound = true;
+							if (!bIsConjunction)
 							{
-								bFound = false;
 								break;
 							}
 						}
+						else if (bIsConjunction)
+						{
+							bFound = false;
+							break;
+						}
+					}
 
-						return bFound;
-					};
+					return bFound;
+				};
 
 				if (!InConfig.IncludeMetadataValues.IsEmpty())
 				{
@@ -343,9 +343,9 @@ namespace PCGSettingsHelpers
 #endif // WITH_EDITOR
 
 				auto HasPropertyFlags = [Property](uint64 InFlags, bool bIsConjunction) -> bool
-					{
-						return bIsConjunction ? Property->HasAllPropertyFlags(InFlags) : Property->HasAnyPropertyFlags(InFlags);
-					};
+				{
+					return bIsConjunction ? Property->HasAllPropertyFlags(InFlags) : Property->HasAnyPropertyFlags(InFlags);
+				};
 
 				if (InConfig.IncludePropertyFlags != 0)
 				{
@@ -368,111 +368,123 @@ namespace PCGSettingsHelpers
 					continue;
 				}
 
-				auto RecursiveExtraction = [&InConfig, &Res, &LabelCache, Property, InClass, &GatherAliases, &InAlreadySeenProperties](const UStruct* NextClass)
+				auto RecursiveExtraction = [&InConfig, &Res, &LabelCache, Property, InClass, &GatherAliases, &InAlreadySeenProperties](const UStruct* NextClass, const FProperty* ContainerProperty)
+				{
+					// Reached max depth
+					if (InConfig.MaxStructDepth == 0)
 					{
-						// Reached max depth
-						if (InConfig.MaxStructDepth == 0)
-						{
-							return;
-						}
+						return;
+					}
 
-						// Use the seed, and don't check metadata for PCG overridable.
-						FPCGGetAllOverridableParamsConfig RecurseConfig = InConfig;
-						InAlreadySeenProperties.Add(Property);
-						RecurseConfig.bUseSeed = true;
+					// Use the seed, and don't check metadata for PCG overridable.
+					FPCGGetAllOverridableParamsConfig RecurseConfig = InConfig;
+					InAlreadySeenProperties.Add(ContainerProperty ? ContainerProperty : Property);
+					RecurseConfig.bUseSeed = true;
 #if WITH_EDITOR
-						RecurseConfig.IncludeMetadataValues.Remove(PCGObjectMetadata::Overridable);
-						RecurseConfig.IncludeMetadataValues.Remove(PCGObjectMetadata::OverridableCPUAndGPU);
-						RecurseConfig.IncludeMetadataValues.Remove(PCGObjectMetadata::OverridableCPUAndGPUWithReadback);
+					RecurseConfig.IncludeMetadataValues.Remove(PCGObjectMetadata::Overridable);
+					RecurseConfig.IncludeMetadataValues.Remove(PCGObjectMetadata::OverridableCPUAndGPU);
+					RecurseConfig.IncludeMetadataValues.Remove(PCGObjectMetadata::OverridableCPUAndGPUWithReadback);
 
-						// Also check if we have some filtering on the child properties
-						TArray<FString> ChildPropertiesFiltering;
-						if (Property->HasMetaData(PCGObjectMetadata::OverridableChildProperties))
-						{
-							const FString& ChildPropertiesMetadata = Property->GetMetaData(PCGObjectMetadata::OverridableChildProperties);
-							ChildPropertiesFiltering = PCGHelpers::GetStringArrayFromCommaSeparatedList(ChildPropertiesMetadata);
-						}
+					// Also check if we have some filtering on the child properties
+					TArray<FString> ChildPropertiesFiltering;
+					if (Property->HasMetaData(PCGObjectMetadata::OverridableChildProperties))
+					{
+						const FString& ChildPropertiesMetadata = Property->GetMetaData(PCGObjectMetadata::OverridableChildProperties);
+						ChildPropertiesFiltering = PCGHelpers::GetStringArrayFromCommaSeparatedList(ChildPropertiesMetadata);
+					}
 #endif // WITH_EDITOR
 
-						if (RecurseConfig.MaxStructDepth > 0)
-						{
-							RecurseConfig.MaxStructDepth--;
-						}
+					if (RecurseConfig.MaxStructDepth > 0)
+					{
+						RecurseConfig.MaxStructDepth--;
+					}
 
-						for (FPCGSettingsOverridableParam& ChildParam : GetAllOverridableParams_Internal(NextClass, RecurseConfig, InAlreadySeenProperties))
-						{
-							FName Label = ChildParam.Label;
-							bool bHasNameClash = false;
+					for (FPCGSettingsOverridableParam& ChildParam : GetAllOverridableParams_Internal(NextClass, RecurseConfig, InAlreadySeenProperties))
+					{
+						FName Label = ChildParam.Label;
+						bool bHasNameClash = false;
 #if WITH_EDITOR
-							// Don't check for label clash, as they would all be equal to None in non-editor build
-							const int32 CachedLabelIndex = LabelCache.IndexOfByKey(Label);
-							if (CachedLabelIndex != INDEX_NONE)
-							{
-								// If we have a clash, we will use the full path, so mark this param and the other that clashed to use the full path.
-								Res[CachedLabelIndex].bHasNameClash = true;
-								Res[CachedLabelIndex].Label = FName(Res[CachedLabelIndex].GetPropertyPath());
-								bHasNameClash = true;
-							}
+						// Don't check for label clash, as they would all be equal to None in non-editor build
+						const int32 CachedLabelIndex = LabelCache.IndexOfByKey(Label);
+						if (CachedLabelIndex != INDEX_NONE)
+						{
+							// If we have a clash, we will use the full path, so mark this param and the other that clashed to use the full path.
+							Res[CachedLabelIndex].bHasNameClash = true;
+							Res[CachedLabelIndex].Label = FName(Res[CachedLabelIndex].GetPropertyPath());
+							bHasNameClash = true;
+						}
 
-							LabelCache.AddUnique(Label);
+						LabelCache.AddUnique(Label);
 
-							// Also check if the current property path is filtered out
-							if (!ChildPropertiesFiltering.IsEmpty() && !ChildPropertiesFiltering.Contains(ChildParam.GetPropertyPath()))
-							{
-								continue;
-							}
+						// Also check if the current property path is filtered out
+						if (!ChildPropertiesFiltering.IsEmpty() && !ChildPropertiesFiltering.Contains(ChildParam.GetPropertyPath()))
+						{
+							continue;
+						}
 #endif // WITH_EDITOR
 
-							FPCGSettingsOverridableParam& Param = Res.Emplace_GetRef();
-							Param.PropertiesNames.Add(Property->GetFName());
-							Param.PropertiesNames.Append(std::move(ChildParam.PropertiesNames));
-							Param.Properties.Add(Property);
-							Param.Properties.Append(std::move(ChildParam.Properties));
-							Param.PropertyClass = InClass;
+						FPCGSettingsOverridableParam& Param = Res.Emplace_GetRef();
+						Param.PropertiesNames.Add(*Property->GetAuthoredName());
+						Param.PropertiesNames.Append(std::move(ChildParam.PropertiesNames));
+						Param.NumContainers = ChildParam.NumContainers;
+						if (ContainerProperty)
+						{
+							Param.NumContainers++;
+							Param.Properties.Add(ContainerProperty);
+						}
+
+						Param.Properties.Add(Property);
+						Param.Properties.Append(std::move(ChildParam.Properties));
+						Param.PropertyClass = InClass;
 #if WITH_EDITOR
-							Param.UnderlyingType = ChildParam.UnderlyingType;
+						Param.UnderlyingType = ChildParam.UnderlyingType;
 #endif //WITH_EDITOR
 
-							if (!ChildParam.bHasNameClash && !bHasNameClash)
-							{
-								Param.Label = Label;
-							}
-							else
-							{
+						if (!ChildParam.bHasNameClash && !bHasNameClash)
+						{
+							Param.Label = Label;
+						}
+						else
+						{
 #if WITH_EDITOR
-								Param.Label = FName(Param.GetPropertyPath());
+							Param.Label = FName(Param.GetPropertyPath());
 #endif // WITH_EDITOR
-								Param.bHasNameClash = true;
-							}
+							Param.bHasNameClash = true;
+						}
 
 #if WITH_EDITOR
-							Param.bSupportsGPU = ChildParam.bSupportsGPU;
-							Param.bRequiresGPUReadback = ChildParam.bRequiresGPUReadback;
-							GatherAliases(Property, Param);
-							// For all already found aliases, add them to the current param, with the index incremented by 1.
-							for (TPair<int32, FPCGPropertyAliases>& It : ChildParam.MapOfAliases)
-							{
-								Param.MapOfAliases.Emplace(It.Key + 1, std::move(It.Value));
-							}
-#endif // WITH_EDITOR
+						Param.bSupportsGPU = ChildParam.bSupportsGPU;
+						Param.bRequiresGPUReadback = ChildParam.bRequiresGPUReadback;
+						GatherAliases(Property, Param);
+						// For all already found aliases, add them to the current param, with the index incremented by 1.
+						for (TPair<int32, FPCGPropertyAliases>& It : ChildParam.MapOfAliases)
+						{
+							Param.MapOfAliases.Emplace(It.Key + 1, std::move(It.Value));
 						}
+#endif // WITH_EDITOR
+					}
 
-						InAlreadySeenProperties.Pop(EAllowShrinking::No);
-					};
+					InAlreadySeenProperties.Pop(EAllowShrinking::No);
+				};
 
-				const FProperty* PropertyToCheck = Property;
-				if (InConfig.bExtractArrays && Property->IsA<FArrayProperty>())
+				const FProperty* OriginalProperty = Property;
+				const FProperty* ContainerProperty = nullptr;
+				if (InConfig.bExtractArrays)
 				{
-					PropertyToCheck = CastFieldChecked<FArrayProperty>(Property)->Inner;
+					if (const FArrayProperty* ArrayProperty = CastField<FArrayProperty>(Property))
+					{
+						ContainerProperty = ArrayProperty;
+						Property = ArrayProperty->Inner;
+					}
 				}
 
 				// Always extract instanced objects, but be careful with infinite recursion (class having an instanced property on itself)
-				if (PropertyToCheck->IsA<FObjectProperty>() && !InAlreadySeenProperties.Contains(PropertyToCheck) && (InConfig.bExtractObjects || PropertyToCheck->HasAllPropertyFlags(CPF_InstancedReference)))
+				if (Property->IsA<FObjectProperty>() && !InAlreadySeenProperties.Contains(OriginalProperty) && (InConfig.bExtractObjects || Property->HasAllPropertyFlags(CPF_InstancedReference)))
 				{
-					RecursiveExtraction(CastFieldChecked<FObjectProperty>(PropertyToCheck)->PropertyClass);
+					RecursiveExtraction(CastFieldChecked<FObjectProperty>(Property)->PropertyClass, ContainerProperty);
 				}
 				// Validate that the property can be overriden by params
-				else if (PCGAttributeAccessorHelpers::IsPropertyAccessorSupported(PropertyToCheck))
+				else if (PCGAttributeAccessorHelpers::IsPropertyAccessorSupported(Property))
 				{
 					FName Label = NAME_None;
 					bool bHasNameClash = false;
@@ -491,21 +503,27 @@ namespace PCGSettingsHelpers
 #endif // WITH_EDITOR
 
 					FPCGSettingsOverridableParam& Param = Res.Emplace_GetRef();
-					Param.PropertiesNames.Add(Property->GetFName());
+					Param.PropertiesNames.Add(*Property->GetAuthoredName());
+					if (ContainerProperty)
+					{
+						Param.Properties.Add(ContainerProperty);
+						Param.NumContainers = 1;
+					}
+
 					Param.Properties.Add(Property);
 					Param.PropertyClass = InClass;
 					Param.bHasNameClash = bHasNameClash;
 #if WITH_EDITOR
 					Param.Label = bHasNameClash ? FName(Param.GetPropertyPath()) : Label;
-					Param.bSupportsGPU = Property->HasMetaData(PCGObjectMetadata::OverridableCPUAndGPU) || Property->HasMetaData(PCGObjectMetadata::OverridableCPUAndGPUWithReadback);
-					Param.bRequiresGPUReadback = Property->HasMetaData(PCGObjectMetadata::OverridableCPUAndGPUWithReadback);
-					GatherAliases(Property, Param);
+					Param.bSupportsGPU = OriginalProperty->HasMetaData(PCGObjectMetadata::OverridableCPUAndGPU) || Property->HasMetaData(PCGObjectMetadata::OverridableCPUAndGPUWithReadback);
+					Param.bRequiresGPUReadback = OriginalProperty->HasMetaData(PCGObjectMetadata::OverridableCPUAndGPUWithReadback);
+					GatherAliases(OriginalProperty, Param);
 					Param.UnderlyingType = PCGAttributeAccessorHelpers::GetMetadataTypeForProperty(Property);
 #endif // WITH_EDITOR
 				}
-				else if (const FStructProperty* StructProperty = CastField<FStructProperty>(PropertyToCheck))
+				else if (const FStructProperty* StructProperty = CastField<FStructProperty>(Property))
 				{
-					RecursiveExtraction(StructProperty->Struct);
+					RecursiveExtraction(StructProperty->Struct, ContainerProperty);
 				}
 			}
 		}
diff --git a/Engine/Plugins/PCG/Source/PCG/Private/Metadata/Accessors/PCGAttributeAccessorHelpers.cpp b/Engine/Plugins/PCG/Source/PCG/Private/Metadata/Accessors/PCGAttributeAccessorHelpers.cpp
--- a/Engine/Plugins/PCG/Source/PCG/Private/Metadata/Accessors/PCGAttributeAccessorHelpers.cpp
+++ b/Engine/Plugins/PCG/Source/PCG/Private/Metadata/Accessors/PCGAttributeAccessorHelpers.cpp
@@ -13,8 +13,8 @@
 #include "Metadata/PCGMetadataAttributeTpl.h"
 #include "Metadata/Accessors/IPCGAttributeAccessorTpl.h"
 #include "Metadata/Accessors/PCGAttributeAccessor.h"
+#include "Metadata/Accessors/PCGAttributeAccessorHelpersInternal.h"
 #include "Metadata/Accessors/PCGCustomAccessor.h"
-#include "Metadata/Accessors/PCGPropertyAccessor.h"
 #include "Metadata/Accessors/PCGAttributeExtractor.h"
 
 #include "StructUtils/UserDefinedStruct.h"
@@ -92,143 +92,6 @@ namespace PCGAttributeAccessorHelpers
 		return InAccessor;
 	}
 
-	// Empty signature to passthrough types to functors.
-	template <typename T>
-	struct Signature
-	{
-		using Type = T;
-	};
-
-	template <typename Func>
-	decltype(auto) DispatchPropertyTypes(const FProperty* InProperty, Func&& Functor)
-	{
-		// Use the FPCGPropertyPathAccessor (on soft object path) as dummy type to get the functor return type
-		// because this accessor takes a generic FProperty, while others take a more specialized type (like FNumericProperty).
-		using ReturnType = decltype(Functor(Signature<FPCGPropertyPathAccessor<FSoftObjectPath>>{}, InProperty));
-
-		if (!InProperty)
-		{
-			return ReturnType{};
-		}
-
-		if (const FNumericProperty* NumericProperty = CastField<FNumericProperty>(InProperty))
-		{
-			if (NumericProperty->IsFloatingPoint())
-			{
-				// As floating properties are mostly all double in UE, convert float to double attributes.
-				return Functor(Signature<FPCGNumericPropertyAccessor<double>>{}, NumericProperty);
-			}
-			else if (NumericProperty->IsInteger())
-			{
-				// But for int32/int64 we can distinguish between the two. Everything of size 32 or less is will be an int32, 64bits integers will be int64.
-				if (NumericProperty->IsA<FInt64Property>() || NumericProperty->IsA<FUInt64Property>())
-				{
-					return Functor(Signature<FPCGNumericPropertyAccessor<int64>>{}, NumericProperty);
-				}
-				else
-				{
-					return Functor(Signature<FPCGNumericPropertyAccessor<int32>>{}, NumericProperty);
-				}
-			}
-		}
-		else if (const FBoolProperty* BoolProperty = CastField<FBoolProperty>(InProperty))
-		{
-			return Functor(Signature<FPCGPropertyAccessor<bool, FBoolProperty>>{}, BoolProperty);
-		}
-		else if (const FStrProperty* StringProperty = CastField<FStrProperty>(InProperty))
-		{
-			return Functor(Signature<FPCGPropertyAccessor<FString, FStrProperty>>{}, StringProperty);
-		}
-		else if (const FNameProperty* NameProperty = CastField<FNameProperty>(InProperty))
-		{
-			return Functor(Signature<FPCGPropertyAccessor<FName, FNameProperty>>{}, NameProperty);
-		}
-		else if (const FEnumProperty* EnumProperty = CastField<FEnumProperty>(InProperty))
-		{
-			return Functor(Signature<FPCGEnumPropertyAccessor>{}, EnumProperty);
-		}
-		else if (const FSoftClassProperty* SoftClassProperty = CastField<FSoftClassProperty>(InProperty))
-		{
-			return Functor(Signature<FPCGPropertySoftClassPathAccessor>{}, SoftClassProperty);
-		}
-		else if (const FSoftObjectProperty* SoftObjectProperty = CastField<FSoftObjectProperty>(InProperty))
-		{
-			return Functor(Signature<FPCGPropertySoftObjectPathAccessor>{}, SoftObjectProperty);
-		}
-		else if (const FClassProperty* ClassProperty = CastField<FClassProperty>(InProperty))
-		{
-			return Functor(Signature<FPCGPropertyObjectPtrAccessor<FClassProperty>>{}, ClassProperty);
-		}
-		else if (const FObjectProperty* ObjectProperty = CastField<FObjectProperty>(InProperty))
-		{
-			return Functor(Signature<FPCGPropertyObjectPtrAccessor<FObjectProperty>>{}, ObjectProperty);
-		}
-		else if (const FStructProperty* StructProperty = CastField<FStructProperty>(InProperty))
-		{
-			if (StructProperty->Struct == TBaseStructure<FVector>::Get())
-			{
-				return Functor(Signature<FPCGPropertyStructAccessor<FVector>>{}, StructProperty);
-			}
-			else if (StructProperty->Struct == TBaseStructure<FVector4>::Get())
-			{
-				return Functor(Signature<FPCGPropertyStructAccessor<FVector4>>{}, StructProperty);
-			}
-			else if (StructProperty->Struct == TBaseStructure<FQuat>::Get())
-			{
-				return Functor(Signature<FPCGPropertyStructAccessor<FQuat>>{}, StructProperty);
-			}
-			else if (StructProperty->Struct == TBaseStructure<FTransform>::Get())
-			{
-				return Functor(Signature<FPCGPropertyStructAccessor<FTransform>>{}, StructProperty);
-			}
-			else if (StructProperty->Struct == TBaseStructure<FRotator>::Get())
-			{
-				return Functor(Signature<FPCGPropertyStructAccessor<FRotator>>{}, StructProperty);
-			}
-			else if (StructProperty->Struct == TBaseStructure<FVector2D>::Get())
-			{
-				return Functor(Signature<FPCGPropertyStructAccessor<FVector2D>>{}, StructProperty);
-			}
-			else if (StructProperty->Struct == TBaseStructure<FSoftObjectPath>::Get())
-			{
-				return Functor(Signature<FPCGPropertyStructAccessor<FSoftObjectPath>>{}, StructProperty);
-			}
-			else if (StructProperty->Struct == TBaseStructure<FSoftClassPath>::Get())
-			{
-				return Functor(Signature<FPCGPropertyStructAccessor<FSoftClassPath>>{}, StructProperty);
-			}
-			else if (StructProperty->Struct == TBaseStructure<FLinearColor>::Get())
-			{
-				return Functor(Signature<FPCGLinearColorAccessor>{}, StructProperty);
-			}
-			else if (StructProperty->Struct == TBaseStructure<FColor>::Get())
-			{
-				return Functor(Signature<FPCGColorAccessor>{}, StructProperty);
-			}
-			else if (StructProperty->Struct && StructProperty->Struct->IsChildOf(FPCGAttributePropertySelector::StaticStruct()))
-			{
-				return Functor(Signature<FPCGAttributePropertySelectorAccessor>{}, StructProperty);
-			}
-		}
-
-		return ReturnType{};
-	}
-
-	template <typename Func>
-	decltype(auto) DispatchPropertyTypes(const FName InPropertyName, const UStruct* InStruct, Func&& Functor)
-	{
-		if (InStruct)
-		{
-			if (const FProperty* Property = InStruct->FindPropertyByName(InPropertyName))
-			{
-				return DispatchPropertyTypes(Property, std::forward<Func>(Functor));
-			}
-		}
-
-		using ReturnType = decltype(Functor(static_cast<FPCGPropertyPathAccessor<FSoftObjectPath>*>(nullptr), static_cast<const FProperty*>(nullptr)));
-		return ReturnType{};
-	}
-
 	bool GetPropertyChain(const TArray<FName>& InPropertyNames, const UStruct* InStruct, TArray<const FProperty*>& OutProperties)
 	{
 		check(InStruct);
@@ -259,6 +122,13 @@ namespace PCGAttributeAccessorHelpers
 
 			OutProperties.Add(Property);
 
+			// If the property is a container, we also need to add the inner property
+			if (const FArrayProperty* ArrayProperty = CastField<FArrayProperty>(Property))
+			{
+				Property = ArrayProperty->Inner;
+				OutProperties.Add(Property);
+			}
+
 			// Check for a struct or object for all properties except the last one.
 			if (i < InPropertyNames.Num() - 1)
 			{
@@ -284,7 +154,7 @@ namespace PCGAttributeAccessorHelpers
 
 TUniquePtr<IPCGAttributeAccessor> PCGAttributeAccessorHelpers::CreatePropertyAccessor(const FProperty* InProperty)
 {
-	return PCGAttributeAccessorHelpers::DispatchPropertyTypes(InProperty, [](auto SignatureDummy, const auto* TypedProperty) -> TUniquePtr<IPCGAttributeAccessor>
+	return PCGAttributeAccessorHelpers::Internal::DispatchPropertyTypes(InProperty, [](auto SignatureDummy, const auto* TypedProperty) -> TUniquePtr<IPCGAttributeAccessor>
 	{
 		using TypedAccessor = typename decltype(SignatureDummy)::Type;
 		return MakeUnique<TypedAccessor>(TypedProperty);
@@ -293,7 +163,7 @@ TUniquePtr<IPCGAttributeAccessor> PCGAttributeAccessorHelpers::CreatePropertyAcc
 
 TUniquePtr<IPCGAttributeAccessor> PCGAttributeAccessorHelpers::CreatePropertyAccessor(const FName InPropertyName, const UStruct* InStruct)
 {
-	return PCGAttributeAccessorHelpers::DispatchPropertyTypes(InPropertyName, InStruct, [](auto SignatureDummy, const auto* TypedProperty) -> TUniquePtr<IPCGAttributeAccessor>
+	return PCGAttributeAccessorHelpers::Internal::DispatchPropertyTypes(InPropertyName, InStruct, [](auto SignatureDummy, const auto* TypedProperty) -> TUniquePtr<IPCGAttributeAccessor>
 	{
 		using TypedAccessor = typename decltype(SignatureDummy)::Type;
 		return MakeUnique<TypedAccessor>(TypedProperty);
@@ -307,7 +177,7 @@ TUniquePtr<IPCGAttributeAccessor> PCGAttributeAccessorHelpers::CreatePropertyCha
 		return TUniquePtr<IPCGAttributeAccessor>{};
 	}
 
-	return PCGAttributeAccessorHelpers::DispatchPropertyTypes(InProperties.Last(), [&InProperties](auto SignatureDummy, const auto* TypedProperty) -> TUniquePtr<IPCGAttributeAccessor>
+	return PCGAttributeAccessorHelpers::Internal::DispatchPropertyTypes(InProperties.Last(), [&InProperties](auto SignatureDummy, const auto* TypedProperty) -> TUniquePtr<IPCGAttributeAccessor>
 	{
 		using TypedAccessor = typename decltype(SignatureDummy)::Type;
 		return MakeUnique<TypedAccessor>(TypedProperty, std::forward<TArray<const FProperty*>>(InProperties));
@@ -329,7 +199,7 @@ TUniquePtr<IPCGAttributeAccessor> PCGAttributeAccessorHelpers::CreatePropertyCha
 
 bool PCGAttributeAccessorHelpers::IsPropertyAccessorSupported(const FProperty* InProperty)
 {
-	return PCGAttributeAccessorHelpers::DispatchPropertyTypes(InProperty, [](auto, const auto*) -> bool
+	return PCGAttributeAccessorHelpers::Internal::DispatchPropertyTypes(InProperty, [](auto, const auto*) -> bool
 	{
 		return true;
 	});
@@ -337,7 +207,7 @@ bool PCGAttributeAccessorHelpers::IsPropertyAccessorSupported(const FProperty* I
 
 bool PCGAttributeAccessorHelpers::IsPropertyAccessorSupported(const FName InPropertyName, const UStruct* InStruct)
 {
-	return PCGAttributeAccessorHelpers::DispatchPropertyTypes(InPropertyName, InStruct, [](auto, const auto*) -> bool
+	return PCGAttributeAccessorHelpers::Internal::DispatchPropertyTypes(InPropertyName, InStruct, [](auto, const auto*) -> bool
 	{
 		return true;
 	});
@@ -354,7 +224,7 @@ bool PCGAttributeAccessorHelpers::IsPropertyAccessorChainSupported(const TArray<
 EPCGMetadataTypes PCGAttributeAccessorHelpers::GetMetadataTypeForProperty(const FProperty* InProperty)
 {
 	bool bFound = false;
-	const EPCGMetadataTypes Result = PCGAttributeAccessorHelpers::DispatchPropertyTypes(InProperty, [&bFound](auto SignatureDummy, const auto* TypedProperty) -> EPCGMetadataTypes
+	const EPCGMetadataTypes Result = PCGAttributeAccessorHelpers::Internal::DispatchPropertyTypes(InProperty, [&bFound](auto SignatureDummy, const auto* TypedProperty) -> EPCGMetadataTypes
 	{
 		bFound = true;
 		using TypedAccessor = typename decltype(SignatureDummy)::Type;
@@ -437,6 +307,11 @@ TUniquePtr<const IPCGAttributeAccessor> PCGAttributeAccessorHelpers::CreateConst
 		if (Result && OutResult && ParamData && ParamData->Metadata)
 		{
 			OutResult->bHasMultipleDataInAttributeSet = ParamData->Metadata->GetLocalItemCount() > 1;
+			
+			if (InParam.NumContainers > 0)
+			{
+				OutResult->Keys = PCGAttributeAccessorHelpers::CreateConstKeys(ParamData, InputSelector);
+			}
 		}
 
 		return Result;
diff --git a/Engine/Plugins/PCG/Source/PCG/Private/Metadata/Accessors/PCGPropertyAccessor.cpp b/Engine/Plugins/PCG/Source/PCG/Private/Metadata/Accessors/PCGPropertyAccessor.cpp
--- a/Engine/Plugins/PCG/Source/PCG/Private/Metadata/Accessors/PCGPropertyAccessor.cpp
+++ b/Engine/Plugins/PCG/Source/PCG/Private/Metadata/Accessors/PCGPropertyAccessor.cpp
@@ -4,32 +4,6 @@
 
 #include "UObject/EnumProperty.h"
 
-TArray<const void*> PCGPropertyAccessor::GetContainerKeys(int32 Index, int32 Range, const IPCGAttributeAccessorKeys& Keys)
-{
-	TArray<const void*> ContainerKeys;
-	ContainerKeys.SetNumUninitialized(Range);
-	TArrayView<const void*> ContainerKeysView(ContainerKeys);
-	if (!Keys.GetKeys(Index, ContainerKeysView))
-	{
-		return TArray<const void*>{};
-	}
-
-	return ContainerKeys;
-}
-
-TArray<void*> PCGPropertyAccessor::GetContainerKeys(int32 Index, int32 Range, IPCGAttributeAccessorKeys& Keys)
-{
-	TArray<void*> ContainerKeys;
-	ContainerKeys.SetNumUninitialized(Range);
-	TArrayView<void*> ContainerKeysView(ContainerKeys);
-	if (!Keys.GetKeys(Index, ContainerKeysView))
-	{
-		return TArray<void*>{};
-	}
-
-	return ContainerKeys;
-}
-
 IPCGPropertyChain::IPCGPropertyChain(const FProperty* Property, TArray<const FProperty*>&& ExtraProperties)
 	: PropertyChain(std::move(ExtraProperties))
 {
diff --git a/Engine/Plugins/PCG/Source/PCG/Private/Metadata/PCGAttributePropertySelector.cpp b/Engine/Plugins/PCG/Source/PCG/Private/Metadata/PCGAttributePropertySelector.cpp
--- a/Engine/Plugins/PCG/Source/PCG/Private/Metadata/PCGAttributePropertySelector.cpp
+++ b/Engine/Plugins/PCG/Source/PCG/Private/Metadata/PCGAttributePropertySelector.cpp
@@ -9,6 +9,7 @@
 #include "Metadata/PCGMetadataAttribute.h"
 
 #include "Serialization/ArchiveCrc32.h"
+#include "String/ParseTokens.h"
 
 #include UE_INLINE_GENERATED_CPP_BY_NAME(PCGAttributePropertySelector)
 
@@ -365,7 +366,7 @@ bool FPCGAttributePropertySelector::ResetExtraNames()
 	}
 }
 
-bool FPCGAttributePropertySelector::Update(const FString& NewValue)
+bool FPCGAttributePropertySelector::Update(const FStringView NewValue)
 {
 	if (NewValue.IsEmpty())
 	{
@@ -373,7 +374,12 @@ bool FPCGAttributePropertySelector::Update(const FString& NewValue)
 	}
 	
 	TArray<FString> NewValues;
-	NewValue.ParseIntoArray(NewValues, PCGAttributePropertySelectorConstants::ExtraSeparator, /*InCullEmpty=*/ false);
+	UE::String::ParseTokens(
+		NewValue,
+		FStringView(PCGAttributePropertySelectorConstants::ExtraSeparator),
+		[&NewValues](FStringView Token) { NewValues.Emplace(Token); },
+		UE::String::EParseTokensOptions::IgnoreCase
+	);
 
 	check(!NewValues.IsEmpty())
 
@@ -423,6 +429,12 @@ bool FPCGAttributePropertySelector::Update(const FString& NewValue)
 	}
 }
 
+bool FPCGAttributePropertySelector::Update(const char* NewValue)
+{
+	FString Temp(NewValue);
+	return Update(MakeStringView(Temp));
+}
+
 void FPCGAttributePropertySelector::AddToCrc(FArchiveCrc32& Ar) const
 {
 	EPCGAttributePropertySelection RealSelection = GetSelection();
diff --git a/Engine/Plugins/PCG/Source/PCG/Private/Metadata/PCGObjectPropertyOverride.cpp b/Engine/Plugins/PCG/Source/PCG/Private/Metadata/PCGObjectPropertyOverride.cpp
--- a/Engine/Plugins/PCG/Source/PCG/Private/Metadata/PCGObjectPropertyOverride.cpp
+++ b/Engine/Plugins/PCG/Source/PCG/Private/Metadata/PCGObjectPropertyOverride.cpp
@@ -7,6 +7,7 @@
 #include "PCGModule.h"
 #include "PCGParamData.h"
 #include "Metadata/Accessors/PCGAttributeAccessorHelpers.h"
+#include "Metadata/Accessors/PCGContainerAccessor.h"
 
 #include "GameFramework/Actor.h"
 
@@ -25,46 +26,89 @@ namespace PCGObjectPropertyOverrideHelpers
 
 	void ApplyOverrides(const TArray<FPCGObjectPropertyOverrideDescription>& InObjectPropertyOverrideDescriptions, const TArray<TPair<UObject*, int32>>& TargetObjectToOverrideIndices, FName OverridesPinLabel, int32 InInputDataIndex, FPCGContext* Context, bool bPropagateEditChangeEvent)
 	{
-		if (!Context || TargetObjectToOverrideIndices.IsEmpty())
+		if (!Context || TargetObjectToOverrideIndices.IsEmpty() || OverridesPinLabel == NAME_None || InInputDataIndex == INDEX_NONE)
 		{
 			return;
 		}
 
-		const TArray<FPCGTaggedData> OverrideInputs = Context->InputData.GetInputsByPin(OverridesPinLabel);
-
-		int32 InputIndex = InInputDataIndex;
-
-		if (OverrideInputs.Num() == 0)
+		const TArray<FPCGTaggedData> OverridesTaggedData = Context->InputData.GetInputsByPin(OverridesPinLabel);
+		if (OverridesTaggedData.IsEmpty())
 		{
 			return;
 		}
-		else if (OverrideInputs.Num() == 1)
+		
+		if (!OverridesTaggedData.IsValidIndex(InInputDataIndex))
 		{
-			InputIndex = 0;
+			PCGLog::LogWarningOnGraph(FText::Format(LOCTEXT("InconsistentDataCount", "The data provided on pin '{0}' does not have a consistent size with the input index '{1}'. Will use the first one."), FText::FromName(OverridesPinLabel), FText::AsNumber(InInputDataIndex)), Context);
+			InInputDataIndex = 0;
 		}
-		else if (OverrideInputs.Num() > 1 && !OverrideInputs.IsValidIndex(InputIndex))
+
+		const UPCGData* OverrideData = OverridesTaggedData[InInputDataIndex].Data;
+		
+		FApplyOverrideParams Params =
 		{
-			PCGLog::LogWarningOnGraph(FText::Format(LOCTEXT("InconsistentDataCount", "The data provided on pin '{0}' does not have a consistent size with the input index '{1}'. Will use the first one."), FText::FromName(OverridesPinLabel), FText::AsNumber(InInputDataIndex)), Context);
-			InputIndex = 0;
+			.ObjectPropertyOverrideDescriptions = InObjectPropertyOverrideDescriptions,
+			.TargetObjectToOverrideIndices = TargetObjectToOverrideIndices,
+			.OverridesData = MakeConstArrayView(&OverrideData, 1),
+			.OptionalContext = Context,
+			.bPropagateEditChangeEvent = bPropagateEditChangeEvent
+		};
+
+		return ApplyOverrides(Params);
+	}
+
+	void ApplyOverrides(const FApplyOverrideParams& InParams)
+	{
+		if (InParams.TargetObjectToOverrideIndices.IsEmpty() || InParams.OverridesData.IsEmpty())
+		{
+			return;
 		}
 
-		const UPCGData* OverrideData = OverrideInputs[InputIndex].Data;
-		if (!OverrideData)
+		if (InParams.OverridesData.Num() == 0)
 		{
-			PCGLog::LogErrorOnGraph(LOCTEXT("InvalidActorOverrideData", "Invalid input data for Object Property Overrides pin."), Context);
 			return;
 		}
+		else if (InParams.OverridesData.Num() > 1)
+		{
+			if (!InParams.bOneDataSetPerElement)
+			{
+				PCGLog::LogWarningOnGraph(LOCTEXT("TooManyData", "When not having a Data set per element it is expected to pass just a single override data. Will use the first one"), InParams.OptionalContext);
+			}
+			else if (InParams.OverridesData.Num() != InParams.TargetObjectToOverrideIndices.Num())
+			{
+				PCGLog::LogErrorOnGraph(FText::Format(LOCTEXT("InconsistentDataCount", "When having one data set per element, it is expected to have the same number of override data than elements, we expected {0} and got {1} data."), InParams.TargetObjectToOverrideIndices.Num(), InParams.OverridesData.Num()), InParams.OptionalContext);
+				return;
+			}
+		}
 
-		for (const TPair<UObject*, int32>& TargetObjectAndIndex : TargetObjectToOverrideIndices)
+		for (const TPair<UObject*, int32>& TargetObjectAndIndex : InParams.TargetObjectToOverrideIndices)
 		{
 			UObject* TargetObject = TargetObjectAndIndex.Key;
 			int32 InputKeyIndex = TargetObjectAndIndex.Value;
 
-			FPCGObjectOverrides ObjectOverrides(TargetObject);
-			ObjectOverrides.Initialize(InObjectPropertyOverrideDescriptions, TargetObject, OverrideData, Context, bPropagateEditChangeEvent);
+			if (!IsValid(TargetObject))
+			{
+				continue;
+			}
+
+			const UPCGData* OverrideData = InParams.bOneDataSetPerElement ? InParams.OverridesData[InputKeyIndex % InParams.OverridesData.Num()] : InParams.OverridesData[0];
+
+			FPCGObjectOverrides<UObject> ObjectOverrides(TargetObject);
+
+			FPCGObjectOverrides<UObject>::FInitializeParams Params =
+			{
+				.OverrideDescriptions = InParams.ObjectPropertyOverrideDescriptions,
+				.TemplateObject = TargetObject,
+				.SourceData = OverrideData,
+				.OptionalContext = InParams.OptionalContext,
+				.bInPropagateEditChangeEvent = InParams.bPropagateEditChangeEvent,
+				.bAllowedToUseAllSourceDataForContainers = InParams.bOneDataSetPerElement
+			};
+			
+			ObjectOverrides.Initialize(Params);
 			if (!ObjectOverrides.Apply(InputKeyIndex)) // Use the First Entry of the param data for override (similar to what is done in Parameter Overrides in FPCGContext)
 			{
-				PCGLog::LogErrorOnGraph(FText::Format(LOCTEXT("ApplyOverrideFailed", "Failed to apply property overrides to object '%s'."), FText::FromName(TargetObject->GetClass()->GetFName())), Context);
+				PCGLog::LogErrorOnGraph(FText::Format(LOCTEXT("ApplyOverrideFailed", "Failed to apply property overrides to object '%s'."), FText::FromName(TargetObject->GetClass()->GetFName())), InParams.OptionalContext);
 			}
 		}
 	}
@@ -82,68 +126,86 @@ namespace PCGObjectPropertyOverrideHelpers
 
 void FPCGObjectSingleOverride::Initialize(const FPCGAttributePropertySelector& InputSelector, const FString& OutputProperty, const UStruct* TemplateClass, const UPCGData* SourceData, FPCGContext* Context, bool bComputePropertyEditChain)
 {
-	InputKeys = PCGAttributeAccessorHelpers::CreateConstKeys(SourceData, InputSelector);
-	ObjectOverrideInputAccessor = PCGAttributeAccessorHelpers::CreateConstAccessor(SourceData, InputSelector);
+	FInitializeParams Params =
+	{
+		.InputSelector = &InputSelector,
+		.OutputProperty = MakeStringView(OutputProperty),
+		.TemplateClass = TemplateClass,
+		.SourceData = SourceData,
+		.OptionalContext = Context,
+		.bComputePropertyEditChain = bComputePropertyEditChain
+	};
+
+	Initialize(Params);
+}
+
+void FPCGObjectSingleOverride::Initialize(const FInitializeParams& InParams)
+{
+	if (!InParams.InputSelector)
+	{
+		return;
+	}
+	
+	InputKeys = PCGAttributeAccessorHelpers::CreateConstKeys(InParams.SourceData, *InParams.InputSelector);
+	ObjectOverrideInputAccessor = PCGAttributeAccessorHelpers::CreateConstAccessor(InParams.SourceData, *InParams.InputSelector);
 
 	if (!ObjectOverrideInputAccessor.IsValid())
 	{
-		PCGLog::LogWarningOnGraph(FText::Format(NSLOCTEXT("PCGObjectPropertyOverride", "OverrideInputInvalid", "ObjectOverride for input '{0}' does not exist or is unsupported. Check the attribute or property selection."), InputSelector.GetDisplayText()), Context);
+		PCGLog::LogWarningOnGraph(FText::Format(NSLOCTEXT("PCGObjectPropertyOverride", "OverrideInputInvalid", "ObjectOverride for input '{0}' does not exist or is unsupported. Check the attribute or property selection."), InParams.InputSelector->GetDisplayText()), InParams.OptionalContext);
 		return;
 	}
 
-	const FPCGAttributePropertySelector OutputSelector = FPCGAttributePropertySelector::CreateSelectorFromString(OutputProperty);
+	const FPCGAttributePropertySelector OutputSelector = FPCGAttributePropertySelector::CreateSelectorFromString(InParams.OutputProperty);
 	// TODO: Move implementation into a new helper: PCGAttributeAccessorHelpers::CreatePropertyAccessor(const FPCGAttributePropertySelector& Selector, UStruct* Class)
 	const TArray<FString>& ExtraNames = OutputSelector.GetExtraNames();
-	if (ExtraNames.IsEmpty())
-	{
-		if (FProperty* Property = TemplateClass->FindPropertyByName(FName(OutputProperty)))
-		{
-			ObjectOverrideOutputAccessor = PCGAttributeAccessorHelpers::CreatePropertyAccessor(Property);
-			TArray<const FStructProperty*> EncounteredStructProps;
-			bWillNeedLoading = Property->ContainsObjectReference(EncounteredStructProps, EPropertyObjectReferenceType::Strong);
 
-#if WITH_EDITOR
-			if (bComputePropertyEditChain)
-			{
-				EditPropertyChain.AddHead(Property);
-			}
-#endif // WITH_EDITOR
-		}
+	TArray<FName> PropertyNames;
+	PropertyNames.Reserve(ExtraNames.Num() + 1);
+	PropertyNames.Add(OutputSelector.GetAttributeName());
+	for (const FString& Name : ExtraNames)
+	{
+		PropertyNames.Add(FName(Name));
 	}
-	else
+
+	TArray<const FProperty*> PropertyChain;
+	PCGAttributeAccessorHelpers::GetPropertyChain(PropertyNames, InParams.TemplateClass, PropertyChain);
+	if (!PropertyChain.IsEmpty())
 	{
-		TArray<FName> PropertyNames;
-		PropertyNames.Reserve(ExtraNames.Num() + 1);
-		PropertyNames.Add(OutputSelector.GetAttributeName());
-		for (const FString& Name : ExtraNames)
-		{
-			PropertyNames.Add(FName(Name));
-		}
+		TArray<const FStructProperty*> EncounteredStructProps;
+		bWillNeedLoading = PropertyChain.Last()->ContainsObjectReference(EncounteredStructProps, EPropertyObjectReferenceType::Strong);
 
-		TArray<const FProperty*> PropertyChain;
-		PCGAttributeAccessorHelpers::GetPropertyChain(PropertyNames, TemplateClass, PropertyChain);
-		if (!PropertyChain.IsEmpty())
+		int32 NumContainers = 0;
+		for (const FProperty* Property : PropertyChain)
 		{
-			TArray<const FStructProperty*> EncounteredStructProps;
-			bWillNeedLoading = PropertyChain.Last()->ContainsObjectReference(EncounteredStructProps, EPropertyObjectReferenceType::Strong);
+			if (ensure(Property) && Property->IsA<FArrayProperty>())
+			{
+				NumContainers++;
+			}
 
 #if WITH_EDITOR
-			if (bComputePropertyEditChain)
+			if (InParams.bComputePropertyEditChain)
 			{
-				for (const FProperty* Property : PropertyChain)
-				{
-					EditPropertyChain.AddTail(const_cast<FProperty*>(Property));
-				}
+				EditPropertyChain.AddTail(const_cast<FProperty*>(Property));
 			}
 #endif // WITH_EDITOR
+		}
 
-			ObjectOverrideOutputAccessor = PCGAttributeAccessorHelpers::CreatePropertyChainAccessor(MoveTemp(PropertyChain));
+		if (NumContainers > 1)
+		{
+			// Unsupported
+			ObjectOverrideInputAccessor.Reset();
+
+			PCGLog::LogErrorOnGraph(FText::Format(NSLOCTEXT("PCGObjectPropertyOverride", "OverrideOutputTooManyContainers", "ObjectOverride for output '{0}' is within multiple containers. Unsupported."), OutputSelector.GetDisplayText()), InParams.OptionalContext);
+			return;
 		}
+		
+		ObjectOverrideOutputAccessor = NumContainers == 1 ? PCGContainerAccessorHelpers::MakeContainerAccessor(PropertyChain) : PCGAttributeAccessorHelpers::CreatePropertyChainAccessor(MoveTemp(PropertyChain));
+		bIsContainer = NumContainers == 1;
 	}
 
 	if (!ObjectOverrideOutputAccessor.IsValid())
 	{
-		PCGLog::LogWarningOnGraph(FText::Format(NSLOCTEXT("PCGObjectPropertyOverride", "OverrideOutputInvalid", "ObjectOverride for object property '{0}' is invalid or unsupported. Check the attribute or property selection."), FText::FromString(OutputProperty)), Context);
+		PCGLog::LogWarningOnGraph(FText::Format(NSLOCTEXT("PCGObjectPropertyOverride", "OverrideOutputInvalid", "ObjectOverride for object property '{0}' is invalid or unsupported. Check the attribute or property selection."), FText::FromStringView(InParams.OutputProperty)), InParams.OptionalContext);
 		return;
 	}
 
@@ -152,11 +214,11 @@ void FPCGObjectSingleOverride::Initialize(const FPCGAttributePropertySelector& I
 		PCGLog::LogWarningOnGraph(
 			FText::Format(
 				NSLOCTEXT("PCGObjectPropertyOverride", "TypesIncompatible", "ObjectOverride cannot set input '{0}' to output '{1}'. Cannot convert type '{2}' to type '{3}'. Will be skipped."),
-				InputSelector.GetDisplayText(),
-				FText::FromString(OutputProperty),
+				InParams.InputSelector->GetDisplayText(),
+				FText::FromStringView(InParams.OutputProperty),
 				PCG::Private::GetTypeNameText(ObjectOverrideInputAccessor->GetUnderlyingType()),
 				PCG::Private::GetTypeNameText(ObjectOverrideOutputAccessor->GetUnderlyingType())),
-			Context);
+			InParams.OptionalContext);
 
 		ObjectOverrideInputAccessor.Reset();
 		ObjectOverrideOutputAccessor.Reset();
@@ -172,7 +234,8 @@ void FPCGObjectSingleOverride::Initialize(const FPCGAttributePropertySelector& I
 
 	PCGMetadataAttribute::CallbackWithRightType(ObjectOverrideOutputAccessor->GetUnderlyingType(), CreateGetterSetter);
 
-	ObjectOverrideInputSelector = InputSelector;
+	ObjectOverrideInputSelector = *InParams.InputSelector;
+	bAllowedToUseAllSourceDataForContainers = InParams.bAllowedToUseAllSourceDataForContainers;
 }
 
 bool FPCGObjectSingleOverride::IsValid() const
diff --git a/Engine/Plugins/PCG/Source/PCG/Private/PCGContext.cpp b/Engine/Plugins/PCG/Source/PCG/Private/PCGContext.cpp
--- a/Engine/Plugins/PCG/Source/PCG/Private/PCGContext.cpp
+++ b/Engine/Plugins/PCG/Source/PCG/Private/PCGContext.cpp
@@ -20,6 +20,7 @@
 #include "Metadata/Accessors/IPCGAttributeAccessor.h"
 #include "Metadata/Accessors/PCGAttributeAccessorKeys.h"
 #include "Metadata/Accessors/PCGAttributeAccessorHelpers.h"
+#include "Metadata/Accessors/PCGContainerAccessor.h"
 
 #include "GameFramework/Actor.h"
 #include "UObject/Package.h"
@@ -251,17 +252,42 @@ void FPCGContext::InitializeSettings(bool bSkipPostLoad)
 							continue;
 						}
 
-						FSoftObjectPath ObjectPath;
-						if (PCGAttributeAccessorHelpers::GetOverrideParamValue(*AttributeAccessor, ObjectPath))
+						// If we have a container on the param, we have to get all the paths
+						TArray<FSoftObjectPath> ObjectPaths;
+						if (ensure(Param.NumContainers <= 1))
 						{
+							if (Param.NumContainers == 1 && ensure(AccessorResult.Keys.IsValid()))
+							{
+								ObjectPaths.SetNum(AccessorResult.Keys->GetNum());
+								if (!AttributeAccessor->GetRange<FSoftObjectPath>(ObjectPaths, 0, *AccessorResult.Keys, EPCGAttributeAccessorFlags::AllowBroadcastAndConstructible))
+								{
+									continue;
+								}
+							}
+							else
+							{
+								if (!PCGAttributeAccessorHelpers::GetOverrideParamValue(*AttributeAccessor, ObjectPaths.Emplace_GetRef()))
+								{
+									continue;
+								}
+							}
+
 							FGCScopeGuard GCScope;
-							if (!ObjectPath.ResolveObject())
+							for (const FSoftObjectPath& ObjectPath : ObjectPaths)
 							{
-								// We have an override value that is not loaded. We need to schedule the task on the main thread.
-								bOverrideSettingsOnMainThread = true;
-								break;
+								if (!ObjectPath.ResolveObject())
+								{
+									// We have an override value that is not loaded. We need to schedule the task on the main thread.
+									bOverrideSettingsOnMainThread = true;
+									break;
+								}
 							}
 						}
+
+						if (bOverrideSettingsOnMainThread)
+						{
+							break;
+						}
 					}
 				}
 			
@@ -416,45 +442,60 @@ PRAGMA_ENABLE_DEPRECATION_WARNINGS
 				PCGE_LOG_C(Warning, GraphAndLog, this, FText::Format(LOCTEXT("HasMultipleAttributeSetsOnOverridePin", "Multiple attribute sets were found on the {0}. We will use the first one."), OverridePinText));
 			}
 
-			if (AccessorResult.bHasMultipleDataInAttributeSet)
+			if (AccessorResult.bHasMultipleDataInAttributeSet && Param.NumContainers == 0)
 			{
 				PCGE_LOG_C(Warning, GraphAndLog, this, FText::Format(LOCTEXT("HasMultipleDataInAttributeSet", "Multi entry attribute set was found on the {0}. We will only use the first entry to override."), OverridePinText));
 			}
 		}
+		
+		if (Param.NumContainers > 1)
+		{
+			PCGLog::LogWarningOnGraph(FText::Format(LOCTEXT("TooManyContainers", "Property {0} is within multiple containers (i.e. array of array) which is not supported."), FText::FromName(Param.Label)), this);
+			continue;
+		}
+		
+		FPCGAttributeAccessorKeysSingleObjectPtr PropertyObjectKey(Container);
+		
+		TUniquePtr<IPCGAttributeAccessor> PropertyAccessor = Param.NumContainers == 1
+			? PCGContainerAccessorHelpers::MakeContainerAccessor(Param.Properties)
+			: PCGAttributeAccessorHelpers::CreatePropertyChainAccessor(TArray<const FProperty*>(Param.Properties));
+		
+		if (!PropertyAccessor.IsValid() || (Param.NumContainers == 1 && !AccessorResult.Keys))
+		{
+			continue;
+		}
 
-		TUniquePtr<IPCGAttributeAccessor> PropertyAccessor = PCGAttributeAccessorHelpers::CreatePropertyChainAccessor(TArray<const FProperty*>(Param.Properties));
-		check(PropertyAccessor.IsValid());
-
-		const bool bParamOverridden = PCGMetadataAttribute::CallbackWithRightType(PropertyAccessor->GetUnderlyingType(), [this, &AttributeAccessor, &PropertyAccessor, &Param, &AttributeName, Container](auto Dummy) -> bool
+		const bool bParamOverridden = PCGMetadataAttribute::CallbackWithRightType(PropertyAccessor->GetUnderlyingType(), [this, &AttributeAccessor, &PropertyAccessor, &Param, &AttributeName, &AccessorResult, &PropertyObjectKey](auto Dummy) -> bool
 		{
 			using PropertyType = decltype(Dummy);
-
-			// Override were using the first entry (0) by default.
-			FPCGAttributeAccessorKeysEntries FirstEntry(PCGFirstEntryKey);
 					
-			PropertyType Value{};
-			if (!PCGAttributeAccessorHelpers::GetOverrideParamValue(*AttributeAccessor, Value))
+			TArray<PropertyType> Values{};
+			bool bSuccess = false;
+			if (Param.NumContainers == 1 && ensure(AccessorResult.Keys))
+			{
+				bSuccess = PCGAttributeAccessorHelpers::ExtractAllValues<PropertyType>(AttributeAccessor.Get(), AccessorResult.Keys.Get(), Values);
+			}
+			else
+			{
+				bSuccess = PCGAttributeAccessorHelpers::GetOverrideParamValue(*AttributeAccessor, Values.Emplace_GetRef());
+			}
+
+			if (!bSuccess)
 			{
 				PCGE_LOG_C(Warning, GraphAndLog, this, FText::Format(LOCTEXT("ConversionFailed", "Parameter '{0}' ({1}) cannot be converted from attribute '{2}' ({3})"),
 					FText::FromName(Param.Label), 
 					PCG::Private::GetTypeNameText(PCG::Private::MetadataTypes<PropertyType>::Id),
 					FText::FromName(AttributeName),
 					PCG::Private::GetTypeNameText(AttributeAccessor->GetUnderlyingType())));
-				
 				return false;
 			}
 
-			// Setting properties (ex: FSoftObjectPath) can end up doing StaticFindObject which needs to be protected from running at same time as GC
-			FGCScopeGuard GCScope;
-			FPCGAttributeAccessorKeysSingleObjectPtr PropertyObjectKey(Container);
-			PropertyAccessor->Set<PropertyType>(Value, PropertyObjectKey);
-
 			// Add some validation on Object Ptr overrides, to make sure that the object/class stored there has the right class
 			if constexpr (std::is_base_of_v<FSoftObjectPath, PropertyType>)
 			{
 				if (const FObjectPropertyBase* ObjectProperty = CastField<const FObjectPropertyBase>(Param.Properties.Last()))
 				{
-					if (PropertyAccessor->Get<PropertyType>(Value, PropertyObjectKey))
+					for (PropertyType& Value : Values)
 					{
 						bool bInvalid = false;
 						if constexpr (std::is_same_v<FSoftObjectPath, PropertyType>)
@@ -492,12 +533,16 @@ PRAGMA_ENABLE_DEPRECATION_WARNINGS
 						if (bInvalid)
 						{
 							PCGE_LOG_C(Error, GraphAndLog, this, FText::Format(LOCTEXT("WrongObjectClass", "Parameter '{0}' was set with an attribute that is not a child class. It will be nulled out."), FText::FromName(Param.Label)));
-							PropertyAccessor->Set<PropertyType>(PropertyType{}, PropertyObjectKey);
+							Value = PropertyType{};
 						}
 					}
 				}
 			}
 
+			// Setting properties (ex: FSoftObjectPath) can end up doing StaticFindObject which needs to be protected from running at same time as GC
+			FGCScopeGuard GCScope;
+			PropertyAccessor->SetRange<PropertyType>(Values, 0, PropertyObjectKey);
+
 			return true;
 		});
 
diff --git a/Engine/Plugins/PCG/Source/PCG/Private/PCGSettings.cpp b/Engine/Plugins/PCG/Source/PCG/Private/PCGSettings.cpp
--- a/Engine/Plugins/PCG/Source/PCG/Private/PCGSettings.cpp
+++ b/Engine/Plugins/PCG/Source/PCG/Private/PCGSettings.cpp
@@ -13,6 +13,7 @@
 #include "PCGSubsystem.h"
 #include "Elements/PCGAddTag.h"
 #include "Helpers/PCGHelpers.h"
+#include "Helpers/PCGPropertyHelpers.h"
 #include "Helpers/PCGSettingsHelpers.h"
 #include "Metadata/Accessors/IPCGAttributeAccessor.h"
 #include "Metadata/Accessors/PCGAttributeAccessorHelpers.h"
@@ -94,7 +95,7 @@ private:
 
 FString FPCGSettingsOverridableParam::GetPropertyPath() const
 {
-	return FString::JoinBy(Properties, PCGSettings::PropertyPathSeparator, [](const FProperty* InProperty) { return InProperty ? InProperty->GetAuthoredName() : FString(); });
+	return FString::JoinBy(PropertiesNames, PCGSettings::PropertyPathSeparator, [](const FName InName) { return InName.ToString(); });
 }
 
 bool FPCGSettingsOverridableParam::IsHardReferenceOverride() const
@@ -508,7 +509,7 @@ void UPCGSettings::FillOverridableParamsPins(TArray<FPCGPinProperties>& OutPins)
 		ParamPin.SetOverrideOrUserParamPin();
 
 #if WITH_EDITOR
-		ParamPin.Tooltip = LOCTEXT("GlobalParamPinTooltip", "Atribute Set containing multiple parameters to override. Names must match perfectly.");
+		ParamPin.Tooltip = LOCTEXT("GlobalParamPinTooltip", "Attribute Set containing multiple parameters to override. Names must match perfectly.");
 #endif // WITH_EDITOR
 	}
 
@@ -546,10 +547,13 @@ void UPCGSettings::FillOverridableParamsPins(TArray<FPCGPinProperties>& OutPins)
 				Tooltip = *TooltipPtr + TEXT("\n");
 			}
 
-			ParamPin.Tooltip = FText::Format(LOCTEXT("OverridableParamPinTooltip", "{0}Attribute type is \"{1}\" and its exact name is \"{2}\""),
+			const FText ContainerInfo = OverridableParam.NumContainers > 0 ? LOCTEXT("ContainerTooltip", " This is an array.") : FText{};
+
+			ParamPin.Tooltip = FText::Format(LOCTEXT("OverridableParamPinTooltip", "{0}Attribute type is \"{1}\" and its exact name is \"{2}\".{3}"),
 				FText::FromString(Tooltip),
 				FText::FromString(Property->GetCPPType()),
-				FText::FromString(OverridableParam.GetPropertyPath()));
+				FText::FromString(OverridableParam.GetPropertyPath()),
+				ContainerInfo);
 		}
 #endif // WITH_EDITOR
 	}
@@ -1009,6 +1013,7 @@ PRAGMA_ENABLE_DEPRECATION_WARNINGS
 TArray<FPCGSettingsOverridableParam> UPCGSettings::GatherOverridableParams() const
 {
 	PCGSettingsHelpers::FPCGGetAllOverridableParamsConfig Config;
+	Config.bExtractArrays = true;
 	Config.bUseSeed = UseSeed();
 	Config.IncludeMetadataValues.Add(PCGObjectMetadata::Overridable);
 	Config.IncludeMetadataValues.Add(PCGObjectMetadata::OverridableCPUAndGPU);
@@ -1070,8 +1075,7 @@ void UPCGSettings::InitializeCachedOverridableParams(bool bReset)
 					continue;
 				}
 
-				FString Path = FString::JoinBy(Param.PropertiesNames, PCGSettings::PropertyPathSeparator, [](const FName PropertyName) { return PropertyName.ToString();});
-				PropertyPathToLabel.Emplace(Path, Param.Label);
+				PropertyPathToLabel.Emplace(Param.GetPropertyPath(), Param.Label);
 			}
 		}
 
@@ -1086,7 +1090,7 @@ void UPCGSettings::InitializeCachedOverridableParams(bool bReset)
 					continue;
 				}
 
-				FString Path = FString::JoinBy(Param.PropertiesNames, PCGSettings::PropertyPathSeparator, [](const FName PropertyName) { return PropertyName.ToString();});
+				const FString Path = Param.GetPropertyPath();
 				if (FName* OldLabel = PropertyPathToLabel.Find(Path))
 				{
 					if (*OldLabel != Param.Label)
@@ -1128,19 +1132,35 @@ void UPCGSettings::InitializeCachedOverridableParams(bool bReset)
 
 		Param.Properties.Reset(Param.PropertiesNames.Num());
 
+		// Do not support multi containers
+		if (Param.NumContainers > 1)
+		{
+			continue;
+		}
+
 		// Some properties might not be available at runtime. Remove them from the list, since they won't be overridable anymore.
-		const FProperty* CurrentProperty = Param.PropertyClass->FindPropertyByName(Param.PropertiesNames[0]);
+		const FProperty* CurrentProperty = PCGPropertyHelpers::FindPropertyByName(Param.PropertyClass, Param.PropertiesNames[0]);
 		if (!CurrentProperty)
 		{
 			CachedOverridableParams.RemoveAt(i--);
 			continue;
 		}
 
-		Param.Properties.Add(CurrentProperty);
-
-		for (int32 j = 1; j < Param.PropertiesNames.Num(); ++j)
+		for (int32 j = 0; j < Param.PropertiesNames.Num(); ++j)
 		{
-			// If we have multiple depth properties, it should be Struct/Object properties by construction
+			// If we have multiple depth properties, it can be containers. In that case, we have 2 properties for a single name (the container property and inner property)
+			if (const FArrayProperty* ArrayProperty = CastField<FArrayProperty>(CurrentProperty))
+			{
+				Param.Properties.Add(CurrentProperty);
+				CurrentProperty = ArrayProperty->Inner;
+			}
+			else if (const FSetProperty* SetProperty = CastField<FSetProperty>(CurrentProperty))
+			{
+				Param.Properties.Add(CurrentProperty);
+				CurrentProperty = SetProperty->ElementProp;
+			}
+
+			// Otherwise it can be objects/structs.
 			const UStruct* UnderlyingStruct = nullptr;
 			if (const FStructProperty* StructProperty = CastField<FStructProperty>(CurrentProperty))
 			{
@@ -1151,14 +1171,20 @@ void UPCGSettings::InitializeCachedOverridableParams(bool bReset)
 				UnderlyingStruct = ObjectProperty->PropertyClass;
 			}
 
-			if (ensure(UnderlyingStruct))
+			if (UnderlyingStruct && j < Param.PropertiesNames.Num() - 1)
 			{
-				CurrentProperty = UnderlyingStruct->FindPropertyByName(Param.PropertiesNames[j]);
+				Param.Properties.Add(CurrentProperty);
+				CurrentProperty = PCGPropertyHelpers::FindPropertyByName(UnderlyingStruct, Param.PropertiesNames[j + 1]);
 				check(CurrentProperty);
+			}
+			else if (j == Param.PropertiesNames.Num() - 1)
+			{
+				// If we are at the last property we're fine.
 				Param.Properties.Add(CurrentProperty);
 			}
 			else
 			{
+				// Otherwise we're ill-formed, reset the properties.
 				Param.Properties.Empty();
 				break;
 			}
diff --git a/Engine/Plugins/PCG/Source/PCG/Public/Elements/PCGApplyOnActor.h b/Engine/Plugins/PCG/Source/PCG/Public/Elements/PCGApplyOnActor.h
--- a/Engine/Plugins/PCG/Source/PCG/Public/Elements/PCGApplyOnActor.h
+++ b/Engine/Plugins/PCG/Source/PCG/Public/Elements/PCGApplyOnActor.h
@@ -55,6 +55,14 @@ public:
 	UPROPERTY(BlueprintReadWrite, EditAnywhere, Category = Settings)
 	TArray<FName> PostProcessFunctionNames;
 
+	/**
+	 * In case of container override, allow for each element in input, to use the full override data in the Property Overrides pin.
+	 * In that case it is expected to have as many data in the Property Overrides pin as there are elements in the In pin.
+	 * (it's also allowed to have just a single data in Property Overrides, in that case, this data will be used for all the elements in the In pin)
+	 */
+	UPROPERTY(BlueprintReadWrite, EditAnywhere, Category = Settings)
+	bool bUseOneOverrideDataPerElement = false;
+
 	/** Opt-in option to silence errors when the path is Empty or nothing to extract. */
 	UPROPERTY(BlueprintReadWrite, EditAnywhere, Category = Settings, AdvancedDisplay)
 	bool bSilenceErrorOnEmptyObjectPath = false;
diff --git a/Engine/Plugins/PCG/Source/PCG/Public/Helpers/PCGPropertyHelpers.h b/Engine/Plugins/PCG/Source/PCG/Public/Helpers/PCGPropertyHelpers.h
--- a/Engine/Plugins/PCG/Source/PCG/Public/Helpers/PCGPropertyHelpers.h
+++ b/Engine/Plugins/PCG/Source/PCG/Public/Helpers/PCGPropertyHelpers.h
@@ -116,6 +116,13 @@ namespace PCGPropertyHelpers
 	template <typename T>
 	TArray<T> ExtractAttributeSetAsArrayOfStructs(const UPCGParamData* InParamData, const TMap<FName, TTuple<FName, bool>>* OptionalNameMapping = nullptr, FPCGContext* OptionalContext = nullptr);
 
+	/**
+	* More permissive search for a property name if the InStruct is a UserDefinedStruct.
+	* Since the names in editor and at runtime can be different for the properties defined in a UserDefinedStruct, the search will also look for sanitized version of the name if it is not a valid name.
+	* Otherwise just return FindFProperty
+	*/
+	PCG_API const FProperty* FindPropertyByName(const UStruct* InStruct, const FName InName);
+
 	/**
 	* More permissive search for a property name in a UserDefinedStruct.
 	* Since the names in editor and at runtime can be different for the properties defined in a UserDefinedStruct, the search will also look for sanitized version of the name if it is not a valid name.
diff --git a/Engine/Plugins/PCG/Source/PCG/Public/Metadata/Accessors/IPCGAttributeAccessorTpl.h b/Engine/Plugins/PCG/Source/PCG/Public/Metadata/Accessors/IPCGAttributeAccessorTpl.h
--- a/Engine/Plugins/PCG/Source/PCG/Public/Metadata/Accessors/IPCGAttributeAccessorTpl.h
+++ b/Engine/Plugins/PCG/Source/PCG/Public/Metadata/Accessors/IPCGAttributeAccessorTpl.h
@@ -94,8 +94,8 @@ private:
 	{
 		//TRACE_CPUPROFILER_EVENT_SCOPE(IPCGAttributeAccessorT::InternalSetRange);
 
-		// Can't set if read only, and discard any set if there are too many values.
-		if (bReadOnly || (Index + InValues.Num()) > Keys.GetNum())
+		// Can't set if read only, and discard any set that goes beyond the number of keys.
+		if (bReadOnly || Index >= Keys.GetNum())
 		{
 			return false;
 		}
diff --git a/Engine/Plugins/PCG/Source/PCG/Public/Metadata/Accessors/PCGAttributeAccessorHelpers.h b/Engine/Plugins/PCG/Source/PCG/Public/Metadata/Accessors/PCGAttributeAccessorHelpers.h
--- a/Engine/Plugins/PCG/Source/PCG/Public/Metadata/Accessors/PCGAttributeAccessorHelpers.h
+++ b/Engine/Plugins/PCG/Source/PCG/Public/Metadata/Accessors/PCGAttributeAccessorHelpers.h
@@ -73,6 +73,9 @@ namespace PCGAttributeAccessorHelpers
 		bool bPinConnected = false;
 		bool bHasMultipleAttributeSetsOnOverridePin = false;
 		bool bHasMultipleDataInAttributeSet = false;
+
+		// If the param is in a container, we'll need the keys
+		TUniquePtr<const IPCGAttributeAccessorKeys> Keys;
 	};
 
 	/**
diff --git a/Engine/Plugins/PCG/Source/PCG/Public/Metadata/Accessors/PCGPropertyAccessor.h b/Engine/Plugins/PCG/Source/PCG/Public/Metadata/Accessors/PCGPropertyAccessor.h
--- a/Engine/Plugins/PCG/Source/PCG/Public/Metadata/Accessors/PCGPropertyAccessor.h
+++ b/Engine/Plugins/PCG/Source/PCG/Public/Metadata/Accessors/PCGPropertyAccessor.h
@@ -15,8 +15,27 @@ class FEnumProperty;
 
 namespace PCGPropertyAccessor
 {
-	TArray<const void*> GetContainerKeys(int32 Index, int32 Range, const IPCGAttributeAccessorKeys& Keys);
-	TArray<void*> GetContainerKeys(int32 Index, int32 Range, IPCGAttributeAccessorKeys& Keys);
+	template <typename KeyType, typename ...ArrayArgs>
+	void GetContainerKeys(int32 Index, int32 Range, IPCGAttributeAccessorKeys& Keys, TArray<KeyType*, ArrayArgs...>& OutContainer)
+	{
+		OutContainer.SetNumUninitialized(Range);
+		TArrayView<KeyType*> ContainerKeysView(OutContainer);
+		if (!Keys.GetKeys(Index, ContainerKeysView))
+		{
+			OutContainer.Reset();
+		}
+	}
+
+	template <typename KeyType, typename ...ArrayArgs>
+	void GetContainerKeys(int32 Index, int32 Range, const IPCGAttributeAccessorKeys& Keys, TArray<const KeyType*, ArrayArgs...>& OutContainer)
+	{
+		OutContainer.SetNumUninitialized(Range);
+		TArrayView<const KeyType*> ContainerKeysView(OutContainer);
+		if (!Keys.GetKeys(Index, ContainerKeysView))
+		{
+			OutContainer.Reset();
+		}
+	}
 
 	template <typename T>
 	void AddressOffset(const TArray<const FProperty*>& InProperties, TArrayView<T> InContainerKeys)
@@ -47,7 +66,8 @@ namespace PCGPropertyAccessor
 	template <typename T, typename Func>
 	bool IterateGet(const TArray<const FProperty*>& Properties, TArrayView<T>& OutValues, int32 Index, const IPCGAttributeAccessorKeys& Keys, Func&& Getter)
 	{
-		TArray<const void*> ContainerKeys = GetContainerKeys(Index, OutValues.Num(), Keys);
+		TArray<const void*, TInlineAllocator<256>> ContainerKeys;
+		GetContainerKeys(Index, OutValues.Num(), Keys, ContainerKeys);
 		if (ContainerKeys.IsEmpty())
 		{
 			return false;
@@ -67,7 +87,8 @@ namespace PCGPropertyAccessor
 	template <typename T, typename Func>
 	bool IterateSet(const TArray<const FProperty*>& Properties, TArrayView<const T>& InValues, int32 Index, IPCGAttributeAccessorKeys& Keys, Func&& Setter)
 	{
-		TArray<void*> ContainerKeys = GetContainerKeys(Index, InValues.Num(), Keys);
+		TArray<void*, TInlineAllocator<256>> ContainerKeys;
+		GetContainerKeys(Index, InValues.Num(), Keys, ContainerKeys);
 		if (ContainerKeys.IsEmpty())
 		{
 			return false;
diff --git a/Engine/Plugins/PCG/Source/PCG/Public/Metadata/PCGAttributePropertySelector.h b/Engine/Plugins/PCG/Source/PCG/Public/Metadata/PCGAttributePropertySelector.h
--- a/Engine/Plugins/PCG/Source/PCG/Public/Metadata/PCGAttributePropertySelector.h
+++ b/Engine/Plugins/PCG/Source/PCG/Public/Metadata/PCGAttributePropertySelector.h
@@ -101,7 +101,12 @@ public:
 	UE_API bool IsValid() const;
 
 	// Update the selector with an incoming string.
-	UE_API bool Update(const FString& NewValue);
+	UE_API bool Update(const FStringView NewValue);
+
+	// For deprecation of Update(const FString&)
+	// @todo_pcg in another pass
+	// UE_DEPRECATED(5.8, "Make sure to use the string view version")
+	UE_API bool Update(const char* NewValue);
 
 	template <typename T>
 	static T CreateFromOtherSelector(const FPCGAttributePropertySelector& InOther)
@@ -158,7 +163,7 @@ public:
 	}
 
 	template <typename T = FPCGAttributePropertySelector>
-	static T CreateSelectorFromString(const FString& String)
+	static T CreateSelectorFromString(const FStringView String)
 	{
 		static_assert(std::is_base_of_v<FPCGAttributePropertySelector, T>, "The type must be of base class 'FPCGAttributePropertySelector'");
 		T Selector;
diff --git a/Engine/Plugins/PCG/Source/PCG/Public/Metadata/PCGObjectPropertyOverride.h b/Engine/Plugins/PCG/Source/PCG/Public/Metadata/PCGObjectPropertyOverride.h
--- a/Engine/Plugins/PCG/Source/PCG/Public/Metadata/PCGObjectPropertyOverride.h
+++ b/Engine/Plugins/PCG/Source/PCG/Public/Metadata/PCGObjectPropertyOverride.h
@@ -56,7 +56,21 @@ namespace PCGObjectPropertyOverrideHelpers
 	/** Apply property overrides to the TargetObject directly from the ObjectPropertiesOverride pin. Use CreateObjectPropertiesOverridePin(). */
 	PCG_API void ApplyOverridesFromParams(const TArray<FPCGObjectPropertyOverrideDescription>& InObjectPropertyOverrideDescriptions, UObject* TargetObject, FName OverridesPinLabel, FPCGContext* Context, bool bPropagateEditChangeEvent = false);
 
-	PCG_API void ApplyOverrides(const TArray<FPCGObjectPropertyOverrideDescription>& InObjectPropertyOverrideDescriptions, const TArray<TPair<UObject*, int32>>& TargetObjectAndIndex, FName OverridesPinLabel, int32 InputDataIndex, FPCGContext* Context, bool bPropagateEditChangeEvent = false);
+	struct FApplyOverrideParams
+	{
+		TConstArrayView<FPCGObjectPropertyOverrideDescription> ObjectPropertyOverrideDescriptions;
+		TConstArrayView<TPair<UObject*, int32>> TargetObjectToOverrideIndices;
+		TConstArrayView<const UPCGData*> OverridesData;
+		FPCGContext* OptionalContext = nullptr;
+		bool bPropagateEditChangeEvent = false;
+		bool bOneDataSetPerElement = false;
+	};
+
+	PCG_API void ApplyOverrides(const FApplyOverrideParams& InParams);
+
+	// @todo_pcg in another pass
+	// UE_DEPRECATED(5.8, "Use the Params version")
+	PCG_API void ApplyOverrides(const TArray<FPCGObjectPropertyOverrideDescription>& InObjectPropertyOverrideDescriptions, const TArray<TPair<UObject*, int32>>& TargetObjectToOverrideIndices, FName OverridesPinLabel, int32 InputDataIndex, FPCGContext* Context, bool bPropagateEditChangeEvent = false);
 }
 
 /**
@@ -74,7 +88,22 @@ struct FPCGObjectSingleOverride
 	FPCGObjectSingleOverride& operator=(const FPCGObjectSingleOverride&) = delete;
 	FPCGObjectSingleOverride& operator=(FPCGObjectSingleOverride&&) = default;
 
-	/** Initialize the single object override. Call before using Apply(InputKeyIndex, OutputKey). Can also compute the edit property chain if it is needed to propagate the change. Editor only. */
+	struct FInitializeParams
+	{
+		const FPCGAttributePropertySelector* InputSelector = nullptr;
+		FStringView OutputProperty;
+		const UStruct* TemplateClass = nullptr;
+		const UPCGData* SourceData = nullptr;
+		FPCGContext* OptionalContext = nullptr;
+		bool bComputePropertyEditChain = false;
+		bool bAllowedToUseAllSourceDataForContainers = false;
+	};
+
+	/** Initialize the single object override. Call before using Apply(InputKeyIndex, OutputKey). Can also compute the edit property chain if it is needed to propagate the change (Editor only) */
+	PCG_API void Initialize(const FInitializeParams& InParams);
+
+	// @todo_pcg in another pass
+	//UE_DEPRECATED(5.8, "Use Params version")
 	PCG_API void Initialize(const FPCGAttributePropertySelector& InputSelector, const FString& OutputProperty, const UStruct* TemplateClass, const UPCGData* SourceData, FPCGContext* Context, bool bComputePropertyEditChain = false);
 
 	/** Returns true if initialization succeeded in creating the accessors and accessor keys. */
@@ -106,6 +135,8 @@ private:
 #endif // WITH_EDITOR
 
 	bool bWillNeedLoading = false;
+	bool bAllowedToUseAllSourceDataForContainers = false;
+	bool bIsContainer = false;
 
 	// InputKeyIndex, OutputKeys
 	using ApplyOverrideFunction = bool(FPCGObjectSingleOverride::*)(int32, IPCGAttributeAccessorKeys&);
@@ -119,14 +150,23 @@ private:
 			return false;
 		}
 
-		Type Value{};
 		check(ObjectOverrideInputAccessor.IsValid());
-		if (ObjectOverrideInputAccessor->Get<Type>(Value, InputKeyIndex, *InputKeys.Get(), EPCGAttributeAccessorFlags::AllowBroadcast | EPCGAttributeAccessorFlags::AllowConstructible))
+		check(ObjectOverrideOutputAccessor.IsValid());
+
+		if (bAllowedToUseAllSourceDataForContainers && bIsContainer)
 		{
-			check(ObjectOverrideOutputAccessor.IsValid());
-			if (ObjectOverrideOutputAccessor->Set<Type>(Value, OutputKey))
+			TArray<Type> Values;
+			if (PCGAttributeAccessorHelpers::ExtractAllValues(ObjectOverrideInputAccessor.Get(), InputKeys.Get(), Values, EPCGAttributeAccessorFlags::AllowBroadcastAndConstructible))
 			{
-				return true;
+				return ObjectOverrideOutputAccessor->SetRange<Type>(Values, 0, OutputKey);
+			}
+		}
+		else
+		{
+			Type Value{};
+			if (ObjectOverrideInputAccessor->Get<Type>(Value, InputKeyIndex, *InputKeys.Get(), EPCGAttributeAccessorFlags::AllowBroadcastAndConstructible))
+			{
+				return ObjectOverrideOutputAccessor->Set<Type>(Value, OutputKey);
 			}
 		}
 
@@ -149,41 +189,79 @@ struct FPCGObjectOverrides
 	FPCGObjectOverrides& operator=(const FPCGObjectOverrides&) = delete;
 	FPCGObjectOverrides& operator=(FPCGObjectOverrides&&) = default;
 
-	/** Initialize the object overrides. Call before using Apply(InputKeyIndex). */
+	struct FInitializeParams
+	{
+		TConstArrayView<FPCGObjectPropertyOverrideDescription> OverrideDescriptions;
+		T* TemplateObject = nullptr;
+		const UPCGData* SourceData;
+		FPCGContext* OptionalContext = nullptr;
+		bool bInPropagateEditChangeEvent = false;
+		bool bAllowedToUseAllSourceDataForContainers = false;
+	};
+
+	// @todo_pcg in another pass
+	// UE_DEPRECATED(5.8, "Use the Params version")
 	void Initialize(const TArray<FPCGObjectPropertyOverrideDescription>& OverrideDescriptions, T* TemplateObject, const UPCGData* SourceData, FPCGContext* Context, bool bInPropagateEditChangeEvent = false)
+	{
+		FInitializeParams Params =
+		{
+			.OverrideDescriptions = OverrideDescriptions,
+			.TemplateObject = TemplateObject,
+			.SourceData = SourceData,
+			.OptionalContext = Context,
+			.bInPropagateEditChangeEvent = bInPropagateEditChangeEvent
+		};
+
+		Initialize(Params);
+	}
+
+	/** Initialize the object overrides. Call before using Apply(InputKeyIndex). */
+	void Initialize(const FInitializeParams& InParams)
 	{
 		bInitialized = false;
 
-		if (!TemplateObject)
+		if (!InParams.TemplateObject)
 		{
-			PCGLog::LogErrorOnGraph(NSLOCTEXT("PCGObjectPropertyOverride", "InitializeOverrideFailedNoObject", "Failed to initialize property overrides. No template object was provided."), Context);
+			PCGLog::LogErrorOnGraph(NSLOCTEXT("PCGObjectPropertyOverride", "InitializeOverrideFailedNoObject", "Failed to initialize property overrides. No template object was provided."), InParams.OptionalContext);
 			return;
 		}
 
 		CachedOutputObject = nullptr; 
-		OutputKey = FPCGAttributeAccessorKeysSingleObjectPtr<T>(TemplateObject);
-		bPropagateEditChangeEvent = bInPropagateEditChangeEvent;
+		OutputKey = FPCGAttributeAccessorKeysSingleObjectPtr<T>(InParams.TemplateObject);
+		bPropagateEditChangeEvent = InParams.bInPropagateEditChangeEvent;
 
 		const UStruct* ClassObject = nullptr;
 		if constexpr (TIsDerivedFrom<T, UObject>::Value)
 		{
-			ClassObject = TemplateObject->GetClass();
-			CachedOutputObject = bPropagateEditChangeEvent ? TemplateObject : nullptr;
+			ClassObject = InParams.TemplateObject->GetClass();
+			CachedOutputObject = bPropagateEditChangeEvent ? InParams.TemplateObject : nullptr;
 		}
 		else
 		{
-			ClassObject = TemplateObject->StaticStruct();
+			ClassObject = InParams.TemplateObject->StaticStruct();
 		}
 
-		ObjectSingleOverrides.Empty(OverrideDescriptions.Num());
+		ObjectSingleOverrides.Empty(InParams.OverrideDescriptions.Num());
 
-		for (int32 i = 0; i < OverrideDescriptions.Num(); ++i)
+		for (int32 i = 0; i < InParams.OverrideDescriptions.Num(); ++i)
 		{
-			const FPCGAttributePropertyInputSelector InputSelector = OverrideDescriptions[i].InputSource.CopyAndFixLast(SourceData);
-			const FString& OutputProperty = OverrideDescriptions[i].PropertyTarget;
+			const FPCGAttributePropertyInputSelector InputSelector = InParams.OverrideDescriptions[i].InputSource.CopyAndFixLast(InParams.SourceData);
+			const FString& OutputProperty = InParams.OverrideDescriptions[i].PropertyTarget;
 
 			FPCGObjectSingleOverride Override{};
-			Override.Initialize(InputSelector, OutputProperty, ClassObject, SourceData, Context, bPropagateEditChangeEvent);
+
+			FPCGObjectSingleOverride::FInitializeParams Params =
+			{
+				.InputSelector = &InputSelector,
+				.OutputProperty = OutputProperty,
+				.TemplateClass = ClassObject,
+				.SourceData = InParams.SourceData,
+				.OptionalContext = InParams.OptionalContext,
+				.bComputePropertyEditChain = bPropagateEditChangeEvent,
+				.bAllowedToUseAllSourceDataForContainers = InParams.bAllowedToUseAllSourceDataForContainers
+			};
+
+			Override.Initialize(Params);
 
 			if (Override.IsValid())
 			{
@@ -191,7 +269,7 @@ struct FPCGObjectOverrides
 			}
 			else
 			{
-				PCGLog::LogErrorOnGraph(FText::Format(NSLOCTEXT("PCGObjectPropertyOverride", "InitializeOverrideFailed", "Failed to initialize override '{0}' for property {1} on object '{2}'."), InputSelector.GetDisplayText(), FText::FromString(OutputProperty), FText::FromName(ClassObject->GetFName())), Context);
+				PCGLog::LogErrorOnGraph(FText::Format(NSLOCTEXT("PCGObjectPropertyOverride", "InitializeOverrideFailed", "Failed to initialize override '{0}' for property {1} on object '{2}'."), InputSelector.GetDisplayText(), FText::FromString(OutputProperty), FText::FromName(ClassObject->GetFName())), InParams.OptionalContext);
 			}
 		}
 
diff --git a/Engine/Plugins/PCG/Source/PCG/Public/PCGSettings.h b/Engine/Plugins/PCG/Source/PCG/Public/PCGSettings.h
--- a/Engine/Plugins/PCG/Source/PCG/Public/PCGSettings.h
+++ b/Engine/Plugins/PCG/Source/PCG/Public/PCGSettings.h
@@ -128,6 +128,10 @@ struct FPCGSettingsOverridableParam
 	EPCGMetadataTypes UnderlyingType = EPCGMetadataTypes::Unknown;
 #endif // WITH_EDITORONLY_DATA
 
+	/** Indicates the number of containers in the property chain. */
+	UPROPERTY()
+	int32 NumContainers = 0;
+
 	bool HasAliases() const { return !MapOfAliases.IsEmpty(); }
 
 	// Transient
