[PCG] Support overriding Excluded/Included Tags in Scene Capture node

#rnx

diff --git a/Engine/Plugins/PCG/Source/PCG/Public/Elements/PCGSceneCapture.h b/Engine/Plugins/PCG/Source/PCG/Public/Elements/PCGSceneCapture.h
--- a/Engine/Plugins/PCG/Source/PCG/Public/Elements/PCGSceneCapture.h
+++ b/Engine/Plugins/PCG/Source/PCG/Public/Elements/PCGSceneCapture.h
@@ -45,14 +45,12 @@ public:
 	UPROPERTY(BlueprintReadWrite, EditAnywhere, Category = "Settings", meta = (PCG_Overridable))
 	bool bOnlyCaptureContentMatchingIncludedTags = false;
 
-	// @todo_pcg: Allow overriding this array when array overrides are supported.
 	/** Tags to match against actors/components when selecting for inclusion in the scene capture. */
-	UPROPERTY(BlueprintReadWrite, EditAnywhere, Category = "Settings", meta = (EditCondition = "bOnlyCaptureContentMatchingIncludedTags", EditConditionHides))
+	UPROPERTY(BlueprintReadWrite, EditAnywhere, Category = "Settings", meta = (EditCondition = "bOnlyCaptureContentMatchingIncludedTags", EditConditionHides, PCG_Overridable))
 	TArray<FName> IncludedTags;
 
-	// @todo_pcg: Allow overriding this array when array overrides are supported.
 	/** Tags to match against actors/components when selecting for exclusion from the scene capture. */
-	UPROPERTY(BlueprintReadWrite, EditAnywhere, Category = "Settings")
+	UPROPERTY(BlueprintReadWrite, EditAnywhere, Category = "Settings", meta = (PCG_Overridable))
 	TArray<FName> ExcludedTags;
 
 	/** Exclude content created from PCG. */
