[PCG] Support transforms in texture data interface

#rnx

diff --git a/Engine/Plugins/PCG/Shaders/Private/PCGTextureDataInterface.ush b/Engine/Plugins/PCG/Shaders/Private/PCGTextureDataInterface.ush
--- a/Engine/Plugins/PCG/Shaders/Private/PCGTextureDataInterface.ush
+++ b/Engine/Plugins/PCG/Shaders/Private/PCGTextureDataInterface.ush
@@ -35,8 +35,7 @@ PCG_TEXTURE_UAV_DECLARATION(0);
 SamplerState {DataInterfaceName}_Sampler_Linear;
 SamplerState {DataInterfaceName}_Sampler_Point;
 
-// todo_pcg: Change to using texture transform matrix so we support different orientations.
-float4 {DataInterfaceName}_TextureBounds[PCG_MAX_NUM_SRV_BINDINGS];
+float4x4 {DataInterfaceName}_TextureInverseTransforms[PCG_MAX_NUM_SRV_BINDINGS];
 
 /** Layout of TextureInfos buffer:
  * - TextureInfo0
@@ -46,6 +45,8 @@ float4 {DataInterfaceName}_TextureBounds[PCG_MAX_NUM_SRV_BINDINGS];
  *   - uint32 - Sampler
  *   - uint32 - ResolutionX
  *   - uint32 - ResolutionY
+ *   - float  - TexelSizeWorldX (cm)
+ *   - float  - TexelSizeWorldY (cm)
  * - TextureInfo1
  * - TextureInfo2
  * - ...
@@ -134,18 +135,11 @@ bool GetThreadData_{DataInterfaceName}(uint InThreadIndex, out uint OutDataIndex
 
 float2 GetTexCoords_{DataInterfaceName}(uint InDataIndex, float3 InWorldPos)
 {
-	const float2 Min = {DataInterfaceName}_TextureBounds[InDataIndex].xy;
-	const float2 Max = {DataInterfaceName}_TextureBounds[InDataIndex].zw;
+	const float4x4 InverseTransform = {DataInterfaceName}_TextureInverseTransforms[InDataIndex];
+	const float4 LocalPos = mul(float4(InWorldPos, 1.0), InverseTransform);
 
-	const float2 Extents = Max - Min;
-
-	if (!all(Extents > 0.0f))
-	{
-		return (float2)0;
-	}
-
-	// Inverse lerp
-	return (InWorldPos.xy - Min) / Extents;
+	// LocalPos will be in [-1, 1] so transform to a texture coordinate in [0, 1]
+	return LocalPos.xy * 0.5 + 0.5;
 }
 
 float2 GetTexelSize_{DataInterfaceName}(uint InDataIndex)
@@ -158,12 +152,7 @@ float2 GetTexelSize_{DataInterfaceName}(uint InDataIndex)
 
 float2 GetTexelSizeWorld_{DataInterfaceName}(uint InDataIndex)
 {
-	const float2 Min = {DataInterfaceName}_TextureBounds[InDataIndex].xy;
-	const float2 Max = {DataInterfaceName}_TextureBounds[InDataIndex].zw;
-
-	uint2 Resolution;
-	GetTextureResolutionInternal_{DataInterfaceName}(InDataIndex, Resolution);
-	return (Max - Min) / float2(Resolution);
+	return asfloat({DataInterfaceName}_TextureInfos[InDataIndex * 2 + 1].zw);
 }
 
 #define PCG_CONDITIONAL_TEXTURE_SAMPLE(Index) \
diff --git a/Engine/Plugins/PCG/Source/PCG/Private/Compute/DataInterfaces/PCGTextureDataInterface.cpp b/Engine/Plugins/PCG/Source/PCG/Private/Compute/DataInterfaces/PCGTextureDataInterface.cpp
--- a/Engine/Plugins/PCG/Source/PCG/Private/Compute/DataInterfaces/PCGTextureDataInterface.cpp
+++ b/Engine/Plugins/PCG/Source/PCG/Private/Compute/DataInterfaces/PCGTextureDataInterface.cpp
@@ -81,7 +81,6 @@ FPCGTextureBindingInfo::FPCGTextureBindingInfo(const UPCGBaseTextureData* InText
 	Texture = InTextureData->GetTextureRHI();
 	ExportedTexture = InTextureData->GetRefCountedTexture();
 	Transform = InTextureData->GetTransform();
-	TextureBounds = InTextureData->GetBounds();
 	bPointSample = InTextureData->Filter == EPCGTextureFilter::Point;
 
 	if (Texture.IsValid())
@@ -101,7 +100,6 @@ FPCGTextureBindingInfo::FPCGTextureBindingInfo(const FPCGDataDesc& InDataDesc, c
 	// @todo_pcg: Need to support other information here too, like format, slice index, resource type, etc.
 	Size = FIntPoint(FMath::Max(InDataDesc.GetElementCount().X, 1), FMath::Max(InDataDesc.GetElementCount().Y, 1));
 	Transform = InTransform;
-	TextureBounds = FBox(FVector(-1.0f, -1.0f, 0.0f), FVector(1.0f, 1.0f, 0.0f)).TransformBy(Transform);
 }
 
 bool FPCGTextureBindingInfo::operator==(const FPCGTextureBindingInfo& Other) const
@@ -112,8 +110,7 @@ bool FPCGTextureBindingInfo::operator==(const FPCGTextureBindingInfo& Other) con
 		&& Size == Other.Size
 		&& Dimension == Other.Dimension
 		&& bPointSample == Other.bPointSample
-		&& Transform.Equals(Other.Transform)
-		&& TextureBounds == Other.TextureBounds;
+		&& Transform.Equals(Other.Transform);
 }
 
 bool FPCGTextureBindingInfo::IsValid() const
@@ -194,7 +191,7 @@ BEGIN_SHADER_PARAMETER_STRUCT(FPCGTextureDataInterfaceParameters,)
 	SHADER_PARAMETER_RDG_TEXTURE_UAV_ARRAY(RWTexture2D<float4>, TextureUAV, [PCGTextureDataInterfaceConstants::MAX_NUM_UAV_BINDINGS])
 	SHADER_PARAMETER_SAMPLER(SamplerState, Sampler_Linear)
 	SHADER_PARAMETER_SAMPLER(SamplerState, Sampler_Point)
-	SHADER_PARAMETER_ARRAY(FVector4f, TextureBounds, [PCGTextureDataInterfaceConstants::MAX_NUM_SRV_BINDINGS])
+	SHADER_PARAMETER_ARRAY(FMatrix44f, TextureInverseTransforms, [PCGTextureDataInterfaceConstants::MAX_NUM_SRV_BINDINGS])
 	SHADER_PARAMETER(uint32, NumTextureInfos)
 	SHADER_PARAMETER_RDG_BUFFER_SRV(StructuredBuffer<FUintVector4>, TextureInfos)
 END_SHADER_PARAMETER_STRUCT()
@@ -458,15 +455,12 @@ void FPCGTextureDataProviderProxy::GatherDispatchData(FDispatchData const& InDis
 
 			if (BindingIndex < BindingInfos.Num())
 			{
-				Parameters.TextureBounds[BindingIndex] = FVector4f(
-					BindingInfos[BindingIndex].TextureBounds.Min.X,
-					BindingInfos[BindingIndex].TextureBounds.Min.Y,
-					BindingInfos[BindingIndex].TextureBounds.Max.X,
-					BindingInfos[BindingIndex].TextureBounds.Max.Y);
+				const FTransform& Transform = BindingInfos[BindingIndex].Transform;
+				Parameters.TextureInverseTransforms[BindingIndex] = !FMath::IsNearlyZero(Transform.GetDeterminant()) ? FMatrix44f(Transform.ToInverseMatrixWithScale()) : FMatrix44f::Identity;
 			}
 			else
 			{
-				Parameters.TextureBounds[BindingIndex] = FVector4f::Zero();
+				Parameters.TextureInverseTransforms[BindingIndex] = FMatrix44f::Identity;
 			}
 		}
 
@@ -671,6 +665,8 @@ void FPCGTextureDataProviderProxy::PackTextureInfos(FRDGBuilder& GraphBuilder)
 	//   - uint32 - Sampler
 	//   - uint32 - ResolutionX
 	//   - uint32 - ResolutionY
+	//   - float  - TexelSizeWorldX (cm)
+	//   - float  - TexelSizeWorldY (cm)
 	// - TextureInfo1
 	// - TextureInfo2
 	// - ...
@@ -689,8 +685,12 @@ void FPCGTextureDataProviderProxy::PackTextureInfos(FRDGBuilder& GraphBuilder)
 		{
 			const FPCGTextureInfo& TextureInfo = TextureInfos[TextureInfoIndex];
 			const FPCGTextureBindingInfo& BindingInfo = BindingInfos[TextureInfo.BindingIndex];
+			const FVector& TextureWorldScale = BindingInfo.Transform.GetScale3D();
+			const float TexelWorldSizeX = BindingInfo.Size.X != 0 ? (TextureWorldScale.X * 2.0) / BindingInfo.Size.X : 0.0f;
+			const float TexelWorldSizeY = BindingInfo.Size.Y != 0 ? (TextureWorldScale.Y * 2.0) / BindingInfo.Size.Y : 0.0f;
+
 			TextureInfosPackedData[TextureInfoIndex * 2 + 0] = FUintVector4(TextureInfo.BindingIndex, TextureInfo.SliceIndex, static_cast<uint32>(BindingInfo.Dimension), BindingInfo.bPointSample);
-			TextureInfosPackedData[TextureInfoIndex * 2 + 1] = FUintVector4(BindingInfo.Size.X, BindingInfo.Size.Y, 0u, 0u);
+			TextureInfosPackedData[TextureInfoIndex * 2 + 1] = FUintVector4(BindingInfo.Size.X, BindingInfo.Size.Y, FMath::AsUInt(TexelWorldSizeX), FMath::AsUInt(TexelWorldSizeY));
 		}
 
 		GraphBuilder.QueueBufferUpload(TextureInfosBuffer, MakeArrayView(TextureInfosPackedData));
diff --git a/Engine/Plugins/PCG/Source/PCG/Private/Compute/DataInterfaces/PCGTextureDataInterface.h b/Engine/Plugins/PCG/Source/PCG/Private/Compute/DataInterfaces/PCGTextureDataInterface.h
--- a/Engine/Plugins/PCG/Source/PCG/Private/Compute/DataInterfaces/PCGTextureDataInterface.h
+++ b/Engine/Plugins/PCG/Source/PCG/Private/Compute/DataInterfaces/PCGTextureDataInterface.h
@@ -39,9 +39,7 @@ struct FPCGTextureBindingInfo
 	ETextureDimension Dimension = ETextureDimension::Texture2D;
 	bool bPointSample = false;
 
-	// @todo_pcg: Eventually replace TextureBounds with only the Transform
 	FTransform Transform = FTransform::Identity;
-	FBox TextureBounds = FBox(EForceInit::ForceInit);
 };
 
 struct FPCGTextureInfo
