[PCG] Added node to add an index-based value to tags and/or attribute

diff --git a/Engine/Plugins/PCG/Source/PCG/Private/Elements/PCGWriteDataIndex.cpp b/Engine/Plugins/PCG/Source/PCG/Private/Elements/PCGWriteDataIndex.cpp
new file mode 100644
--- /dev/null
+++ b/Engine/Plugins/PCG/Source/PCG/Private/Elements/PCGWriteDataIndex.cpp
@@ -0,0 +1,79 @@
+// Copyright Epic Games, Inc. All Rights Reserved.
+
+#include "Elements/PCGWriteDataIndex.h"
+
+#include "PCGContext.h"
+#include "PCGModule.h"
+#include "Elements/Metadata/PCGMetadataElementCommon.h"
+
+#include UE_INLINE_GENERATED_CPP_BY_NAME(PCGWriteDataIndex)
+
+#define LOCTEXT_NAMESPACE "PCGWriteDataIndexElement"
+
+UPCGWriteDataIndexSettings::UPCGWriteDataIndexSettings()
+{
+	IndexAttribute = FPCGAttributePropertySelector::CreateAttributeSelector<FPCGAttributePropertyOutputSelector>(TEXT("Index"), TEXT("Data"));
+}
+
+TArray<FPCGPinProperties> UPCGWriteDataIndexSettings::OutputPinProperties() const
+{
+	TArray<FPCGPinProperties> PinProperties;
+	PinProperties.Emplace(PCGPinConstants::DefaultOutputLabel, EPCGDataType::Any);
+
+	return PinProperties;
+}
+
+bool FPCGWriteDataIndexElement::ExecuteInternal(FPCGContext* Context) const
+{
+	TRACE_CPUPROFILER_EVENT_SCOPE(FPCGWriteDataIndexElement::Execute);
+
+	check(Context);
+
+	const UPCGWriteDataIndexSettings* Settings = Context->GetInputSettings<UPCGWriteDataIndexSettings>();
+	check(Settings);
+
+	TArray<FPCGTaggedData> InputData = Context->InputData.GetInputsByPin(PCGPinConstants::DefaultInputLabel);
+	TArray<FPCGTaggedData>& OutputData = Context->OutputData.TaggedData;
+
+	const FName IndexAttributeName = Settings->IndexAttribute.GetName();
+
+	// Validate that the output tag is correct or passthrough otherwise
+	if (Settings->bAddTag && Settings->IndexTag.IsEmpty())
+	{
+		PCGLog::LogErrorOnGraph(LOCTEXT("Invalid tag", "The index tag must not be empty."), Context);
+		OutputData = InputData;
+		return true;
+	}
+
+	for(int DataIndex = 0; DataIndex < InputData.Num(); ++DataIndex)
+	{
+		const FPCGTaggedData& Input = InputData[DataIndex];
+		FPCGTaggedData& OutputTaggedData = OutputData.Emplace_GetRef(Input);
+
+		const UPCGData* Data = Input.Data;
+
+		if (Settings->bAddAttribute)
+		{
+			UPCGData* NewData = Data->DuplicateData(Context);
+			FPCGMetadataDomain* IndexAttributeDomain = NewData ? NewData->MutableMetadata()->GetMetadataDomainFromSelector(Settings->IndexAttribute) : nullptr;
+			FPCGMetadataAttribute<int32>* IndexAttribute = IndexAttributeDomain ? IndexAttributeDomain->FindOrCreateAttribute<int32>(IndexAttributeName, DataIndex, /*bAllowInterpolation=*/false) : nullptr;
+
+			if (!IndexAttribute)
+			{
+				PCGLog::LogErrorOnGraph(FText::Format(LOCTEXT("FailedCreateIndexAttribute", "Failed to create the index attribute '{0}'."), FText::FromName(IndexAttributeName)), Context);
+				continue;
+			}
+
+			OutputTaggedData.Data = NewData;
+		}
+
+		if (Settings->bAddTag)
+		{
+			OutputTaggedData.Tags.Add(FString::Format(TEXT("{0}:{1}"), { Settings->IndexTag, DataIndex }));
+		}
+	}
+
+	return true;
+}
+
+#undef LOCTEXT_NAMESPACE
\ No newline at end of file
diff --git a/Engine/Plugins/PCG/Source/PCG/Public/Elements/PCGWriteDataIndex.h b/Engine/Plugins/PCG/Source/PCG/Public/Elements/PCGWriteDataIndex.h
new file mode 100644
--- /dev/null
+++ b/Engine/Plugins/PCG/Source/PCG/Public/Elements/PCGWriteDataIndex.h
@@ -0,0 +1,53 @@
+// Copyright Epic Games, Inc. All Rights Reserved.
+
+#pragma once
+
+#include "PCGSettings.h"
+
+#include "PCGWriteDataIndex.generated.h"
+
+class FPCGWriteDataIndexElement : public IPCGElement
+{
+protected:
+	virtual bool ExecuteInternal(FPCGContext* Context) const override;
+	virtual bool SupportsBasePointDataInputs(FPCGContext* InContext) const override { return true; }
+	// @todo_pcg: support gpu resident data
+};
+
+/**
+* Writes the data's index in the input collection to a tag or to an attribute
+*/
+UCLASS(BlueprintType, ClassGroup=(Procedural))
+class UPCGWriteDataIndexSettings : public UPCGSettings
+{
+	GENERATED_BODY()
+
+public:
+	UPCGWriteDataIndexSettings();
+
+	//~Begin UPCGSettings interface
+#if WITH_EDITOR
+	virtual FName GetDefaultNodeName() const override { return FName(TEXT("WriteDataIndex")); }
+	virtual FText GetDefaultNodeTitle() const override { return NSLOCTEXT("PCGWriteDataIndexElement", "NodeTitle", "Write Data Index"); }
+	virtual FText GetNodeTooltipText() const override { return NSLOCTEXT("PCGWriteDataIndexElement", "NodeTooltip", "Writes the data's index in the input collection to a tag or an attribute."); }
+	virtual EPCGSettingsType GetType() const override { return EPCGSettingsType::Generic; }
+#endif // WITH_EDITOR
+
+protected:
+	virtual TArray<FPCGPinProperties> OutputPinProperties() const override;
+	virtual FPCGElementPtr CreateElement() const override { return MakeShared<FPCGWriteDataIndexElement>(); }
+	//~End UPCGSettings interface
+
+public:
+	UPROPERTY(EditAnywhere, BlueprintReadWrite, Category = Settings, meta = (InlineEditConditionToggle, PCG_Overridable))
+	bool bAddTag = true;
+
+	UPROPERTY(EditAnywhere, BlueprintReadWrite, Category = Settings, meta = (EditCondition = "bAddTag", PCG_Overridable))
+	FString IndexTag = "Index";
+
+	UPROPERTY(EditAnywhere, BlueprintReadWrite, Category = Settings, meta = (InlineEditConditionToggle, PCG_Overridable))
+	bool bAddAttribute = false;
+
+	UPROPERTY(EditAnywhere, BlueprintReadWrite, Category = Settings, meta = (EditCondition = "bAddAttribute", PCG_Overridable, PCG_DiscardPropertySelection, PCG_DiscardExtraSelection))
+	FPCGAttributePropertyOutputSelector IndexAttribute;
+};
\ No newline at end of file
