[PCG] Made write data index node pins dynamically typed.

#rnx

diff --git a/Engine/Plugins/PCG/Source/PCG/Public/Elements/PCGWriteDataIndex.h b/Engine/Plugins/PCG/Source/PCG/Public/Elements/PCGWriteDataIndex.h
--- a/Engine/Plugins/PCG/Source/PCG/Public/Elements/PCGWriteDataIndex.h
+++ b/Engine/Plugins/PCG/Source/PCG/Public/Elements/PCGWriteDataIndex.h
@@ -32,6 +32,7 @@ public:
 	virtual FText GetNodeTooltipText() const override { return NSLOCTEXT("PCGWriteDataIndexElement", "NodeTooltip", "Writes the data's index in the input collection to a tag or an attribute."); }
 	virtual EPCGSettingsType GetType() const override { return EPCGSettingsType::Generic; }
 #endif // WITH_EDITOR
+	virtual bool HasDynamicPins() const override { return true; }
 
 protected:
 	virtual TArray<FPCGPinProperties> OutputPinProperties() const override;
