[PCG] Clear PCG Link only support on PCGComponents for now

#rnx

diff --git a/Engine/Plugins/PCG/Source/PCGEditor/Private/Widgets/SPCGExecutionSourceActionWidget.cpp b/Engine/Plugins/PCG/Source/PCGEditor/Private/Widgets/SPCGExecutionSourceActionWidget.cpp
--- a/Engine/Plugins/PCG/Source/PCGEditor/Private/Widgets/SPCGExecutionSourceActionWidget.cpp
+++ b/Engine/Plugins/PCG/Source/PCGEditor/Private/Widgets/SPCGExecutionSourceActionWidget.cpp
@@ -230,9 +230,13 @@ EVisibility SPCGExecutionSourceActionWidget::GetClearPCGLinkButtonVisibility() c
 {
 	for (const TWeakInterfacePtr<IPCGGraphExecutionSource>& ExecutionSource : ExecutionSources)
 	{
-		if (ExecutionSource.IsValid() && !ExecutionSource->GetExecutionState().IsManagedByRuntimeGenSystem())
+		// @todo_pcg: support this in IPCGGraphExecutionSource
+		if (UPCGComponent* Component = Cast<UPCGComponent>(ExecutionSource.Get()))
 		{
-			return EVisibility::Visible;
+			if (IsValid(Component) && !Component->IsManagedByRuntimeGenSystem())
+			{
+				return EVisibility::Visible;
+			}
 		}
 	}
 
