[PCG] Move UPCGManagedResource management out of UPCGComponent and in FPCGManagedResourceContainer
- Reduces code duplication for new execution source types

diff --git a/Engine/Plugins/Experimental/PCGInterops/PCGFastGeoInterop/Source/PCGFastGeoInterop/Private/Compute/PrimitiveFactories/PCGPrimitiveFactoryFastGeoPISMC.cpp b/Engine/Plugins/Experimental/PCGInterops/PCGFastGeoInterop/Source/PCGFastGeoInterop/Private/Compute/PrimitiveFactories/PCGPrimitiveFactoryFastGeoPISMC.cpp
--- a/Engine/Plugins/Experimental/PCGInterops/PCGFastGeoInterop/Source/PCGFastGeoInterop/Private/Compute/PrimitiveFactories/PCGPrimitiveFactoryFastGeoPISMC.cpp
+++ b/Engine/Plugins/Experimental/PCGInterops/PCGFastGeoInterop/Source/PCGFastGeoInterop/Private/Compute/PrimitiveFactories/PCGPrimitiveFactoryFastGeoPISMC.cpp
@@ -148,7 +148,10 @@ bool FPCGPrimitiveFactoryFastGeoPISMC::Create(FPCGContext* InContext)
 	ManagedPrimitives->SetFastGeoContainer(FastGeo);
 	ManagedPrimitives->SetObjectReferences(CollectObjectReferences());
 
-	ExecutionSource->GetExecutionState().AddToManagedResources(ManagedPrimitives);
+	if (FPCGManagedResourceContainer* Container = ExecutionSource->GetExecutionState().GetManagedResourceContainer(); ensure(Container))
+	{
+		Container->AddManagedResource(ManagedPrimitives);
+	}
 
 	return true;
 }
diff --git a/Engine/Plugins/PCG/Source/PCG/Private/PCGComponent.cpp b/Engine/Plugins/PCG/Source/PCG/Private/PCGComponent.cpp
--- a/Engine/Plugins/PCG/Source/PCG/Private/PCGComponent.cpp
+++ b/Engine/Plugins/PCG/Source/PCG/Private/PCGComponent.cpp
@@ -71,11 +71,6 @@ namespace PCGComponent
 	}
 
 #if WITH_EDITOR
-	static TAutoConsoleVariable<bool> CVarLoadManagedActorsForCleanup(
-		TEXT("pcg.LoadManagedActorsForCleanup"),
-		true,
-		TEXT("When enabled unloaded actors will get loaded before clean up"));
-
 	/** Have a CVar for dynamic tracking for runtime until we have a better solution to also have "culling" for local components and dynamic tracking. */
 	static TAutoConsoleVariable<bool> CVarDisableDynamicTrackingForRuntimeGen(
 		TEXT("pcg.RuntimeGeneration.DisableDynamicTracking"),
@@ -98,63 +93,6 @@ namespace PCGComponent
 		true,
 		TEXT("Delete generated preview resources (actors) when PCG Component gets unloaded."));
 
-	struct FPCGManagedActorLoadingScope
-	{
-		FPCGManagedActorLoadingScope()
-		{
-#if WITH_EDITOR
-			bShouldLoadActors = ensure(!IsGarbageCollecting()) && CVarLoadManagedActorsForCleanup.GetValueOnAnyThread();
-#endif
-		}
-
-		FPCGManagedActorLoadingScope(bool bInShouldLoadActors)
-		{
-#if WITH_EDITOR
-			bShouldLoadActors = ensure(!IsGarbageCollecting()) && bInShouldLoadActors;
-#endif
-		}
-
-		void AddResources(UPCGComponent* InComponent, const TArray<UPCGManagedResource*>& InManagedResources)
-		{
-#if WITH_EDITOR
-			if (!bShouldLoadActors)
-			{
-				return;
-			}
-
-			if (!InComponent || !InComponent->GetWorld() || InComponent->GetWorld()->IsGameWorld())
-			{
-				return;
-			}
-
-			UWorldPartition* WorldPartition = FWorldPartitionHelpers::GetWorldPartition(InComponent);
-			if (!WorldPartition)
-			{
-				return;
-			}
-
-			for (UPCGManagedResource* ManagedResource : InManagedResources)
-			{
-				if (UPCGManagedActors* ManagedActors = Cast<UPCGManagedActors>(ManagedResource))
-				{
-					for (const TSoftObjectPtr<AActor>& ManagedActor : ManagedActors->GetConstGeneratedActors())
-					{
-						if (const FWorldPartitionActorDescInstance* ActorDescInstance = WorldPartition->GetActorDescInstanceByPath(ManagedActor.ToSoftObjectPath()))
-						{
-							LoadedActors.Add(FWorldPartitionReference(WorldPartition, ActorDescInstance->GetGuid()));
-						}
-					}
-				}
-			}
-#endif // WITH_EDITOR 
-		}
-
-#if WITH_EDITOR
-		TSet<FWorldPartitionReference> LoadedActors;
-		bool bShouldLoadActors;
-#endif
-	};
-
 	[[nodiscard]] double RoundTrackingPriority(double InTrackingPriority)
 	{
 		// Round to 4 decimals (10^4)
@@ -166,6 +104,7 @@ namespace PCGComponent
 UPCGComponent::UPCGComponent(const FObjectInitializer& InObjectInitializer)
 	: Super(InObjectInitializer)
 {
+	ManagedResourceContainer.SetOwner(this);
 	ExecutionState.Component = this;
 	GraphInstance = InObjectInitializer.CreateDefaultSubobject<UPCGGraphInstance>(this, TEXT("PCGGraphInstance"));
 	SchedulingPolicyClass = UPCGSchedulingPolicyDistanceAndDirection::StaticClass();
@@ -303,18 +242,17 @@ void UPCGComponent::SetGraphInterfaceLocal(UPCGGraphInterface* InGraphInterface)
 
 void UPCGComponent::AddToManagedResources(UPCGManagedResource* InResource)
 {
-	PCGGeneratedResourcesLogging::LogAddToManagedResources(this, InResource);
+	ManagedResourceContainer.AddManagedResource(InResource);
+}
 
+void UPCGComponent::OnManagedResourceAdded(UPCGManagedResource* InResource)
+{
 	if (InResource)
 	{
-		if(!ensure(InResource->GetOuter() == this))
+		if (!ensure(InResource->GetOuter() == this))
 		{
 			UPCGBlueprintHelpers::ThrowBlueprintException(LOCTEXT("ResourceNotOutered", "Managed resources need to be outered to their PCG component."));
 		}
-
-		UE::TScopeLock ResourcesLock(GeneratedResourcesLock);
-		check(!GeneratedResourcesInaccessible);
-		GeneratedResources.Add(InResource);
 	}
 }
 
@@ -335,10 +273,10 @@ void UPCGComponent::AddComponentsToManagedResources(const TArray<UActorComponent
 	}
 
 	UE::TScopeLock ResourcesLock(GeneratedResourcesLock);
-	check(!GeneratedResourcesInaccessible);
+	ensure(ManagedResourceContainer.AreResourcesAccessible());
 
 	UPCGManagedComponentDefaultList* DefaultList = nullptr;
-	for (const TObjectPtr<UPCGManagedResource>& ManagedResource : GeneratedResources)
+	for (const TObjectPtr<UPCGManagedResource>& ManagedResource : ManagedResourceContainer.GetConstManagedResourcesNoLock())
 	{
 		if (UPCGManagedComponentDefaultList* ExistingList = Cast<UPCGManagedComponentDefaultList>(ManagedResource))
 		{
@@ -352,7 +290,7 @@ void UPCGComponent::AddComponentsToManagedResources(const TArray<UActorComponent
 		DefaultList = NewObject<UPCGManagedComponentDefaultList>(this);
 
 		PCGGeneratedResourcesLogging::LogAddToManagedResources(this, DefaultList);
-		GeneratedResources.Add(DefaultList);
+		ManagedResourceContainer.GetManagedResourcesNoLock().Add(DefaultList);
 	}
 
 	check(DefaultList);
@@ -395,59 +333,17 @@ void UPCGComponent::AddActorsToManagedResources(const TArray<AActor*>& InActors)
 
 void UPCGComponent::ForEachManagedResource(TFunctionRef<void(UPCGManagedResource*)> InFunction)
 {
-	UE::TScopeLock ResourcesLock(GeneratedResourcesLock);
-	check(!GeneratedResourcesInaccessible);
-	for (const TObjectPtr<UPCGManagedResource>& ManagedResource : GeneratedResources)
-	{
-		if (ManagedResource)
-		{
-			InFunction(ManagedResource);
-		}
-	}
+	ManagedResourceContainer.ForEachManagedResource(InFunction);
 }
 
 void UPCGComponent::ForEachConstManagedResource(TFunctionRef<void(const UPCGManagedResource*)> InFunction) const
 {
-	UE::TScopeLock ResourcesLock(GeneratedResourcesLock);
-	check(!GeneratedResourcesInaccessible);
-	for (const TObjectPtr<UPCGManagedResource>& ManagedResource : GeneratedResources)
-	{
-		if (ManagedResource)
-		{
-			InFunction(ManagedResource);
-		}
-	}
+	ManagedResourceContainer.ForEachConstManagedResource(InFunction);
 }
 
 bool UPCGComponent::IsAnyObjectManagedByResource(const TArrayView<const UObject*> InObjects) const
 {
-	TRACE_CPUPROFILER_EVENT_SCOPE(UPCGComponent::IsAnyObjectManagedByResource);
-	
-	if (GeneratedResourcesInaccessible)
-	{
-		ensure(false);
-		return false;
-	}
-
-	bool bIsManagedByResource = false;
-	ForEachConstManagedResource([&bIsManagedByResource, &InObjects](const UPCGManagedResource* ManagedResource)
-	{
-		if (bIsManagedByResource || !ManagedResource)
-		{
-			return;
-		}
-
-		for (const UObject* Object : InObjects)
-		{
-			if (ManagedResource->IsManaging(Object))
-			{
-				bIsManagedByResource = true;
-				break;
-			}
-		}
-	});
-	
-	return bIsManagedByResource;
+	return ManagedResourceContainer.IsAnyObjectManagedByResource(InObjects);
 }
 
 bool UPCGComponent::ShouldGenerate(bool bForce, EPCGComponentGenerationTrigger RequestedGenerationTrigger) const
@@ -1067,24 +963,45 @@ void UPCGComponent::PurgeUnlinkedResources(const AActor* InActor)
 		}
 	});
 	
-	ThisActor->ForEachComponent<UPCGComponent>(/*bIncludeFromChildActors=*/true, [&ActorComponentList, &ActorsToDelete](UPCGComponent* Component)
+	bool bCanPurge = true;
+
+	ThisActor->ForEachComponent<UPCGComponent>(/*bIncludeFromChildActors=*/true, [&ActorComponentList, &ActorsToDelete, &bCanPurge](UPCGComponent* Component)
 	{
-		for (UPCGManagedResource* ManagedResource : Component->GeneratedResources)
+		if(!bCanPurge)
 		{
-			if (UPCGManagedComponent* ManagedComponent = Cast<UPCGManagedComponent>(ManagedResource))
-			{
-				ActorComponentList.RemoveSwap(ManagedComponent->GeneratedComponent.Get());
-			}
-			else if (UPCGManagedActors* ManagedActors = Cast<UPCGManagedActors>(ManagedResource))
+			return;
+		}
+
+		UE::TScopeLock ResourcesLock(Component->GeneratedResourcesLock);
+		if(ensure(Component->ManagedResourceContainer.AreResourcesAccessible()))
+		{
+			for (UPCGManagedResource* ManagedResource : Component->ManagedResourceContainer.GetManagedResourcesNoLock())
 			{
-				for (const TSoftObjectPtr<AActor>& GeneratedActor : ManagedActors->GetConstGeneratedActors())
+				if (UPCGManagedComponent* ManagedComponent = Cast<UPCGManagedComponent>(ManagedResource))
 				{
-					ActorsToDelete.Remove(GeneratedActor);
+					ActorComponentList.RemoveSwap(ManagedComponent->GeneratedComponent.Get());
+				}
+				else if (UPCGManagedActors* ManagedActors = Cast<UPCGManagedActors>(ManagedResource))
+				{
+					for (const TSoftObjectPtr<AActor>& GeneratedActor : ManagedActors->GetConstGeneratedActors())
+					{
+						ActorsToDelete.Remove(GeneratedActor);
+					}
 				}
 			}
 		}
+		else
+		{
+			bCanPurge = false;
+		}
 	});
 
+	if (!bCanPurge)
+	{
+		UE_LOG(LogPCG, Error, TEXT("[UPCGComponent::PurgeUnlinkedResources] Coouldd not purge resources because some were inacessible."));
+		return;
+	}
+
 	for (UActorComponent* ActorComponent : ActorComponentList)
 	{
 		ActorComponent->DestroyComponent();
@@ -1155,7 +1072,7 @@ FPCGTaskId UPCGComponent::CleanupLocal(bool bRemoveComponents, const TArray<FPCG
 		return InvalidPCGTaskId;
 	}
 
-	const bool bNeedsLocalCleanup = bGenerated || IsGenerating() || !GeneratedResources.IsEmpty();
+	const bool bNeedsLocalCleanup = bGenerated || IsGenerating() || !ManagedResourceContainer.IsEmpty();
 	if (!bNeedsLocalCleanup && Subsystem->GetPCGComponentPartitionActorMappings(this).IsEmpty())
 	{
 		return InvalidPCGTaskId;
@@ -1452,15 +1369,15 @@ bool UPCGComponent::MoveResourcesToNewActor(AActor* InNewActor, bool bCreateChil
 	TSet<TSoftObjectPtr<AActor>> MovedActors;
 
 	// Do not rely on CVar to load actors when Clearing PCG Link
-	PCGComponent::FPCGManagedActorLoadingScope Scope(/*bInShouldLoadActors=*/true);
+	FPCGManagedActorLoadingScope Scope(/*bInShouldLoadActors=*/true);
 #endif
 	{
 		UE::TScopeLock ResourcesLock(GeneratedResourcesLock);
-		check(!GeneratedResourcesInaccessible);
+		ensure(ManagedResourceContainer.AreResourcesAccessible());
 #if WITH_EDITOR
-		Scope.AddResources(this, GeneratedResources);
+		Scope.AddResources(this, ManagedResourceContainer.GetConstManagedResourcesNoLock());
 #endif
-		for (TObjectPtr<UPCGManagedResource>& GeneratedResource : GeneratedResources)
+		for (TObjectPtr<UPCGManagedResource>& GeneratedResource : ManagedResourceContainer.GetManagedResourcesNoLock())
 		{
 			if (GeneratedResource)
 			{
@@ -1482,7 +1399,7 @@ bool UPCGComponent::MoveResourcesToNewActor(AActor* InNewActor, bool bCreateChil
 			}
 		}
 
-		GeneratedResources.Empty();
+		ManagedResourceContainer.GetManagedResourcesNoLock().Empty();
 	}
 
 	if (CreatedChildActor)
@@ -1540,7 +1457,7 @@ bool UPCGComponent::MoveResourcesToNewActor(AActor* InNewActor, bool bCreateChil
 
 void UPCGComponent::CleanupLocalImmediate(bool bRemoveComponents, bool bCleanupLocalComponents)
 {
-	PCGGeneratedResourcesLogging::LogCleanupLocalImmediate(this, bRemoveComponents, GeneratedResources);
+	PCGGeneratedResourcesLogging::LogCleanupLocalImmediate(this, bRemoveComponents, ManagedResourceContainer.GetManagedResourcesCopy());
 
 	UPCGSubsystem* Subsystem = GetSubsystem();
 
@@ -1557,43 +1474,7 @@ void UPCGComponent::CleanupLocalImmediate(bool bRemoveComponents, bool bCleanupL
 	// Cancels generation of this component if there is an ongoing generation in progress.
 	CancelGeneration();
 
-	{
-		TSet<TSoftObjectPtr<AActor>> ActorsToDelete;
-		PCGComponent::FPCGManagedActorLoadingScope Scope;
-
-		if (!bRemoveComponents && UPCGManagedResource::DebugForcePurgeAllResourcesOnGenerate())
-		{
-			bRemoveComponents = true;
-		}
-
-		{
-			UE::TScopeLock ResourcesLock(GeneratedResourcesLock);
-			check(!GeneratedResourcesInaccessible);
-			Scope.AddResources(this, GeneratedResources);
-			for (int32 ResourceIndex = GeneratedResources.Num() - 1; ResourceIndex >= 0; --ResourceIndex)
-			{
-				// Note: resources can be null here in some loading + bp object cases
-				UPCGManagedResource* Resource = GeneratedResources[ResourceIndex];
-
-				PCGGeneratedResourcesLogging::LogCleanupLocalImmediateResource(this, Resource);
-
-				if (!Resource || Resource->Release(bRemoveComponents, ActorsToDelete))
-				{
-#if WITH_EDITOR
-					if (Resource && !Resource->IsMarkedTransientOnLoad())
-					{
-						Resource->Rename(nullptr, GetTransientPackage(), REN_DoNotDirty | REN_DontCreateRedirectors | REN_NonTransactional);
-					}
-#endif
-
-					GeneratedResources.RemoveAt(ResourceIndex);
-				}
-
-			}
-		}
-
-		UPCGActorHelpers::DeleteActors(GetWorld(), ActorsToDelete.Array());
-	}
+	ManagedResourceContainer.Cleanup(bRemoveComponents);
 
 	PostCleanupGraph(bRemoveComponents);
 
@@ -1603,7 +1484,7 @@ void UPCGComponent::CleanupLocalImmediate(bool bRemoveComponents, bool bCleanupL
 		Subsystem->CleanupLocalComponentsImmediate(this, bRemoveComponents);
 	}
 
-	PCGGeneratedResourcesLogging::LogCleanupLocalImmediateFinished(this, GeneratedResources);
+	PCGGeneratedResourcesLogging::LogCleanupLocalImmediateFinished(this, ManagedResourceContainer.GetManagedResourcesCopy());
 }
 
 FPCGTaskId UPCGComponent::CreateCleanupTask(bool bRemoveComponents, const TArray<FPCGTaskId>& Dependencies)
@@ -1613,7 +1494,7 @@ FPCGTaskId UPCGComponent::CreateCleanupTask(bool bRemoveComponents, const TArray
 		UE_LOG(LogPCG, Log, TEXT("[%s] --- CLEANUP COMPONENT ---"), GetOwner() ? *GetOwner()->GetName() : TEXT("MissingComponent"));
 	}
 
-	if ((!bGenerated && GeneratedResources.IsEmpty() && !IsGenerating()) || IsCleaningUp())
+	if ((!bGenerated && ManagedResourceContainer.IsEmpty() && !IsGenerating()) || IsCleaningUp())
 	{
 		return InvalidPCGTaskId;
 	}
@@ -1632,169 +1513,17 @@ FPCGTaskId UPCGComponent::CreateCleanupTask(bool bRemoveComponents, const TArray
 		AllDependencies = &AdditionalDependencies;
 	}
 
-	struct FCleanupContext : public PCGComponent::FPCGManagedActorLoadingScope
-	{
-		bool bIsFirstIteration = true;
-		int32 ResourceIndex = -1;
-		TSet<TSoftObjectPtr<AActor>> ActorsToDelete;
-	};
+	TSharedPtr<FPCGManagedResourceContainer::FCleanupTask> ContainerCleanup = ManagedResourceContainer.CreateCleanupTask();
 
-	TSharedPtr<FCleanupContext> Context = MakeShared<FCleanupContext>();
-	TWeakObjectPtr<UPCGComponent> ThisComponentWeakPtr(this);
-	TWeakObjectPtr<UWorld> WorldPtr(GetWorld());
-
-	auto AbortCleanup = [Context, ThisComponentWeakPtr]()
+	auto AbortCleanup = [ContainerCleanup]()
 	{
-		UPCGComponent* ThisComponent = ThisComponentWeakPtr.Get();
-		// If the component is not valid anymore, just early out
-		if (!IsValid(ThisComponent))
-		{
-			return true;
-		}
-
-		UE::TScopeLock ResourcesLock(ThisComponent->GeneratedResourcesLock);
-		ThisComponent->GeneratedResourcesInaccessible = false;
+		ContainerCleanup->Abort();
 		return true;
 	};
 
-	auto CleanupTask = [Context, ThisComponentWeakPtr, WorldPtr, bRemoveComponents]()
+	auto CleanupTask = [ContainerCleanup, bRemoveComponents]()
 	{
-		if (UPCGComponent* ThisComponent = ThisComponentWeakPtr.Get())
-		{
-			// If the component is not valid anymore, just early out
-			if (!IsValid(ThisComponent))
-			{
-				return true;
-			}
-
-			UE::TScopeLock ResourcesLock(ThisComponent->GeneratedResourcesLock);
-
-			// Safeguard to track illegal modifications of the generated resources array while doing cleanup
-			if (Context->bIsFirstIteration)
-			{
-				check(!ThisComponent->GeneratedResourcesInaccessible);
-				ThisComponent->GeneratedResourcesInaccessible = true;
-				Context->ResourceIndex = ThisComponent->GeneratedResources.Num() - 1;
-				Context->bIsFirstIteration = false;
-				Context->AddResources(ThisComponent, ThisComponent->GeneratedResources);
-			}
-
-			// Going backward
-			if (Context->ResourceIndex >= 0)
-			{
-				UPCGManagedResource* Resource = ThisComponent->GeneratedResources[Context->ResourceIndex];
-
-				if (!Resource && ThisComponent->GetOwner())
-				{
-					UE_LOG(LogPCG, Error, TEXT("[UPCGComponent::CreateCleanupTask] Null generated resource encountered on actor \"%s\"."), *ThisComponent->GetOwner()->GetFName().ToString());
-				}
-
-				PCGGeneratedResourcesLogging::LogCreateCleanupTaskResource(ThisComponentWeakPtr.Get(), Resource);
-
-				if (!Resource || Resource->Release(bRemoveComponents, Context->ActorsToDelete))
-				{
-#if WITH_EDITOR
-					if (Resource && !Resource->IsMarkedTransientOnLoad())
-					{
-						Resource->Rename(nullptr, GetTransientPackage(), REN_DoNotDirty | REN_DontCreateRedirectors | REN_NonTransactional);
-					}
-#endif
-
-					ThisComponent->GeneratedResources.RemoveAtSwap(Context->ResourceIndex);
-				}
-
-				Context->ResourceIndex--;
-
-				// Returning false means the task is not over
-				return false;
-			}
-			else
-			{
-				ThisComponent->GeneratedResourcesInaccessible = false;
-			}
-		}
-
-		TSet<FName> DeletedActorFolders;
-
-		if (UWorld* World = WorldPtr.Get())
-		{
-			const TArray<TSoftObjectPtr<AActor>> ActorsToDelete = Context->ActorsToDelete.Array();
-
-#if WITH_EDITOR
-			for (const TSoftObjectPtr<AActor>& Actor : ActorsToDelete)
-			{
-				if (Actor.IsValid())
-				{
-					FName ActorFolderPath = Actor->GetFolderPath();
-					if (ActorFolderPath != NAME_None)
-					{
-						DeletedActorFolders.Add(ActorFolderPath);
-					}
-				}
-			}
-#endif
-
-			UPCGActorHelpers::DeleteActors(World, ActorsToDelete);
-		}
-
-		if (UPCGComponent* ThisComponent = ThisComponentWeakPtr.Get())
-		{
-#if WITH_EDITOR
-			if (UWorld* ThisWorld = ThisComponent->GetWorld(); ThisWorld && GEditor) // FActorFolders require the editor
-			{
-				for (FName FolderPath : DeletedActorFolders)
-				{
-					FFolder GeneratedFolder(FFolder::GetWorldRootFolder(ThisWorld).GetRootObject(), FolderPath);
-					const bool bFolderExists = GeneratedFolder.IsValid() && FActorFolders::Get().ContainsFolder(*ThisWorld, GeneratedFolder);
-					bool bFoundActors = false;
-
-					if (bFolderExists)
-					{
-						TSet<FName> Folders{ GeneratedFolder.GetPath() };
-						FActorFolders::ForEachActorInFolders(*ThisWorld, Folders, [&bFoundActors](AActor* InActor)
-						{
-							if (InActor)
-							{
-								bFoundActors = true;
-								return false;
-							}
-							else
-							{
-								return true;
-							}
-						});
-					}
-
-					if (bFolderExists && !bFoundActors)
-					{
-						// Delete all subfolders
-						TArray<FFolder> SubfoldersToDelete;
-						FActorFolders::Get().ForEachFolder(*ThisWorld, [&GeneratedFolder, &SubfoldersToDelete](const FFolder& InFolder)
-						{
-							if (InFolder.IsChildOf(GeneratedFolder))
-							{
-								SubfoldersToDelete.Add(InFolder);
-							}
-
-							return true;
-						});
-
-						for (const FFolder& FolderToDelete : SubfoldersToDelete) //-V1078
-						{
-							FActorFolders::Get().DeleteFolder(*ThisWorld, FolderToDelete);
-						}
-
-						// Finally, delete folder
-						FActorFolders::Get().DeleteFolder(*ThisWorld, GeneratedFolder);
-					}
-				}
-			}
-#endif
-
-			PCGGeneratedResourcesLogging::LogCreateCleanupTaskFinished(ThisComponentWeakPtr.Get(), ThisComponentWeakPtr.IsValid() ? &ThisComponentWeakPtr->GeneratedResources : nullptr);
-		}
-
-		return true;
+		return ContainerCleanup->Execute(bRemoveComponents);
 	};
 
 	return GetSubsystem()->ScheduleGeneric(CleanupTask, AbortCleanup, this, *AllDependencies);
@@ -1802,40 +1531,7 @@ FPCGTaskId UPCGComponent::CreateCleanupTask(bool bRemoveComponents, const TArray
 
 void UPCGComponent::CleanupUnusedManagedResources()
 {
-	PCGGeneratedResourcesLogging::LogCleanupUnusedManagedResources(this, GeneratedResources);
-
-	TSet<TSoftObjectPtr<AActor>> ActorsToDelete;
-
-	{
-		UE::TScopeLock ResourcesLock(GeneratedResourcesLock);
-		check(!GeneratedResourcesInaccessible);
-		for (int32 ResourceIndex = GeneratedResources.Num() - 1; ResourceIndex >= 0; --ResourceIndex)
-		{
-			UPCGManagedResource* Resource = GetValid(GeneratedResources[ResourceIndex]);
-
-			PCGGeneratedResourcesLogging::LogCleanupUnusedManagedResourcesResource(this, Resource);
-
-			if (!Resource && GetOwner())
-			{
-				UE_LOG(LogPCG, Error, TEXT("[UPCGComponent::CleanupUnusedManagedResources] Null generated resource encountered on actor \"%s\"."), *GetOwner()->GetFName().ToString());
-			}
-
-			if (!Resource || Resource->ReleaseIfUnused(ActorsToDelete))
-			{
-#if WITH_EDITOR
-				if (Resource && !Resource->IsMarkedTransientOnLoad())
-				{
-					Resource->Rename(nullptr, GetTransientPackage(), REN_DoNotDirty | REN_DontCreateRedirectors | REN_NonTransactional);
-				}
-#endif
-				GeneratedResources.RemoveAtSwap(ResourceIndex);
-			}
-		}
-	}
-
-	UPCGActorHelpers::DeleteActors(GetWorld(), ActorsToDelete.Array());
-
-	PCGGeneratedResourcesLogging::LogCleanupUnusedManagedResourcesFinished(this, GeneratedResources);
+	ManagedResourceContainer.CleanupUnusedManagedResources();
 }
 
 void UPCGComponent::ClearGraphGeneratedOutput(bool bClearLoadedPreviewData)
@@ -1918,32 +1614,7 @@ void UPCGComponent::EndPlay(const EEndPlayReason::Type EndPlayReason)
 		Subsystem->UnregisterPCGComponent(this);
 	}
 
-	if (ensure(!GeneratedResourcesInaccessible))
-	{
-		UE::TScopeLock ResourcesLock(GeneratedResourcesLock);
-
-		for (int32 ResourceIndex = GeneratedResources.Num() - 1; ResourceIndex >= 0; --ResourceIndex)
-		{
-			if (!GeneratedResources[ResourceIndex])
-			{
-				// Remove null entries.
-				GeneratedResources.RemoveAtSwap(ResourceIndex);
-			}
-			else if (GeneratedResources[ResourceIndex]->ReleaseOnTeardown())
-			{
-				if (GeneratedResources[ResourceIndex])
-				{
-					TSet<TSoftObjectPtr<AActor>> ActorsToDelete;
-					GeneratedResources[ResourceIndex]->Release(/*bHardRelease=*/true, ActorsToDelete);
-
-					// Don't support deleting actors during teardown.
-					ensure(ActorsToDelete.IsEmpty());
-				}
-
-				GeneratedResources.RemoveAtSwap(ResourceIndex);
-			}
-		}
-	}
+	ManagedResourceContainer.EndPlay();
 
 	Super::EndPlay(EndPlayReason);
 }
@@ -1976,14 +1647,20 @@ void UPCGComponent::OnUnregister()
 	if (bUnregisteredThroughLoading && PCGComponent::CVarDeletePreviewResourcesWhenUnloading.GetValueOnAnyThread() && CurrentEditingMode == EPCGEditorDirtyMode::Preview)
 	{
 		TSet<TSoftObjectPtr<AActor>> ActorsToDelete;
-		for (UPCGManagedResource* ManagedResource : GeneratedResources)
 		{
-			// Make sure resource is indeed a transient/preview resource
-			if (ManagedResource &&
-				!ManagedResource->IsMarkedTransientOnLoad() &&
-				!LoadedPreviewResources.Contains(ManagedResource))
+			UE::TScopeLock ResourcesLock(GeneratedResourcesLock);
+			if (ensure(ManagedResourceContainer.AreResourcesAccessible()))
 			{
-				ManagedResource->Release(/*bHardRelease=*/true, ActorsToDelete);
+				for (UPCGManagedResource* ManagedResource : ManagedResourceContainer.GetManagedResourcesNoLock())
+				{
+					// Make sure resource is indeed a transient/preview resource
+					if (ManagedResource &&
+						!ManagedResource->IsMarkedTransientOnLoad() &&
+						!ManagedResourceContainer.GetConstLoadedPreviewResourcesNoLock().Contains(ManagedResource))
+					{
+						ManagedResource->Release(/*bHardRelease=*/true, ActorsToDelete);
+					}
+				}
 			}
 		}
 
@@ -2105,8 +1782,7 @@ void UPCGComponent::Serialize(FArchive& Ar)
 
 	if (!Ar.IsLoading() && CurrentEditingMode == EPCGEditorDirtyMode::Preview)
 	{
-		GeneratedResourcesCopy = GeneratedResources;
-		GeneratedResources = LoadedPreviewResources;
+		ManagedResourceContainer.BeginPreviewSave(GeneratedResourcesCopy);
 
 		GeneratedGraphOutputCopy = GeneratedGraphOutput;
 		GeneratedGraphOutput = LoadedPreviewGeneratedGraphOutput;
@@ -2148,7 +1824,7 @@ void UPCGComponent::Serialize(FArchive& Ar)
 #if WITH_EDITOR
 	if (!Ar.IsLoading() && CurrentEditingMode == EPCGEditorDirtyMode::Preview)
 	{
-		GeneratedResources = GeneratedResourcesCopy;
+		ManagedResourceContainer.EndPreviewSave(GeneratedResourcesCopy);
 		GeneratedGraphOutput = GeneratedGraphOutputCopy;
 	}
 #endif // WITH_EDITOR
@@ -2160,8 +1836,17 @@ void UPCGComponent::PostLoad()
 	Super::PostLoad();
 
 #if WITH_EDITOR
+
+PRAGMA_DISABLE_DEPRECATION_WARNINGS
+	if (!GeneratedResources_DEPRECATED.IsEmpty())
+	{
+		ManagedResourceContainer.SetManagedResources(MoveTemp(GeneratedResources_DEPRECATED));
+	}
+PRAGMA_ENABLE_DEPRECATION_WARNINGS
+
 	// Debug resources are transient and will generate null entries on save, clean them up here
-	GeneratedResources.Remove(nullptr);
+	TArray<TObjectPtr<UPCGManagedResource>>& ContainerGeneratedResources = ManagedResourceContainer.GetManagedResourcesNoLock();
+	ContainerGeneratedResources.Remove(nullptr);
 
 	// Force dirty to be false on load. We should never refresh on load.
 	bDirtyGenerated = false;
@@ -2193,7 +1878,7 @@ void UPCGComponent::PostLoad()
 	}
 
 	/** Deprecation code, should be removed once generated data has been updated */
-	if (GetOwner() && bGenerated && GeneratedResources.Num() == 0)
+	if (GetOwner() && bGenerated && ContainerGeneratedResources.Num() == 0)
 	{
 		TArray<UInstancedStaticMeshComponent*> ISMCs;
 		GetOwner()->GetComponents(ISMCs);
@@ -2204,7 +1889,7 @@ void UPCGComponent::PostLoad()
 			{
 				UPCGManagedISMComponent* ManagedComponent = NewObject<UPCGManagedISMComponent>(this);
 				ManagedComponent->GeneratedComponent = ISMC;
-				GeneratedResources.Add(ManagedComponent);
+				ContainerGeneratedResources.Add(ManagedComponent);
 			}
 		}
 
@@ -2212,7 +1897,7 @@ void UPCGComponent::PostLoad()
 		{
 			UPCGManagedActors* ManagedActors = NewObject<UPCGManagedActors>(this);
 			ManagedActors->GetMutableGeneratedActors() = GeneratedActors_DEPRECATED.Array();
-			GeneratedResources.Add(ManagedActors);
+			ContainerGeneratedResources.Add(ManagedActors);
 			GeneratedActors_DEPRECATED.Reset();
 		}
 	}
@@ -3126,7 +2811,7 @@ void UPCGComponent::OnRefresh(bool bForceRefresh)
 	const bool bWasGeneratedOrGenerating = bWasGenerated || bForceRefresh || IsGenerating();
 
 	// If we are partitioned but we have resources, we need to force a cleanup
-	if (IsPartitioned() && !GeneratedResources.IsEmpty())
+	if (IsPartitioned() && !ManagedResourceContainer.IsEmpty())
 	{
 		CleanupLocalImmediate(/*bRemoveComponents=*/true);
 	}
@@ -4122,29 +3807,12 @@ void UPCGComponent::ResetIgnoredChangeOrigins(bool bLogIfAnyPresent)
 
 void UPCGComponent::SetManagedResources(const TArray<TObjectPtr<UPCGManagedResource>>& Resources)
 {
-	UE::TScopeLock ResourcesLock(GeneratedResourcesLock);
-
-	// We expect the GeneratedResources to be empty here, as otherwise they might not be taken care of properly - they
-	// will be lost down below, but this should not happen. However, if the GeneratedResources are marked as Visible,
-	// then they will be copied over during BP duplication, hence why this will happen, hence the ensure here.
-	ensure(GeneratedResources.IsEmpty());
-
-	GeneratedResources = Resources;
-
-	// Remove any null entries
-	for (int32 ResourceIndex = GeneratedResources.Num() - 1; ResourceIndex >= 0; --ResourceIndex)
-	{
-		if (!GeneratedResources[ResourceIndex])
-		{
-			GeneratedResources.RemoveAtSwap(ResourceIndex);
-		}
-	}
+	ManagedResourceContainer.SetManagedResources(Resources);
 }
 
 void UPCGComponent::GetManagedResources(TArray<TObjectPtr<UPCGManagedResource>>& Resources) const
 {
-	UE::TScopeLock ResourcesLock(GeneratedResourcesLock);
-	Resources = GeneratedResources;
+	Resources =	ManagedResourceContainer.GetManagedResourcesCopy();
 }
 
 void UPCGComponent::SetEditingMode(EPCGEditorDirtyMode InEditingMode, EPCGEditorDirtyMode InSerializedEditingMode)
@@ -4156,51 +3824,18 @@ void UPCGComponent::SetEditingMode(EPCGEditorDirtyMode InEditingMode, EPCGEditor
 #if WITH_EDITOR
 bool UPCGComponent::DeletePreviewResources()
 {
-	bool bResourceWasReleased = false;
-
-	TSet<TSoftObjectPtr<AActor>> ActorsToDelete;
-	// Make sure to release fully the resources that were loaded
-	for (TObjectPtr<UPCGManagedResource> ResourceToRelease : LoadedPreviewResources)
-	{
-		if (!GeneratedResources.Contains(ResourceToRelease))
-		{
-			// Changing the transient state will clear the "marked transient on load" flag
-			ResourceToRelease->ChangeTransientState(EPCGEditorDirtyMode::Normal);
-			ResourceToRelease->Release(/*bHardRelease=*/true, ActorsToDelete);
-			bResourceWasReleased = true;
-		}
-
-		// Either this resource was released in the previous code block or prior to calling DeletePreviewResources its transient state was changed back to normal
-		ensure(!ResourceToRelease->IsMarkedTransientOnLoad());
-	}
-
-	LoadedPreviewResources.Empty();
-
-	if (!ActorsToDelete.IsEmpty())
-	{
-		UPCGActorHelpers::DeleteActors(GetWorld(), ActorsToDelete.Array());
-	}
-
-	return bResourceWasReleased;
+	return ManagedResourceContainer.DeletePreviewResources();
 }
 
 void UPCGComponent::MarkResourcesAsTransientOnLoad()
 {
-	for (TObjectPtr<UPCGManagedResource>& GeneratedResource : GeneratedResources)
-	{
-		if (GeneratedResource)
-		{
-			GeneratedResource->MarkTransientOnLoad();
-		}
-	}
+	ManagedResourceContainer.MarkResourcesAsTransientOnLoad();
 
 	// Force duplicate owner for PIE so that transient changes are reflected in game
 	if (AActor* Owner = GetOwner())
 	{
 		Owner->bForceDuplicateInPIE = true;
 	}
-
-	LoadedPreviewResources = GeneratedResources;
 }
 
 void UPCGComponent::ChangeTransientState(EPCGEditorDirtyMode NewEditingMode)
@@ -4209,38 +3844,12 @@ void UPCGComponent::ChangeTransientState(EPCGEditorDirtyMode NewEditingMode)
 	// Instead, UPCGComponent will react to CurrentEditingMode changes between PreEditUndo/PostEditUndo and call ChangeTransientState accordingly to reflect the change to CurrentEditingMode.
 	TGuardValue<ITransaction*> SuppressTransaction(GUndo, nullptr);
 
-	bool bShouldMarkDirty = false;
+	bool bShouldMarkDirty = ManagedResourceContainer.ChangeTransientState(NewEditingMode);
 
-	// Affect all resources
+	if (NewEditingMode != EPCGEditorDirtyMode::Preview)
 	{
-		UE::TScopeLock ResourcesLock(GeneratedResourcesLock);
-		check(!GeneratedResourcesInaccessible);
-
-		for (TObjectPtr<UPCGManagedResource>& GeneratedResource : GeneratedResources)
-		{
-			if (GeneratedResource)
-			{
-				// Avoid changing transient state when switching to preview mode on marked as transient resources.
-				// When serializing a component we need its marked as transient resources to remain unchanged so they get serialized properly.
-				if (NewEditingMode == EPCGEditorDirtyMode::Preview && GeneratedResource->IsMarkedTransientOnLoad())
-				{
-					continue;
-				}
-
-				GeneratedResource->ChangeTransientState(NewEditingMode);
-				bShouldMarkDirty = true;
-			}
-		}
-
-		// If switching from preview mode to normal or preview-on-load,
-		// we must materialize any kind of change we've done on the packages that had a different behavior on load (e.g. actor packages)
-		if (NewEditingMode != EPCGEditorDirtyMode::Preview)
-		{
-			bShouldMarkDirty |= DeletePreviewResources();
-
-			bShouldMarkDirty |= (LoadedPreviewGeneratedGraphOutput != GeneratedGraphOutput);
-			LoadedPreviewGeneratedGraphOutput.Reset();
-		}
+		bShouldMarkDirty |= (LoadedPreviewGeneratedGraphOutput != GeneratedGraphOutput);
+		LoadedPreviewGeneratedGraphOutput.Reset();
 	}
 
 	if (IsLocalComponent())
@@ -4312,11 +3921,15 @@ TArray<TSoftObjectPtr<AActor>> UPCGComponent::GetManagedActorPaths(AActor* InAct
 	TSet<TSoftObjectPtr<AActor>> ManagedActorPaths;
 	InActor->ForEachComponent<UPCGComponent>(/*bIncludeFromChildActors=*/true, [&ManagedActorPaths](UPCGComponent* Component)
 	{
-		for (UPCGManagedResource* ManagedResource : Component->GeneratedResources)
+		PCG::TScopeLock ResourceLock(Component->GeneratedResourcesLock);
+		if (ensure(Component->ManagedResourceContainer.AreResourcesAccessible()))
 		{
-			if (UPCGManagedActors* ManagedActors = Cast<UPCGManagedActors>(ManagedResource))
+			for (const UPCGManagedResource* ManagedResource : Component->ManagedResourceContainer.GetConstManagedResourcesNoLock())
 			{
-				ManagedActorPaths.Append(ManagedActors->GetConstGeneratedActors());
+				if (const UPCGManagedActors* ManagedActors = Cast<UPCGManagedActors>(ManagedResource))
+				{
+					ManagedActorPaths.Append(ManagedActors->GetConstGeneratedActors());
+				}
 			}
 		}
 	});
@@ -4389,6 +4002,7 @@ void FPCGComponentInstanceData::ApplyToComponent(UActorComponent* Component, con
 		PCGComponent->PreviousEditingMode = PreviousEditingMode;	
 		PCGComponent->CurrentEditingMode = CurrentEditingMode;
 		PCGComponent->bGeneratedOffline = bGeneratedOffline;
+		PCGComponent->ManagedResourceContainer.SetOwner(PCGComponent);
 
 		for (const FDynamicTrackedKeyInstanceData& DynamicTrackedKeyInstanceData : DynamicallyTrackedKeysToSettings)
 		{
@@ -4406,7 +4020,7 @@ void FPCGComponentInstanceData::ApplyToComponent(UActorComponent* Component, con
 			TMap<TObjectPtr<UPCGManagedResource>, TObjectPtr<UPCGManagedResource>> GeneratedResourceMapping;
 #endif
 			TArray<TObjectPtr<UPCGManagedResource>> DuplicatedResources;
-			for (const TObjectPtr<UPCGManagedResource>& Resource : ConstructionSourceComponent->GeneratedResources)
+			for (const TObjectPtr<UPCGManagedResource>& Resource : ConstructionSourceComponent->ManagedResourceContainer.GetConstManagedResourcesNoLock())
 			{
 				if (Resource)
 				{
@@ -4430,7 +4044,7 @@ void FPCGComponentInstanceData::ApplyToComponent(UActorComponent* Component, con
 			PCGComponent->bDirtyGenerated = ConstructionSourceComponent->bDirtyGenerated;
 
 			TArray<TObjectPtr<UPCGManagedResource>> DuplicateLoadedPreviewResources;
-			for (const TObjectPtr<UPCGManagedResource>& Resource : ConstructionSourceComponent->LoadedPreviewResources)
+			for (const TObjectPtr<UPCGManagedResource>& Resource : ConstructionSourceComponent->ManagedResourceContainer.GetConstLoadedPreviewResourcesNoLock())
 			{
 				if (Resource)
 				{
@@ -4450,7 +4064,7 @@ void FPCGComponentInstanceData::ApplyToComponent(UActorComponent* Component, con
 
 			if (DuplicateLoadedPreviewResources.Num() > 0)
 			{
-				PCGComponent->LoadedPreviewResources = DuplicateLoadedPreviewResources;
+				PCGComponent->ManagedResourceContainer.SetLoadedPreviewResourcesNoLock(DuplicateLoadedPreviewResources);
 			}
 
 			PCGComponent->LoadedPreviewGeneratedGraphOutput = ConstructionSourceComponent->LoadedPreviewGeneratedGraphOutput;
diff --git a/Engine/Plugins/PCG/Source/PCG/Private/PCGComponentExecutionState.cpp b/Engine/Plugins/PCG/Source/PCG/Private/PCGComponentExecutionState.cpp
--- a/Engine/Plugins/PCG/Source/PCG/Private/PCGComponentExecutionState.cpp
+++ b/Engine/Plugins/PCG/Source/PCG/Private/PCGComponentExecutionState.cpp
@@ -199,10 +199,22 @@ const FPCGDataCollection* FPCGComponentExecutionState::RetrieveOutputDataForPin(
 	return Component->RetrieveOutputDataForPin(InResourceKey);
 }
 
-void FPCGComponentExecutionState::AddToManagedResources(UPCGManagedResource* InResource)
+void FPCGComponentExecutionState::OnManagedResourceAdded(UPCGManagedResource* InResource)
 {
 	check(Component);
-	return Component->AddToManagedResources(InResource);
+	Component->OnManagedResourceAdded(InResource);
+}
+
+FPCGManagedResourceContainer* FPCGComponentExecutionState::GetManagedResourceContainer()
+{
+	check(Component);
+	return &Component->ManagedResourceContainer;
+}
+
+FTransactionallySafeCriticalSection* FPCGComponentExecutionState::GetManagedResourceContainerLock()
+{
+	check(Component);
+	return &Component->GeneratedResourcesLock;
 }
 
 FPCGTaskId FPCGComponentExecutionState::Generate(const FGenerateParams& InGenerateParams)
diff --git a/Engine/Plugins/PCG/Source/PCG/Private/PCGGraphExecutionStateInterface.cpp b/Engine/Plugins/PCG/Source/PCG/Private/PCGGraphExecutionStateInterface.cpp
--- a/Engine/Plugins/PCG/Source/PCG/Private/PCGGraphExecutionStateInterface.cpp
+++ b/Engine/Plugins/PCG/Source/PCG/Private/PCGGraphExecutionStateInterface.cpp
@@ -3,6 +3,7 @@
 #include "PCGGraphExecutionStateInterface.h"
 
 #include "PCGModule.h"
+#include "PCGManagedResourceContainer.h"
 #include "Editor/IPCGEditorModule.h"
 #include "Subsystems/PCGSubsystem.h"
 #include "Subsystems/PCGEngineSubsystem.h"
@@ -11,7 +12,14 @@
 
 void IPCGGraphExecutionState::AddToManagedResources(UPCGManagedResource* InResource)
 {
-	UE_LOG(LogPCG, Error, TEXT("This execution source (%s) does not support managed resources."), *GetDebugName());
+	if (FPCGManagedResourceContainer* Container = GetManagedResourceContainer())
+	{
+		Container->AddManagedResource(InResource);
+	}
+	else
+	{
+		UE_LOG(LogPCG, Error, TEXT("This execution source (%s) does not support managed resources."), *GetDebugName());
+	}
 }
 
 IPCGBaseSubsystem* IPCGGraphExecutionState::GetSubsystem() const
diff --git a/Engine/Plugins/PCG/Source/PCG/Private/PCGManagedResource.cpp b/Engine/Plugins/PCG/Source/PCG/Private/PCGManagedResource.cpp
--- a/Engine/Plugins/PCG/Source/PCG/Private/PCGManagedResource.cpp
+++ b/Engine/Plugins/PCG/Source/PCG/Private/PCGManagedResource.cpp
@@ -15,6 +15,11 @@
 #include "UObject/Package.h"
 #include "Utils/PCGGeneratedResourcesLogging.h"
 
+#if WITH_EDITOR
+#include "WorldPartition/WorldPartition.h"
+#include "WorldPartition/WorldPartitionHelpers.h"
+#endif
+
 #include UE_INLINE_GENERATED_CPP_BY_NAME(PCGManagedResource)
 
 static TAutoConsoleVariable<bool> CVarForceReleaseResourcesOnGenerate(
@@ -22,6 +27,11 @@ static TAutoConsoleVariable<bool> CVarForceReleaseResourcesOnGenerate(
 	false,
 	TEXT("Purges all tracked generated resources on generate"));
 
+static TAutoConsoleVariable<bool> CVarLoadManagedActorsForCleanup(
+	TEXT("pcg.LoadManagedActorsForCleanup"),
+	true,
+	TEXT("When enabled unloaded actors will get loaded before clean up"));
+
 void UPCGManagedResource::PostApplyToComponent()
 {
 	// Nothing - apply to component should already properly remap most of everything we need to do
@@ -74,6 +84,57 @@ void UPCGManagedResource::ChangeTransientState(EPCGEditorDirtyMode NewEditingMod
 }
 #endif // WITH_EDITOR
 
+FPCGManagedActorLoadingScope::FPCGManagedActorLoadingScope()
+{
+#if WITH_EDITOR
+	bShouldLoadActors = ensure(!IsGarbageCollecting()) && CVarLoadManagedActorsForCleanup.GetValueOnAnyThread();
+#endif
+}
+
+FPCGManagedActorLoadingScope::FPCGManagedActorLoadingScope(bool bInShouldLoadActors)
+{
+#if WITH_EDITOR
+	bShouldLoadActors = ensure(!IsGarbageCollecting()) && bInShouldLoadActors;
+#endif
+}
+
+void FPCGManagedActorLoadingScope::AddResources(IPCGGraphExecutionSource* InExecutionSource, const TArray<UPCGManagedResource*>& InManagedResources)
+{
+#if WITH_EDITOR
+	if (!bShouldLoadActors)
+	{
+		return;
+	}
+
+	if (!InExecutionSource || !InExecutionSource->GetExecutionState().GetWorld() || InExecutionSource->GetExecutionState().GetWorld()->IsGameWorld())
+	{
+		return;
+	}
+
+	// @todo_pcg: check that this is valid way of getting WP here
+	UWorldPartition* WorldPartition = FWorldPartitionHelpers::GetWorldPartition(InExecutionSource->GetExecutionState().GetWorld());
+	if (!WorldPartition)
+	{
+		return;
+	}
+
+	for (UPCGManagedResource* ManagedResource : InManagedResources)
+	{
+		if (UPCGManagedActors* ManagedActors = Cast<UPCGManagedActors>(ManagedResource))
+		{
+			for (const TSoftObjectPtr<AActor>& ManagedActor : ManagedActors->GetConstGeneratedActors())
+			{
+				if (const FWorldPartitionActorDescInstance* ActorDescInstance = WorldPartition->GetActorDescInstanceByPath(ManagedActor.ToSoftObjectPath()))
+				{
+					LoadedActors.Add(FWorldPartitionReference(WorldPartition, ActorDescInstance->GetGuid()));
+				}
+			}
+		}
+	}
+#endif // WITH_EDITOR 
+}
+
+
 void UPCGManagedActors::PostEditImport()
 {
 	// In this case, the managed actors won't be copied along the actor/component,
diff --git a/Engine/Plugins/PCG/Source/PCG/Private/PCGManagedResourceContainer.cpp b/Engine/Plugins/PCG/Source/PCG/Private/PCGManagedResourceContainer.cpp
new file mode 100644
--- /dev/null
+++ b/Engine/Plugins/PCG/Source/PCG/Private/PCGManagedResourceContainer.cpp
@@ -0,0 +1,485 @@
+// Copyright Epic Games, Inc. All Rights Reserved.
+
+#include "PCGManagedResourceContainer.h"
+
+#include "PCGGraphExecutionStateInterface.h"
+#include "PCGManagedResource.h"
+#include "PCGModule.h"
+#include "Helpers/PCGActorHelpers.h"
+#include "Utils/PCGGeneratedResourcesLogging.h"
+
+#include "UObject/Package.h"
+
+#if WITH_EDITOR
+#include "Editor.h"
+#include "EditorActorFolders.h"
+#endif
+
+void FPCGManagedResourceContainer::SetOwner(IPCGGraphExecutionSource* InOwner)
+{
+	Owner = InOwner;
+}
+
+void FPCGManagedResourceContainer::AddManagedResource(UPCGManagedResource* InManagedResource)
+{
+	PCGGeneratedResourcesLogging::LogAddToManagedResources(Owner, InManagedResource);
+
+	if (InManagedResource)
+	{
+		PCG::TScopeLock ResourcesLock(Owner->GetExecutionState().GetManagedResourceContainerLock());
+		ensure(!bAreResourcesInaccessible);
+		GeneratedResources.Add(InManagedResource);
+		Owner->GetExecutionState().OnManagedResourceAdded(InManagedResource);
+	}
+}
+
+void FPCGManagedResourceContainer::SetManagedResources(const TArray<TObjectPtr<UPCGManagedResource>>& InManagedResources)
+{
+	PCG::TScopeLock ResourcesLock(Owner->GetExecutionState().GetManagedResourceContainerLock());
+
+	// We expect the GeneratedResources to be empty here, as otherwise they might not be taken care of properly - they
+	// will be lost down below, but this should not happen. However, if the GeneratedResources are marked as Visible,
+	// then they will be copied over during BP duplication, hence why this will happen, hence the ensure here.
+	ensure(GeneratedResources.IsEmpty());
+
+	GeneratedResources = InManagedResources;
+
+	// Remove any null entries
+	for (int32 ResourceIndex = GeneratedResources.Num() - 1; ResourceIndex >= 0; --ResourceIndex)
+	{
+		if (!GeneratedResources[ResourceIndex])
+		{
+			GeneratedResources.RemoveAtSwap(ResourceIndex);
+		}
+	}
+}
+
+TArray<TObjectPtr<UPCGManagedResource>> FPCGManagedResourceContainer::GetManagedResourcesCopy() const
+{
+	TArray<TObjectPtr<UPCGManagedResource>> ManagedResourcesCopy;
+	
+	{
+		PCG::TScopeLock ResourcesLock(Owner->GetExecutionState().GetManagedResourceContainerLock());
+		ManagedResourcesCopy = GeneratedResources;
+	}
+
+	return ManagedResourcesCopy;
+}
+
+void FPCGManagedResourceContainer::CleanupUnusedManagedResources()
+{
+	PCGGeneratedResourcesLogging::LogCleanupUnusedManagedResources(Owner, GeneratedResources);
+
+	TSet<TSoftObjectPtr<AActor>> ActorsToDelete;
+
+	{
+		PCG::TScopeLock ResourcesLock(Owner->GetExecutionState().GetManagedResourceContainerLock());
+		ensure(!bAreResourcesInaccessible);
+		for (int32 ResourceIndex = GeneratedResources.Num() - 1; ResourceIndex >= 0; --ResourceIndex)
+		{
+			UPCGManagedResource* Resource = GetValid(GeneratedResources[ResourceIndex]);
+
+			PCGGeneratedResourcesLogging::LogCleanupUnusedManagedResourcesResource(Owner, Resource);
+
+			if (!Resource)
+			{
+				UE_LOG(LogPCG, Error, TEXT("[UPCGComponent::CleanupUnusedManagedResources] Null generated resource encountered on execution source \"%s\"."), *Owner->GetExecutionState().GetDebugName());
+			}
+
+			if (!Resource || Resource->ReleaseIfUnused(ActorsToDelete))
+			{
+#if WITH_EDITOR
+				if (Resource && !Resource->IsMarkedTransientOnLoad())
+				{
+					Resource->Rename(nullptr, GetTransientPackage(), REN_DoNotDirty | REN_DontCreateRedirectors | REN_NonTransactional);
+				}
+#endif
+				GeneratedResources.RemoveAtSwap(ResourceIndex);
+			}
+		}
+	}
+
+	if (ActorsToDelete.Num() > 0)
+	{
+		UPCGActorHelpers::DeleteActors(Owner->GetExecutionState().GetWorld(), ActorsToDelete.Array());
+	}
+
+	PCGGeneratedResourcesLogging::LogCleanupUnusedManagedResourcesFinished(Owner, GeneratedResources);
+}
+
+void FPCGManagedResourceContainer::Cleanup(bool bReleaseManagedResources)
+{
+	TSet<TSoftObjectPtr<AActor>> ActorsToDelete;
+	FPCGManagedActorLoadingScope Scope;
+
+	if (!bReleaseManagedResources && UPCGManagedResource::DebugForcePurgeAllResourcesOnGenerate())
+	{
+		bReleaseManagedResources = true;
+	}
+
+	PCG::TScopeLock ResourcesLock(Owner->GetExecutionState().GetManagedResourceContainerLock());
+	ensure(!bAreResourcesInaccessible);
+	Scope.AddResources(Owner, GeneratedResources);
+	for (int32 ResourceIndex = GeneratedResources.Num() - 1; ResourceIndex >= 0; --ResourceIndex)
+	{
+		// Note: resources can be null here in some loading + bp object cases
+		UPCGManagedResource* Resource = GeneratedResources[ResourceIndex];
+
+		PCGGeneratedResourcesLogging::LogCleanupLocalImmediateResource(Owner, Resource);
+
+		if (!Resource || Resource->Release(bReleaseManagedResources, ActorsToDelete))
+		{
+#if WITH_EDITOR
+			if (Resource && !Resource->IsMarkedTransientOnLoad())
+			{
+				Resource->Rename(nullptr, GetTransientPackage(), REN_DoNotDirty | REN_DontCreateRedirectors | REN_NonTransactional);
+			}
+#endif
+
+			GeneratedResources.RemoveAt(ResourceIndex);
+		}
+	}
+
+	UPCGActorHelpers::DeleteActors(Owner->GetExecutionState().GetWorld(), ActorsToDelete.Array());
+}
+
+void FPCGManagedResourceContainer::ForEachManagedResource(TFunctionRef<void(UPCGManagedResource*)> InFunction)
+{
+	PCG::TScopeLock ResourcesLock(Owner->GetExecutionState().GetManagedResourceContainerLock());
+	ensure(!bAreResourcesInaccessible);
+	for (TObjectPtr<UPCGManagedResource>& ManagedResource : GeneratedResources)
+	{
+		if (ManagedResource)
+		{
+			InFunction(ManagedResource);
+		}
+	}
+}
+
+void FPCGManagedResourceContainer::ForEachConstManagedResource(TFunctionRef<void(const UPCGManagedResource*)> InFunction) const
+{
+	PCG::TScopeLock ResourcesLock(Owner->GetExecutionState().GetManagedResourceContainerLock());
+	ensure(!bAreResourcesInaccessible);
+	for (const TObjectPtr<UPCGManagedResource>& ManagedResource : GeneratedResources)
+	{
+		if (ManagedResource)
+		{
+			InFunction(ManagedResource);
+		}
+	}
+}
+
+bool FPCGManagedResourceContainer::IsAnyObjectManagedByResource(const TArrayView<const UObject*> InObjects) const
+{
+	TRACE_CPUPROFILER_EVENT_SCOPE(FPCGManagedResourceContainer::IsAnyObjectManagedByResource);
+
+	if (!ensure(!bAreResourcesInaccessible))
+	{
+		return false;
+	}
+
+	bool bIsManagedByResource = false;
+	ForEachConstManagedResource([&bIsManagedByResource, &InObjects](const UPCGManagedResource* ManagedResource)
+	{
+		if (bIsManagedByResource || !ManagedResource)
+		{
+			return;
+		}
+
+		for (const UObject* Object : InObjects)
+		{
+			if (ManagedResource->IsManaging(Object))
+			{
+				bIsManagedByResource = true;
+				break;
+			}
+		}
+	});
+
+	return bIsManagedByResource;
+}
+
+#if WITH_EDITOR
+
+void FPCGManagedResourceContainer::MarkResourcesAsTransientOnLoad()
+{
+	for (TObjectPtr<UPCGManagedResource>& GeneratedResource : GeneratedResources)
+	{
+		if (GeneratedResource)
+		{
+			GeneratedResource->MarkTransientOnLoad();
+		}
+	}
+
+	LoadedPreviewResources = GeneratedResources;
+}
+
+bool FPCGManagedResourceContainer::ChangeTransientState(EPCGEditorDirtyMode NewEditingMode)
+{
+	bool bShouldMarkDirty = false;
+
+	PCG::TScopeLock ResourcesLock(Owner->GetExecutionState().GetManagedResourceContainerLock());
+	ensure(!bAreResourcesInaccessible);
+
+	for (TObjectPtr<UPCGManagedResource>& GeneratedResource : GeneratedResources)
+	{
+		if (GeneratedResource)
+		{
+			// Avoid changing transient state when switching to preview mode on marked as transient resources.
+			// When serializing a component we need its marked as transient resources to remain unchanged so they get serialized properly.
+			if (NewEditingMode == EPCGEditorDirtyMode::Preview && GeneratedResource->IsMarkedTransientOnLoad())
+			{
+				continue;
+			}
+
+			GeneratedResource->ChangeTransientState(NewEditingMode);
+			bShouldMarkDirty = true;
+		}
+	}
+
+	// If switching from preview mode to normal or preview-on-load,
+	// we must materialize any kind of change we've done on the packages that had a different behavior on load (e.g. actor packages)
+	if (NewEditingMode != EPCGEditorDirtyMode::Preview)
+	{
+		bShouldMarkDirty |= DeletePreviewResources();
+	}
+
+	return bShouldMarkDirty;
+}
+
+bool FPCGManagedResourceContainer::DeletePreviewResources()
+{
+	bool bResourceWasReleased = false;
+
+	TSet<TSoftObjectPtr<AActor>> ActorsToDelete;
+	// Make sure to release fully the resources that were loaded
+	for (TObjectPtr<UPCGManagedResource> ResourceToRelease : LoadedPreviewResources)
+	{
+		if (!GeneratedResources.Contains(ResourceToRelease))
+		{
+			// Changing the transient state will clear the "marked transient on load" flag
+			ResourceToRelease->ChangeTransientState(EPCGEditorDirtyMode::Normal);
+			ResourceToRelease->Release(/*bHardRelease=*/true, ActorsToDelete);
+			bResourceWasReleased = true;
+		}
+
+		// Either this resource was released in the previous code block or prior to calling DeletePreviewResources its transient state was changed back to normal
+		ensure(!ResourceToRelease->IsMarkedTransientOnLoad());
+	}
+
+	LoadedPreviewResources.Empty();
+
+	if (!ActorsToDelete.IsEmpty())
+	{
+		UPCGActorHelpers::DeleteActors(Owner->GetExecutionState().GetWorld(), ActorsToDelete.Array());
+	}
+
+	return bResourceWasReleased;
+}
+
+void FPCGManagedResourceContainer::BeginPreviewSave(TArray<TObjectPtr<UPCGManagedResource>>& OutGeneratedCopy)
+{
+	OutGeneratedCopy = GeneratedResources;
+	GeneratedResources = LoadedPreviewResources;
+}
+
+void FPCGManagedResourceContainer::EndPreviewSave(const TArray<TObjectPtr<UPCGManagedResource>>& InGeneratedCopy)
+{
+	GeneratedResources = InGeneratedCopy;
+}
+
+#endif // WITH_EDITOR
+
+void FPCGManagedResourceContainer::EndPlay()
+{
+	if (ensure(!bAreResourcesInaccessible))
+	{
+		PCG::TScopeLock ResourcesLock(Owner->GetExecutionState().GetManagedResourceContainerLock());
+
+		for (int32 ResourceIndex = GeneratedResources.Num() - 1; ResourceIndex >= 0; --ResourceIndex)
+		{
+			if (!GeneratedResources[ResourceIndex])
+			{
+				// Remove null entries.
+				GeneratedResources.RemoveAtSwap(ResourceIndex);
+			}
+			else if (GeneratedResources[ResourceIndex]->ReleaseOnTeardown())
+			{
+				if (GeneratedResources[ResourceIndex])
+				{
+					TSet<TSoftObjectPtr<AActor>> ActorsToDelete;
+					GeneratedResources[ResourceIndex]->Release(/*bHardRelease=*/true, ActorsToDelete);
+
+					// Don't support deleting actors during teardown.
+					ensure(ActorsToDelete.IsEmpty());
+				}
+
+				GeneratedResources.RemoveAtSwap(ResourceIndex);
+			}
+		}
+	}
+}
+
+TSharedPtr<FPCGManagedResourceContainer::FCleanupTask> FPCGManagedResourceContainer::CreateCleanupTask()
+{
+	TSharedPtr<FCleanupTask> CleanupTask = MakeShared<FCleanupTask>();
+	CleanupTask->WeakExecutionSource = Owner;
+
+	return CleanupTask;
+}
+
+void FPCGManagedResourceContainer::FCleanupTask::Abort()
+{
+	// If the component is not valid anymore, just early out
+	if (!WeakExecutionSource.IsValid())
+	{
+		return;
+	}
+
+	FPCGManagedResourceContainer* Container = WeakExecutionSource->GetExecutionState().GetManagedResourceContainer();
+	if (Container)
+	{
+		PCG::TScopeLock ContainerLock(WeakExecutionSource->GetExecutionState().GetManagedResourceContainerLock());
+		Container->bAreResourcesInaccessible = false;
+	}
+}
+
+bool FPCGManagedResourceContainer::FCleanupTask::Execute(bool bReleaseManagedResources)
+{
+	IPCGGraphExecutionSource* ExecutionSource = WeakExecutionSource.Get();
+	if (!ExecutionSource)
+	{
+		return true;
+	}
+
+	FPCGManagedResourceContainer* Container = ExecutionSource->GetExecutionState().GetManagedResourceContainer();
+	if (!Container)
+	{
+		return true;
+	}
+
+	PCG::TScopeLock ContainerLock(ExecutionSource->GetExecutionState().GetManagedResourceContainerLock());
+
+	// Safeguard to track illegal modifications of the generated resources array while doing cleanup
+	if (bIsFirstIteration)
+	{
+		ensure(!Container->bAreResourcesInaccessible);
+		Container->bAreResourcesInaccessible = true;
+		ResourceIndex = Container->GeneratedResources.Num() - 1;
+		bIsFirstIteration = false;
+		AddResources(WeakExecutionSource.Get(), Container->GeneratedResources);
+	}
+
+	// Going backward
+	if (ResourceIndex >= 0)
+	{
+		UPCGManagedResource* Resource = Container->GeneratedResources[ResourceIndex];
+
+		if (!Resource)
+		{
+			UE_LOG(LogPCG, Error, TEXT("[FPCGManagedResourceContainer::FCleanupTask::Execute] Null generated resource encountered on execution source \"%s\"."), *ExecutionSource->GetExecutionState().GetDebugName());
+		}
+
+		PCGGeneratedResourcesLogging::LogCreateCleanupTaskResource(WeakExecutionSource.Get(), Resource);
+
+		if (!Resource || Resource->Release(bReleaseManagedResources, ActorsToDelete))
+		{
+#if WITH_EDITOR
+			if (Resource && !Resource->IsMarkedTransientOnLoad())
+			{
+				Resource->Rename(nullptr, GetTransientPackage(), REN_DoNotDirty | REN_DontCreateRedirectors | REN_NonTransactional);
+			}
+#endif
+
+			Container->GeneratedResources.RemoveAtSwap(ResourceIndex);
+		}
+
+		ResourceIndex--;
+
+		// Returning false means the task is not over
+		return false;
+	}
+	else
+	{
+		Container->bAreResourcesInaccessible = false;
+	}
+
+	TSet<FName> DeletedActorFolders;
+
+	if (UWorld* World = ExecutionSource->GetExecutionState().GetWorld())
+	{
+		const TArray<TSoftObjectPtr<AActor>> ActorsToDeleteSoft = ActorsToDelete.Array();
+
+#if WITH_EDITOR
+		for (const TSoftObjectPtr<AActor>& Actor : ActorsToDelete)
+		{
+			if (Actor.IsValid())
+			{
+				FName ActorFolderPath = Actor->GetFolderPath();
+				if (ActorFolderPath != NAME_None)
+				{
+					DeletedActorFolders.Add(ActorFolderPath);
+				}
+			}
+		}
+#endif
+
+		UPCGActorHelpers::DeleteActors(World, ActorsToDeleteSoft);
+	}
+
+#if WITH_EDITOR
+	if (UWorld* ThisWorld = ExecutionSource->GetExecutionState().GetWorld(); ThisWorld && GEditor) // FActorFolders require the editor
+	{
+		for (FName FolderPath : DeletedActorFolders)
+		{
+			FFolder GeneratedFolder(FFolder::GetWorldRootFolder(ThisWorld).GetRootObject(), FolderPath);
+			const bool bFolderExists = GeneratedFolder.IsValid() && FActorFolders::Get().ContainsFolder(*ThisWorld, GeneratedFolder);
+			bool bFoundActors = false;
+
+			if (bFolderExists)
+			{
+				TSet<FName> Folders{ GeneratedFolder.GetPath() };
+				FActorFolders::ForEachActorInFolders(*ThisWorld, Folders, [&bFoundActors](AActor* InActor)
+				{
+					if (InActor)
+					{
+						bFoundActors = true;
+						return false;
+					}
+					else
+					{
+						return true;
+					}
+				});
+			}
+
+			if (bFolderExists && !bFoundActors)
+			{
+				// Delete all subfolders
+				TArray<FFolder> SubfoldersToDelete;
+				FActorFolders::Get().ForEachFolder(*ThisWorld, [&GeneratedFolder, &SubfoldersToDelete](const FFolder& InFolder)
+				{
+					if (InFolder.IsChildOf(GeneratedFolder))
+					{
+						SubfoldersToDelete.Add(InFolder);
+					}
+					
+					return true;
+				});
+
+				for (const FFolder& FolderToDelete : SubfoldersToDelete) //-V1078
+				{
+					FActorFolders::Get().DeleteFolder(*ThisWorld, FolderToDelete);
+				}
+
+				// Finally, delete folder
+				FActorFolders::Get().DeleteFolder(*ThisWorld, GeneratedFolder);
+			}
+		}
+	}
+#endif
+
+	PCGGeneratedResourcesLogging::LogCreateCleanupTaskFinished(ExecutionSource, &Container->GeneratedResources);
+
+	return true;
+}
\ No newline at end of file
diff --git a/Engine/Plugins/PCG/Source/PCG/Private/Utils/PCGGeneratedResourcesLogging.cpp b/Engine/Plugins/PCG/Source/PCG/Private/Utils/PCGGeneratedResourcesLogging.cpp
--- a/Engine/Plugins/PCG/Source/PCG/Private/Utils/PCGGeneratedResourcesLogging.cpp
+++ b/Engine/Plugins/PCG/Source/PCG/Private/Utils/PCGGeneratedResourcesLogging.cpp
@@ -2,7 +2,7 @@
 
 #include "PCGGeneratedResourcesLogging.h"
 
-#include "PCGComponent.h"
+#include "PCGGraphExecutionStateInterface.h"
 #include "PCGGraph.h"
 #include "PCGManagedResource.h"
 #include "PCGModule.h"
@@ -34,13 +34,25 @@ namespace PCGGeneratedResourcesLogging
 	}
 
 #if WITH_EDITOR
-	void LogResource(const UPCGComponent* InComponent, UPCGManagedResource* InResource)
+	FString GetExecutionSourceString(const IPCGGraphExecutionSource* InExecutionSource)
+	{
+		return InExecutionSource ? InExecutionSource->GetExecutionState().GetDebugName() : TEXT("MISSINGEXECUTIONSOURCE");
+	}
+
+	FString GetExecutionSourceGraphString(const IPCGGraphExecutionSource* InExecutionSource)
+	{
+		return InExecutionSource && InExecutionSource->GetExecutionState().GetGraph() ? *InExecutionSource->GetExecutionState().GetGraph()->GetName() : TEXT("MISSINGGRAPH");
+	}
+#endif
+
+#if WITH_EDITOR
+	void LogResource(const IPCGGraphExecutionSource* InExecutionSource, UPCGManagedResource* InResource)
 	{
 		if (const UPCGManagedActors* ManagedActors = Cast<UPCGManagedActors>(InResource))
 		{
 			UE_LOG(LogPCG, Log, TEXT("[PCGMANAGEDRESOURCES] [%s/%s]         Managed actors (%d):"),
-				(InComponent && InComponent->GetOwner()) ? *InComponent->GetOwner()->GetName() : TEXT("MISSINGCOMPONENT"),
-				(InComponent && InComponent->GetGraph()) ? *InComponent->GetGraph()->GetName() : TEXT("MISSINGGRAPH"),
+				*GetExecutionSourceString(InExecutionSource),
+				*GetExecutionSourceGraphString(InExecutionSource),
 				ManagedActors->GetConstGeneratedActors().Num());
 
 			const int32 MaxPrint = CVarGenResourcesLoggingMaxPrintCount.GetValueOnAnyThread();
@@ -54,17 +66,18 @@ namespace PCGGeneratedResourcesLogging
 
 				if (Actor.Get())
 				{
-					UE_LOG(LogPCG, Log, TEXT("[PCGMANAGEDRESOURCES] [%s/%s]                 |- '%s' (%d/%d)"),
-						(InComponent && InComponent->GetOwner()) ? *InComponent->GetOwner()->GetName() : TEXT("MISSINGCOMPONENT"),
-						(InComponent && InComponent->GetGraph()) ? *InComponent->GetGraph()->GetName() : TEXT("MISSINGGRAPH"),
+					UE_LOG(LogPCG, Log, TEXT("[PCGMANAGEDRESOURCES] [%s/%s]                 |- '%s' (%d/%d)"), 
+						*GetExecutionSourceString(InExecutionSource), 
+						*GetExecutionSourceGraphString(InExecutionSource),
 						*Actor->GetFName().ToString(),
-						Count, ManagedActors->GetConstGeneratedActors().Num());
+						Count, 
+						ManagedActors->GetConstGeneratedActors().Num());
 				}
 				else
 				{
 					UE_LOG(LogPCG, Log, TEXT("[PCGMANAGEDRESOURCES] [%s/%s]                 |- NULL (%d/%d)"),
-						(InComponent && InComponent->GetOwner()) ? *InComponent->GetOwner()->GetName() : TEXT("MISSINGCOMPONENT"),
-						(InComponent && InComponent->GetGraph()) ? *InComponent->GetGraph()->GetName() : TEXT("MISSINGGRAPH"),
+						*GetExecutionSourceString(InExecutionSource),
+						*GetExecutionSourceGraphString(InExecutionSource),
 						Count, ManagedActors->GetConstGeneratedActors().Num());
 				}
 			}
@@ -74,15 +87,15 @@ namespace PCGGeneratedResourcesLogging
 			if (ManagedComponent->GeneratedComponent)
 			{
 				UE_LOG(LogPCG, Log, TEXT("[PCGMANAGEDRESOURCES] [%s/%s]         Managed component: '%s'"),
-					(InComponent && InComponent->GetOwner()) ? *InComponent->GetOwner()->GetName() : TEXT("MISSINGCOMPONENT"),
-					(InComponent && InComponent->GetGraph()) ? *InComponent->GetGraph()->GetName() : TEXT("MISSINGGRAPH"),
+					*GetExecutionSourceString(InExecutionSource),
+					*GetExecutionSourceGraphString(InExecutionSource),
 					*ManagedComponent->GeneratedComponent->GetReadableName());
 			}
 			else
 			{
 				UE_LOG(LogPCG, Warning, TEXT("[PCGMANAGEDRESOURCES] [%s/%s]         NULL managed component or no owner"),
-					(InComponent && InComponent->GetOwner()) ? *InComponent->GetOwner()->GetName() : TEXT("MISSINGCOMPONENT"),
-					(InComponent && InComponent->GetGraph()) ? *InComponent->GetGraph()->GetName() : TEXT("MISSINGGRAPH"));
+					*GetExecutionSourceString(InExecutionSource),
+					*GetExecutionSourceGraphString(InExecutionSource));
 			}
 		}
 		else
@@ -90,19 +103,19 @@ namespace PCGGeneratedResourcesLogging
 			if (InResource)
 			{
 				UE_LOG(LogPCG, Warning, TEXT("[PCGMANAGEDRESOURCES] [%s/%s]         Unidentified managed resource"),
-					(InComponent && InComponent->GetOwner()) ? *InComponent->GetOwner()->GetName() : TEXT("MISSINGCOMPONENT"),
-					(InComponent && InComponent->GetGraph()) ? *InComponent->GetGraph()->GetName() : TEXT("MISSINGGRAPH"));
+					*GetExecutionSourceString(InExecutionSource),
+					*GetExecutionSourceGraphString(InExecutionSource));
 			}
 			else
 			{
 				UE_LOG(LogPCG, Warning, TEXT("[PCGMANAGEDRESOURCES] [%s/%s]         NULL unidentified resource"),
-					(InComponent && InComponent->GetOwner()) ? *InComponent->GetOwner()->GetName() : TEXT("MISSINGCOMPONENT"),
-					(InComponent && InComponent->GetGraph()) ? *InComponent->GetGraph()->GetName() : TEXT("MISSINGGRAPH"));
+					*GetExecutionSourceString(InExecutionSource),
+					*GetExecutionSourceGraphString(InExecutionSource));
 			}
 		}
 	}
 
-	void LogResources(const UPCGComponent* InComponent, const TArray<UPCGManagedResource*>& Resources)
+	void LogResources(const IPCGGraphExecutionSource* InComponent, const TArray<UPCGManagedResource*>& Resources)
 	{
 		for (UPCGManagedResource* Resource : Resources)
 		{
@@ -111,7 +124,7 @@ namespace PCGGeneratedResourcesLogging
 	}
 #endif // WITH_EDITOR
 
-	void LogAddToManagedResources(const UPCGComponent* InComponent, UPCGManagedResource* InResource)
+	void LogAddToManagedResources(const IPCGGraphExecutionSource* InExecutionSource, UPCGManagedResource* InResource)
 	{
 #if WITH_EDITOR
 		if (!LogEnabled())
@@ -119,15 +132,15 @@ namespace PCGGeneratedResourcesLogging
 			return;
 		}
 
-		UE_LOG(LogPCG, Log, TEXT("[PCGMANAGEDRESOURCES] [%s/%s] UPCGComponent::AddToManagedResources:"),
-			(InComponent && InComponent->GetOwner()) ? *InComponent->GetOwner()->GetName() : TEXT("MISSINGCOMPONENT"),
-			(InComponent && InComponent->GetGraph()) ? *InComponent->GetGraph()->GetName() : TEXT("MISSINGGRAPH"));
+		UE_LOG(LogPCG, Log, TEXT("[PCGMANAGEDRESOURCES] [%s/%s] AddToManagedResources:"),
+			*GetExecutionSourceString(InExecutionSource),
+			*GetExecutionSourceGraphString(InExecutionSource));
 
-		LogResource(InComponent, InResource);
+		LogResource(InExecutionSource, InResource);
 #endif
 	}
 
-	void LogCleanupLocal(const UPCGComponent* InComponent, bool bRemoveComponents)
+	void LogCleanupLocal(const IPCGGraphExecutionSource* InExecutionSource, bool bReleaseManagedResources)
 	{
 #if WITH_EDITOR
 		if (!LogEnabled())
@@ -136,14 +149,14 @@ namespace PCGGeneratedResourcesLogging
 		}
 
 		UE_LOG(LogPCG, Log, TEXT("[PCGMANAGEDRESOURCES]")); // Blank line
-		UE_LOG(LogPCG, Log, TEXT("[PCGMANAGEDRESOURCES] [%s/%s] UPCGComponent::CleanupLocal, bRemoveComponents: %d"),
-			(InComponent && InComponent->GetOwner()) ? *InComponent->GetOwner()->GetName() : TEXT("MISSINGCOMPONENT"),
-			(InComponent && InComponent->GetGraph()) ? *InComponent->GetGraph()->GetName() : TEXT("MISSINGGRAPH"),
-			bRemoveComponents);
+		UE_LOG(LogPCG, Log, TEXT("[PCGMANAGEDRESOURCES] [%s/%s] CleanupLocal, bReleaseManagedResources: %d"),
+			*GetExecutionSourceString(InExecutionSource),
+			*GetExecutionSourceGraphString(InExecutionSource),
+			bReleaseManagedResources);
 #endif
 	}
 
-	void LogCleanupLocalImmediate(const UPCGComponent* InComponent, bool bHardRelease, const TArray<UPCGManagedResource*>& GeneratedResources)
+	void LogCleanupLocalImmediate(const IPCGGraphExecutionSource* InExecutionSource, bool bHardRelease, const TArray<UPCGManagedResource*>& GeneratedResources)
 	{
 #if WITH_EDITOR
 		if (!LogEnabled())
@@ -152,17 +165,17 @@ namespace PCGGeneratedResourcesLogging
 		}
 
 		UE_LOG(LogPCG, Log, TEXT("[PCGMANAGEDRESOURCES]")); // Blank line
-		UE_LOG(LogPCG, Log, TEXT("[PCGMANAGEDRESOURCES] [%s/%s] UPCGComponent::CleanupLocalImmediate BEGIN, bHardRelease: %d, GeneratedResources.Num() = %d"),
-			(InComponent && InComponent->GetOwner()) ? *InComponent->GetOwner()->GetName() : TEXT("MISSINGCOMPONENT"),
-			(InComponent && InComponent->GetGraph()) ? *InComponent->GetGraph()->GetName() : TEXT("MISSINGGRAPH"),
+		UE_LOG(LogPCG, Log, TEXT("[PCGMANAGEDRESOURCES] [%s/%s] CleanupLocalImmediate BEGIN, bHardRelease: %d, GeneratedResources.Num() = %d"),
+			*GetExecutionSourceString(InExecutionSource),
+			*GetExecutionSourceGraphString(InExecutionSource),
 			bHardRelease,
 			GeneratedResources.Num());
 
-		LogResources(InComponent, GeneratedResources);
+		LogResources(InExecutionSource, GeneratedResources);
 #endif
 	}
 
-	void LogCleanupLocalImmediateResource(const UPCGComponent* InComponent, UPCGManagedResource* Resource)
+	void LogCleanupLocalImmediateResource(const IPCGGraphExecutionSource* InExecutionSource, UPCGManagedResource* Resource)
 	{
 #if WITH_EDITOR
 		if (!LogEnabled())
@@ -170,15 +183,15 @@ namespace PCGGeneratedResourcesLogging
 			return;
 		}
 
-		UE_LOG(LogPCG, Log, TEXT("[PCGMANAGEDRESOURCES] [%s/%s] UPCGComponent::CleanupLocalImmediate:"),
-			(InComponent && InComponent->GetOwner()) ? *InComponent->GetOwner()->GetName() : TEXT("MISSINGCOMPONENT"),
-			(InComponent && InComponent->GetGraph()) ? *InComponent->GetGraph()->GetName() : TEXT("MISSINGGRAPH"));
+		UE_LOG(LogPCG, Log, TEXT("[PCGMANAGEDRESOURCES] [%s/%s] CleanupLocalImmediate:"),
+			*GetExecutionSourceString(InExecutionSource),
+			*GetExecutionSourceGraphString(InExecutionSource));
 		
-		LogResource(InComponent, Resource);
+		LogResource(InExecutionSource, Resource);
 #endif
 	}
 
-	void LogCleanupLocalImmediateFinished(const UPCGComponent* InComponent, const TArray<UPCGManagedResource*>& GeneratedResources)
+	void LogCleanupLocalImmediateFinished(const IPCGGraphExecutionSource* InExecutionSource, const TArray<UPCGManagedResource*>& GeneratedResources)
 	{
 #if WITH_EDITOR
 		if (!LogEnabled())
@@ -186,16 +199,16 @@ namespace PCGGeneratedResourcesLogging
 			return;
 		}
 
-		UE_LOG(LogPCG, Log, TEXT("[PCGMANAGEDRESOURCES] [%s/%s] UPCGComponent::CleanupLocalImmediate FINISHED, Final GeneratedResources (%d):"),
-			(InComponent && InComponent->GetOwner()) ? *InComponent->GetOwner()->GetName() : TEXT("MISSINGCOMPONENT"),
-			(InComponent && InComponent->GetGraph()) ? *InComponent->GetGraph()->GetName() : TEXT("MISSINGGRAPH"),
+		UE_LOG(LogPCG, Log, TEXT("[PCGMANAGEDRESOURCES] [%s/%s] CleanupLocalImmediate FINISHED, Final GeneratedResources (%d):"),
+			*GetExecutionSourceString(InExecutionSource),
+			*GetExecutionSourceGraphString(InExecutionSource),
 			GeneratedResources.Num());
 
-		LogResources(InComponent, GeneratedResources);
+		LogResources(InExecutionSource, GeneratedResources);
 #endif
 	}
 
-	void LogCreateCleanupTask(const UPCGComponent* InComponent, bool bRemoveComponents)
+	void LogCreateCleanupTask(const IPCGGraphExecutionSource* InExecutionSource, bool bReleaseManagedResources)
 	{
 #if WITH_EDITOR
 		if (!LogEnabled())
@@ -203,14 +216,14 @@ namespace PCGGeneratedResourcesLogging
 			return;
 		}
 
-		UE_LOG(LogPCG, Log, TEXT("[PCGMANAGEDRESOURCES] [%s/%s] UPCGComponent::CreateCleanupTask, bRemoveComponents: %d"),
-			(InComponent && InComponent->GetOwner()) ? *InComponent->GetOwner()->GetName() : TEXT("MISSINGCOMPONENT"),
-			(InComponent && InComponent->GetGraph()) ? *InComponent->GetGraph()->GetName() : TEXT("MISSINGGRAPH"),
-			bRemoveComponents);
+		UE_LOG(LogPCG, Log, TEXT("[PCGMANAGEDRESOURCES] [%s/%s] CreateCleanupTask, bReleaseManagedResources: %d"),
+			*GetExecutionSourceString(InExecutionSource),
+			*GetExecutionSourceGraphString(InExecutionSource),
+			bReleaseManagedResources);
 #endif
 	}
 
-	void LogCreateCleanupTaskResource(const UPCGComponent* InComponent, UPCGManagedResource* Resource)
+	void LogCreateCleanupTaskResource(const IPCGGraphExecutionSource* InExecutionSource, UPCGManagedResource* Resource)
 	{
 #if WITH_EDITOR
 		if (!LogEnabled())
@@ -218,15 +231,15 @@ namespace PCGGeneratedResourcesLogging
 			return;
 		}
 
-		UE_LOG(LogPCG, Log, TEXT("[PCGMANAGEDRESOURCES] [%s/%s] UPCGComponent::CreateCleanupTask::CleanupTask:"),
-			(InComponent && InComponent->GetOwner()) ? *InComponent->GetOwner()->GetName() : TEXT("MISSINGCOMPONENT"),
-			(InComponent && InComponent->GetGraph()) ? *InComponent->GetGraph()->GetName() : TEXT("MISSINGGRAPH"));
+		UE_LOG(LogPCG, Log, TEXT("[PCGMANAGEDRESOURCES] [%s/%s] CreateCleanupTask::CleanupTask:"),
+			*GetExecutionSourceString(InExecutionSource),
+			*GetExecutionSourceGraphString(InExecutionSource));
 
-		LogResource(InComponent, Resource);
+		LogResource(InExecutionSource, Resource);
 #endif
 	}
 
-	void LogCreateCleanupTaskFinished(const UPCGComponent* InComponent, const TArray<TObjectPtr<UPCGManagedResource>>* InGeneratedResources)
+	void LogCreateCleanupTaskFinished(const IPCGGraphExecutionSource* InExecutionSource, const TArray<TObjectPtr<UPCGManagedResource>>* InGeneratedResources)
 	{
 #if WITH_EDITOR
 		if (!LogEnabled())
@@ -234,9 +247,9 @@ namespace PCGGeneratedResourcesLogging
 			return;
 		}
 
-		UE_LOG(LogPCG, Log, TEXT("[PCGMANAGEDRESOURCES] [%s/%s] UPCGComponent::CreateCleanupTask::CleanupTask FINISHED, GeneratedResources.Num() = %d"),
-			(InComponent && InComponent->GetOwner()) ? *InComponent->GetOwner()->GetName() : TEXT("MISSINGCOMPONENT"),
-			(InComponent && InComponent->GetGraph()) ? *InComponent->GetGraph()->GetName() : TEXT("MISSINGGRAPH"),
+		UE_LOG(LogPCG, Log, TEXT("[PCGMANAGEDRESOURCES] [%s/%s] CreateCleanupTask::CleanupTask FINISHED, GeneratedResources.Num() = %d"),
+			*GetExecutionSourceString(InExecutionSource),
+			*GetExecutionSourceGraphString(InExecutionSource),
 			InGeneratedResources ? InGeneratedResources->Num() : -1);
 #endif
 	}
@@ -249,12 +262,12 @@ namespace PCGGeneratedResourcesLogging
 			return;
 		}
 
-		UE_LOG(LogPCG, Log, TEXT("[PCGMANAGEDRESOURCES] UPCGComponent::CreateCleanupTask::CleanupTask FINISHED, GeneratedResources.Num() = %d"),
+		UE_LOG(LogPCG, Log, TEXT("[PCGMANAGEDRESOURCES] CreateCleanupTask::CleanupTask FINISHED, GeneratedResources.Num() = %d"),
 			InGeneratedResources ? InGeneratedResources->Num() : -1);
 #endif
 	}
 
-	void LogCleanupUnusedManagedResources(const UPCGComponent* InComponent, const TArray<UPCGManagedResource*>& GeneratedResources)
+	void LogCleanupUnusedManagedResources(const IPCGGraphExecutionSource* InExecutionSource, const TArray<UPCGManagedResource*>& GeneratedResources)
 	{
 #if WITH_EDITOR
 		if (!LogEnabled())
@@ -262,14 +275,14 @@ namespace PCGGeneratedResourcesLogging
 			return;
 		}
 
-		UE_LOG(LogPCG, Log, TEXT("[PCGMANAGEDRESOURCES] [%s/%s] UPCGComponent::CleanupUnusedManagedResources BEGIN, GeneratedResources.Num() = %d"),
-			(InComponent && InComponent->GetOwner()) ? *InComponent->GetOwner()->GetName() : TEXT("MISSINGCOMPONENT"),
-			(InComponent && InComponent->GetGraph()) ? *InComponent->GetGraph()->GetName() : TEXT("MISSINGGRAPH"),
+		UE_LOG(LogPCG, Log, TEXT("[PCGMANAGEDRESOURCES] [%s/%s] CleanupUnusedManagedResources BEGIN, GeneratedResources.Num() = %d"),
+			*GetExecutionSourceString(InExecutionSource),
+			*GetExecutionSourceGraphString(InExecutionSource),
 			GeneratedResources.Num());
 #endif
 	}
 
-	void LogCleanupUnusedManagedResourcesResource(const UPCGComponent* InComponent, UPCGManagedResource* InResource)
+	void LogCleanupUnusedManagedResourcesResource(const IPCGGraphExecutionSource* InExecutionSource, UPCGManagedResource* InResource)
 	{
 #if WITH_EDITOR
 		if (!LogEnabled())
@@ -277,15 +290,15 @@ namespace PCGGeneratedResourcesLogging
 			return;
 		}
 
-		UE_LOG(LogPCG, Log, TEXT("[PCGMANAGEDRESOURCES] [%s/%s] UPCGComponent::CleanupUnusedManagedResources:"),
-			(InComponent && InComponent->GetOwner()) ? *InComponent->GetOwner()->GetName() : TEXT("MISSINGCOMPONENT"),
-			(InComponent && InComponent->GetGraph()) ? *InComponent->GetGraph()->GetName() : TEXT("MISSINGGRAPH"));
+		UE_LOG(LogPCG, Log, TEXT("[PCGMANAGEDRESOURCES] [%s/%s] CleanupUnusedManagedResources:"),
+			*GetExecutionSourceString(InExecutionSource),
+			*GetExecutionSourceGraphString(InExecutionSource));
 
-		LogResource(InComponent, InResource);
+		LogResource(InExecutionSource, InResource);
 #endif
 	}
 
-	void LogCleanupUnusedManagedResourcesFinished(const UPCGComponent* InComponent, const TArray<UPCGManagedResource*>& GeneratedResources)
+	void LogCleanupUnusedManagedResourcesFinished(const IPCGGraphExecutionSource* InExecutionSource, const TArray<UPCGManagedResource*>& GeneratedResources)
 	{
 #if WITH_EDITOR
 		if (!LogEnabled())
@@ -293,12 +306,12 @@ namespace PCGGeneratedResourcesLogging
 			return;
 		}
 
-		UE_LOG(LogPCG, Log, TEXT("[PCGMANAGEDRESOURCES] [%s/%s] UPCGComponent::CleanupUnusedManagedResources, FINISHED:"),
-			(InComponent && InComponent->GetOwner()) ? *InComponent->GetOwner()->GetName() : TEXT("MISSINGCOMPONENT"),
-			(InComponent && InComponent->GetGraph()) ? *InComponent->GetGraph()->GetName() : TEXT("MISSINGGRAPH"),
+		UE_LOG(LogPCG, Log, TEXT("[PCGMANAGEDRESOURCES] [%s/%s] CleanupUnusedManagedResources, FINISHED:"),
+			*GetExecutionSourceString(InExecutionSource),
+			*GetExecutionSourceGraphString(InExecutionSource),
 			GeneratedResources.Num());
 
-		LogResources(InComponent, GeneratedResources);
+		LogResources(InExecutionSource, GeneratedResources);
 #endif
 	}
 
@@ -310,13 +323,13 @@ namespace PCGGeneratedResourcesLogging
 			return;
 		}
 
-		const UPCGComponent* Component = InResource ? Cast<UPCGComponent>(InResource->GetOuter()) : nullptr;
+		const IPCGGraphExecutionSource* ExecutionSource = InResource->GetImplementingOuter<IPCGGraphExecutionSource>();
 
 		UE_LOG(LogPCG, Log, TEXT("[PCGMANAGEDRESOURCES] [%s/%s] UPCGManagedResource::Release, SOFT release:"),
-			(Component && Component->GetOwner()) ? *Component->GetOwner()->GetName() : TEXT("MISSINGCOMPONENT"),
-			(Component && Component->GetGraph()) ? *Component->GetGraph()->GetName() : TEXT("MISSINGGRAPH"));
+			*GetExecutionSourceString(ExecutionSource),
+			*GetExecutionSourceGraphString(ExecutionSource));
 
-		LogResource(Component, InResource);
+		LogResource(ExecutionSource, InResource);
 #endif
 	}
 
@@ -328,13 +341,13 @@ namespace PCGGeneratedResourcesLogging
 			return;
 		}
 
-		const UPCGComponent* Component = InResource ? Cast<UPCGComponent>(InResource->GetOuter()) : nullptr;
+		const IPCGGraphExecutionSource* ExecutionSource = InResource->GetImplementingOuter<IPCGGraphExecutionSource>();
 
 		UE_LOG(LogPCG, Log, TEXT("[PCGMANAGEDRESOURCES] [%s/%s] UPCGManagedResource::Release, HARD release:"),
-			(Component && Component->GetOwner()) ? *Component->GetOwner()->GetName() : TEXT("MISSINGCOMPONENT"),
-			(Component && Component->GetGraph()) ? *Component->GetGraph()->GetName() : TEXT("MISSINGGRAPH"));
+			*GetExecutionSourceString(ExecutionSource),
+			*GetExecutionSourceGraphString(ExecutionSource));
 
-		LogResource(Component, InResource);
+		LogResource(ExecutionSource, InResource);
 #endif
 	}
 
@@ -346,12 +359,12 @@ namespace PCGGeneratedResourcesLogging
 			return;
 		}
 
-		const UPCGComponent* Component = InResource ? Cast<UPCGComponent>(InResource->GetOuter()) : nullptr;
+		const IPCGGraphExecutionSource* ExecutionSource = InResource->GetImplementingOuter<IPCGGraphExecutionSource>();
 		const bool bMarkedTransientOnLoad = InResource && InResource->IsMarkedTransientOnLoad();
 
 		UE_LOG(LogPCG, Log, TEXT("[PCGMANAGEDRESOURCES] [%s/%s] UPCGManagedActors::Release, %s release, bMarkedTransientOnLoad: %d, bOnlyMarkedForCleanup: %d, actor count %d scheduled for delete"),
-			(Component && Component->GetOwner()) ? *Component->GetOwner()->GetName() : TEXT("MISSINGCOMPONENT"),
-			(Component && Component->GetGraph()) ? *Component->GetGraph()->GetName() : TEXT("MISSINGGRAPH"),
+			*GetExecutionSourceString(ExecutionSource),
+			*GetExecutionSourceGraphString(ExecutionSource),
 			bHardRelease ? TEXT("HARD") : TEXT("SOFT"),
 			bMarkedTransientOnLoad,
 			bOnlyMarkedForCleanup,
@@ -369,8 +382,8 @@ namespace PCGGeneratedResourcesLogging
 			if (ActorToDelete)
 			{
 				UE_LOG(LogPCG, Log, TEXT("[PCGMANAGEDRESOURCES] [%s/%s]                 |- '%s' (%d/%d)"),
-					(Component && Component->GetOwner()) ? *Component->GetOwner()->GetName() : TEXT("MISSINGCOMPONENT"),
-					(Component && Component->GetGraph()) ? *Component->GetGraph()->GetName() : TEXT("MISSINGGRAPH"),
+					*GetExecutionSourceString(ExecutionSource),
+					*GetExecutionSourceGraphString(ExecutionSource),
 					*ActorToDelete->GetFName().ToString(),
 					Count,
 					ActorsToDelete.Num());
@@ -378,8 +391,8 @@ namespace PCGGeneratedResourcesLogging
 			else
 			{
 				UE_LOG(LogPCG, Log, TEXT("[PCGMANAGEDRESOURCES] [%s/%s]                 |- NULL (%d/%d)"),
-					(Component && Component->GetOwner()) ? *Component->GetOwner()->GetName() : TEXT("MISSINGCOMPONENT"),
-					(Component && Component->GetGraph()) ? *Component->GetGraph()->GetName() : TEXT("MISSINGGRAPH"),
+					*GetExecutionSourceString(ExecutionSource),
+					*GetExecutionSourceGraphString(ExecutionSource),
 					Count,
 					ActorsToDelete.Num());
 			}
@@ -395,13 +408,13 @@ namespace PCGGeneratedResourcesLogging
 			return;
 		}
 
-		const UPCGComponent* Component = InResource ? Cast<UPCGComponent>(InResource->GetOuter()) : nullptr;
+		const IPCGGraphExecutionSource* ExecutionSource = InResource->GetImplementingOuter<IPCGGraphExecutionSource>();
 
 		UE_LOG(LogPCG, Log, TEXT("[PCGMANAGEDRESOURCES] [%s/%s] UPCGManagedComponent::Release, hidden, component:"),
-			(Component && Component->GetOwner()) ? *Component->GetOwner()->GetName() : TEXT("MISSINGCOMPONENT"),
-			(Component && Component->GetGraph()) ? *Component->GetGraph()->GetName() : TEXT("MISSINGGRAPH"));
+			*GetExecutionSourceString(ExecutionSource),
+			*GetExecutionSourceGraphString(ExecutionSource));
 
-		LogResource(Component, InResource);
+		LogResource(ExecutionSource, InResource);
 #endif
 	}
 
@@ -413,11 +426,11 @@ namespace PCGGeneratedResourcesLogging
 			return;
 		}
 
-		const UPCGComponent* Component = InResource ? Cast<UPCGComponent>(InResource->GetOuter()) : nullptr;
+		const IPCGGraphExecutionSource* ExecutionSource = InResource->GetImplementingOuter<IPCGGraphExecutionSource>();
 
 		UE_LOG(LogPCG, Log, TEXT("[PCGMANAGEDRESOURCES] [%s/%s] UPCGManagedComponent::Release, delete null component"),
-			(Component && Component->GetOwner()) ? *Component->GetOwner()->GetName() : TEXT("MISSINGCOMPONENT"),
-			(Component && Component->GetGraph()) ? *Component->GetGraph()->GetName() : TEXT("MISSINGGRAPH"));
+			*GetExecutionSourceString(ExecutionSource),
+			*GetExecutionSourceGraphString(ExecutionSource));
 #endif
 	}
 }
diff --git a/Engine/Plugins/PCG/Source/PCG/Private/Utils/PCGGeneratedResourcesLogging.h b/Engine/Plugins/PCG/Source/PCG/Private/Utils/PCGGeneratedResourcesLogging.h
--- a/Engine/Plugins/PCG/Source/PCG/Private/Utils/PCGGeneratedResourcesLogging.h
+++ b/Engine/Plugins/PCG/Source/PCG/Private/Utils/PCGGeneratedResourcesLogging.h
@@ -6,28 +6,28 @@
 #include "Containers/Set.h"
 #include "UObject/SoftObjectPtr.h"
 
-class UPCGComponent;
+class IPCGGraphExecutionSource;
 class UPCGManagedComponent;
 class UPCGManagedResource;
 class AActor;
 
 namespace PCGGeneratedResourcesLogging
 {
-	void LogAddToManagedResources(const UPCGComponent* InComponent, UPCGManagedResource* InResource);
+	void LogAddToManagedResources(const IPCGGraphExecutionSource* InExecutionSource, UPCGManagedResource* InResource);
 	
-	void LogCleanupLocal(const UPCGComponent* InComponent, bool bRemoveComponents);
+	void LogCleanupLocal(const IPCGGraphExecutionSource* InExecutionSource, bool bRemoveComponents);
 	
-	void LogCleanupLocalImmediate(const UPCGComponent* InComponent, bool bHardRelease, const TArray<UPCGManagedResource*>& GeneratedResources);
-	void LogCleanupLocalImmediateResource(const UPCGComponent* InComponent, UPCGManagedResource* InResource);
-	void LogCleanupLocalImmediateFinished(const UPCGComponent* InComponent, const TArray<UPCGManagedResource*>& GeneratedResources);
+	void LogCleanupLocalImmediate(const IPCGGraphExecutionSource* InExecutionSource, bool bHardRelease, const TArray<UPCGManagedResource*>& GeneratedResources);
+	void LogCleanupLocalImmediateResource(const IPCGGraphExecutionSource* InExecutionSource, UPCGManagedResource* InResource);
+	void LogCleanupLocalImmediateFinished(const IPCGGraphExecutionSource* InExecutionSource, const TArray<UPCGManagedResource*>& GeneratedResources);
 	
-	void LogCreateCleanupTask(const UPCGComponent* InComponent, bool bRemoveComponents);
-	void LogCreateCleanupTaskResource(const UPCGComponent* InComponent, UPCGManagedResource* InResource);
-	void LogCreateCleanupTaskFinished(const UPCGComponent* InComponent, const TArray<TObjectPtr<UPCGManagedResource>>* InGeneratedResources);
+	void LogCreateCleanupTask(const IPCGGraphExecutionSource* InExecutionSource, bool bRemoveComponents);
+	void LogCreateCleanupTaskResource(const IPCGGraphExecutionSource* InExecutionSource, UPCGManagedResource* InResource);
+	void LogCreateCleanupTaskFinished(const IPCGGraphExecutionSource* InExecutionSource, const TArray<TObjectPtr<UPCGManagedResource>>* InGeneratedResources);
 
-	void LogCleanupUnusedManagedResources(const UPCGComponent* InComponent, const TArray<UPCGManagedResource*>& GeneratedResources);
-	void LogCleanupUnusedManagedResourcesResource(const UPCGComponent* InComponent, UPCGManagedResource* Resource);
-	void LogCleanupUnusedManagedResourcesFinished(const UPCGComponent* InComponent, const TArray<UPCGManagedResource*>& GeneratedResources);
+	void LogCleanupUnusedManagedResources(const IPCGGraphExecutionSource* InExecutionSource, const TArray<UPCGManagedResource*>& GeneratedResources);
+	void LogCleanupUnusedManagedResourcesResource(const IPCGGraphExecutionSource* InExecutionSource, UPCGManagedResource* Resource);
+	void LogCleanupUnusedManagedResourcesFinished(const IPCGGraphExecutionSource* InExecutionSource, const TArray<UPCGManagedResource*>& GeneratedResources);
 
 	void LogManagedActorsRelease(const UPCGManagedResource* InResource, const TArray<TSoftObjectPtr<AActor>>& ActorsToDelete, bool bHardRelease, bool bOnlyMarkedForCleanup);
 
diff --git a/Engine/Plugins/PCG/Source/PCG/Public/PCGCommon.h b/Engine/Plugins/PCG/Source/PCG/Public/PCGCommon.h
--- a/Engine/Plugins/PCG/Source/PCG/Public/PCGCommon.h
+++ b/Engine/Plugins/PCG/Source/PCG/Public/PCGCommon.h
@@ -48,21 +48,29 @@ namespace PCG
 	public:
 		UE_NONCOPYABLE(TScopeLock);
 
-		UE_NODISCARD_CTOR inline explicit TScopeLock(LockType& InLock, bool bTryLock = false)
-			: Lock(&InLock), bLocked(false)
+		UE_NODISCARD_CTOR inline explicit TScopeLock(LockType* InLock, bool bTryLock = false)
+			: Lock(InLock), bLocked(false)
 		{
-			check(Lock);
-			if (bTryLock)
-			{
-				bLocked = Lock->TryLock();
-			}
-			else
+			if (InLock)
 			{
-				Lock->Lock();
-				bLocked = true;
+				if (bTryLock)
+				{
+					bLocked = InLock->TryLock();
+				}
+				else
+				{
+					InLock->Lock();
+					bLocked = true;
+				}
 			}
 		}
 
+
+		UE_NODISCARD_CTOR inline explicit TScopeLock(LockType& InLock, bool bTryLock = false)
+			: TScopeLock(&InLock, bTryLock)
+		{
+		}
+
 		inline ~TScopeLock()
 		{
 			if (bLocked && Lock)
diff --git a/Engine/Plugins/PCG/Source/PCG/Public/PCGComponent.h b/Engine/Plugins/PCG/Source/PCG/Public/PCGComponent.h
--- a/Engine/Plugins/PCG/Source/PCG/Public/PCGComponent.h
+++ b/Engine/Plugins/PCG/Source/PCG/Public/PCGComponent.h
@@ -6,6 +6,7 @@
 #include "PCGComponentExecutionState.h"
 #include "PCGGraphExecutionInspection.h"
 #include "PCGGraphExecutionStateInterface.h"
+#include "PCGManagedResourceContainer.h"
 #include "PCGNode.h"
 #include "PCGSettings.h"
 #include "Data/PCGToolData.h"
@@ -200,7 +201,7 @@ public:
 	UE_API void AddActorsToManagedResources(const TArray<AActor*>& InActors);
 
 public:
-	bool AreManagedResourcesAccessible() const { return !GeneratedResourcesInaccessible; }
+	bool AreManagedResourcesAccessible() const { return ManagedResourceContainer.AreResourcesAccessible(); }
 	UE_API void ForEachManagedResource(TFunctionRef<void(UPCGManagedResource*)> InFunction);
 	UE_API void ForEachConstManagedResource(TFunctionRef<void(const UPCGManagedResource*)> InFunction) const;
 
@@ -679,6 +680,8 @@ private:
 	UE_API void RefreshAfterGraphChanged(UPCGGraphInterface* InGraph, EPCGChangeType ChangeType);
 	UE_API void OnGraphChanged(UPCGGraphInterface* InGraph, EPCGChangeType ChangeType);
 
+	UE_API void OnManagedResourceAdded(UPCGManagedResource* InResource);
+
 #if WITH_EDITOR
 	UE_API virtual void PreEditChange(FProperty* PropertyAboutToChange) override;
 	UE_API virtual void PostEditChangeProperty(FPropertyChangedEvent& PropertyChangedEvent) override;
@@ -735,15 +738,16 @@ private:
 	TSet<TSoftObjectPtr<AActor>> GeneratedActors_DEPRECATED;
 #endif
 
+	UPROPERTY()
+	FPCGManagedResourceContainer ManagedResourceContainer;
+
 	// NOTE: This should not be made visible or editable because it will change the way the BP actors are
 	// duplicated/setup and might trigger an ensure in the resources.
+	UE_DEPRECATED(5.8, "Moved to ManagedResourceContainer")
 	UPROPERTY()
-	TArray<TObjectPtr<UPCGManagedResource>> GeneratedResources;
+	TArray<TObjectPtr<UPCGManagedResource>> GeneratedResources_DEPRECATED;
 
 #if WITH_EDITORONLY_DATA
-	UPROPERTY(Transient)
-	TArray<TObjectPtr<UPCGManagedResource>> LoadedPreviewResources;
-
 	UPROPERTY(Transient)
 	FPCGDataCollection LoadedPreviewGeneratedGraphOutput;
 
@@ -751,9 +755,6 @@ private:
 	bool bGenerationInProgress = false;
 #endif
 
-	// When doing a cleanup, locking resource modification. Used as sentinel.
-	bool GeneratedResourcesInaccessible = false;
-
 	// Flag to indicate whether this component was already generated on load.
 	bool bGeneratedOffline = false;
 
diff --git a/Engine/Plugins/PCG/Source/PCG/Public/PCGComponentExecutionState.h b/Engine/Plugins/PCG/Source/PCG/Public/PCGComponentExecutionState.h
--- a/Engine/Plugins/PCG/Source/PCG/Public/PCGComponentExecutionState.h
+++ b/Engine/Plugins/PCG/Source/PCG/Public/PCGComponentExecutionState.h
@@ -41,8 +41,9 @@ public:
 	virtual FPCGTaskId GetGenerationTaskId() const override;
 	virtual void StoreOutputDataForPin(const FString& InResourceKey, const FPCGDataCollection& InData) const override;
 	virtual const FPCGDataCollection* RetrieveOutputDataForPin(const FString& InResourceKey) const override;
-	virtual void AddToManagedResources(UPCGManagedResource* InResource) override;
-	
+	virtual FPCGManagedResourceContainer* GetManagedResourceContainer() override;
+	virtual FTransactionallySafeCriticalSection* GetManagedResourceContainerLock() override;
+	virtual void OnManagedResourceAdded(UPCGManagedResource* InResource) override;
 	virtual FPCGTaskId Generate(const FGenerateParams& InGenerateParams = {}) override;
 
 	UE_DEPRECATED(5.8, "Call/Implement Generate instead.")
diff --git a/Engine/Plugins/PCG/Source/PCG/Public/PCGGraphExecutionStateInterface.h b/Engine/Plugins/PCG/Source/PCG/Public/PCGGraphExecutionStateInterface.h
--- a/Engine/Plugins/PCG/Source/PCG/Public/PCGGraphExecutionStateInterface.h
+++ b/Engine/Plugins/PCG/Source/PCG/Public/PCGGraphExecutionStateInterface.h
@@ -19,6 +19,7 @@ class UPCGGraph;
 class UPCGGraphInstance;
 class UPCGManagedResource;
 struct FPCGGridDescriptor;
+struct FPCGManagedResourceContainer;
 
 namespace PCGUtils
 {
@@ -139,8 +140,14 @@ public:
 	/** Lookup data using a resource key that identifies the pin. */
 	virtual const FPCGDataCollection* RetrieveOutputDataForPin(const FString& InResourceKey) const { return nullptr; }
 
+	UE_DEPRECATED(5.8, "Use GetManagedResourceContainer instead.")
 	PCG_API virtual void AddToManagedResources(UPCGManagedResource* InResource);
 
+	virtual FPCGManagedResourceContainer* GetManagedResourceContainer() { return nullptr; }
+	virtual FTransactionallySafeCriticalSection* GetManagedResourceContainerLock() { return nullptr; }
+
+	virtual void OnManagedResourceAdded(UPCGManagedResource* InResource) {}
+
 	/** Returns the subsystem associated with this source. By default, it will return the UPCGSubsystem if the source has a world, otherwise, the EditorSubsystem (in editor only, obviously). */
 	PCG_API virtual IPCGBaseSubsystem* GetSubsystem() const;
 
diff --git a/Engine/Plugins/PCG/Source/PCG/Public/PCGManagedResource.h b/Engine/Plugins/PCG/Source/PCG/Public/PCGManagedResource.h
--- a/Engine/Plugins/PCG/Source/PCG/Public/PCGManagedResource.h
+++ b/Engine/Plugins/PCG/Source/PCG/Public/PCGManagedResource.h
@@ -10,6 +10,10 @@
 #include "ISMPartition/ISMComponentDescriptor.h"
 #include "Animation/AnimBank.h"
 
+#if WITH_EDITOR
+#include "WorldPartition/WorldPartitionHandle.h"
+#endif
+
 #include "PCGManagedResource.generated.h"
 
 #define UE_API PCG_API
@@ -93,6 +97,20 @@ protected:
 #endif
 };
 
+struct FPCGManagedActorLoadingScope
+{
+	FPCGManagedActorLoadingScope();
+	FPCGManagedActorLoadingScope(bool bInShouldLoadActors);
+	virtual ~FPCGManagedActorLoadingScope() = default;
+
+	void AddResources(IPCGGraphExecutionSource* InExecutionSource, const TArray<UPCGManagedResource*>& InManagedResources);
+
+#if WITH_EDITOR
+	TSet<FWorldPartitionReference> LoadedActors;
+	bool bShouldLoadActors;
+#endif
+};
+
 UCLASS(MinimalAPI, BlueprintType)
 class UPCGManagedActors : public UPCGManagedResource
 {
diff --git a/Engine/Plugins/PCG/Source/PCG/Public/PCGManagedResourceContainer.h b/Engine/Plugins/PCG/Source/PCG/Public/PCGManagedResourceContainer.h
new file mode 100644
--- /dev/null
+++ b/Engine/Plugins/PCG/Source/PCG/Public/PCGManagedResourceContainer.h
@@ -0,0 +1,79 @@
+// Copyright Epic Games, Inc. All Rights Reserved.
+
+#pragma once
+
+#include "PCGCommon.h"
+#include "PCGGraphExecutionStateInterface.h"
+#include "PCGManagedResource.h"
+
+#include "UObject/UObjectGlobals.h"
+#include "UObject/WeakInterfacePtr.h"
+
+#include "PCGManagedResourceContainer.generated.h"
+
+class UPCGManagedResource;
+
+USTRUCT()
+struct FPCGManagedResourceContainer
+{
+	GENERATED_BODY()
+
+public:
+	PCG_API void SetOwner(IPCGGraphExecutionSource* InOwner);
+	
+	bool AreResourcesAccessible() const { return !bAreResourcesInaccessible; }
+	bool IsEmpty() { return GeneratedResources.IsEmpty(); }
+	
+	PCG_API void AddManagedResource(UPCGManagedResource* InManagedResource);
+	PCG_API void CleanupUnusedManagedResources();
+	PCG_API void SetManagedResources(const TArray<TObjectPtr<UPCGManagedResource>>& InManagedResources);
+	PCG_API TArray<TObjectPtr<UPCGManagedResource>> GetManagedResourcesCopy() const;
+	
+	TArray<TObjectPtr<UPCGManagedResource>>& GetManagedResourcesNoLock() { return GeneratedResources; }
+	const TArray<TObjectPtr<UPCGManagedResource>>& GetConstManagedResourcesNoLock() const { return GeneratedResources; }
+	
+	PCG_API void ForEachManagedResource(TFunctionRef<void(UPCGManagedResource*)> InFunction);
+	PCG_API void ForEachConstManagedResource(TFunctionRef<void(const UPCGManagedResource*)> InFunction) const;
+	PCG_API bool IsAnyObjectManagedByResource(const TArrayView<const UObject*> InObjects) const;
+
+	void EndPlay();
+
+	struct FCleanupTask : public FPCGManagedActorLoadingScope
+	{
+		TWeakInterfacePtr<IPCGGraphExecutionSource> WeakExecutionSource;
+		bool bIsFirstIteration = true;
+		int32 ResourceIndex = -1;
+		TSet<TSoftObjectPtr<AActor>> ActorsToDelete;
+
+		PCG_API void Abort();
+		PCG_API bool Execute(bool bReleaseManagedResources);
+	};
+	
+	PCG_API TSharedPtr<FCleanupTask> CreateCleanupTask();
+	PCG_API void Cleanup(bool bReleaseManagedResources);
+
+#if WITH_EDITOR
+	void SetLoadedPreviewResourcesNoLock(const TArray<TObjectPtr<UPCGManagedResource>>& InManagedResources) { LoadedPreviewResources = InManagedResources; }
+	const TArray<TObjectPtr<UPCGManagedResource>>& GetConstLoadedPreviewResourcesNoLock() const { return LoadedPreviewResources; }
+
+	void MarkResourcesAsTransientOnLoad();
+	bool ChangeTransientState(EPCGEditorDirtyMode NewEditingMode);
+	bool DeletePreviewResources();
+
+	void BeginPreviewSave(TArray<TObjectPtr<UPCGManagedResource>>& OutGeneratedCopy);
+	void EndPreviewSave(const TArray<TObjectPtr<UPCGManagedResource>>& InGeneratedCopy);
+#endif	
+
+private:
+	bool bAreResourcesInaccessible = false;
+
+	UPROPERTY()
+	TArray<TObjectPtr<UPCGManagedResource>> GeneratedResources;
+
+#if WITH_EDITORONLY_DATA
+	UPROPERTY()
+	TArray<TObjectPtr<UPCGManagedResource>> LoadedPreviewResources;
+#endif
+
+	IPCGGraphExecutionSource* Owner = nullptr;
+};
\ No newline at end of file
