Fix CIS

#rnx

diff --git a/Engine/Plugins/PCG/Source/PCG/Private/Utils/PCGGeneratedResourcesLogging.cpp b/Engine/Plugins/PCG/Source/PCG/Private/Utils/PCGGeneratedResourcesLogging.cpp
--- a/Engine/Plugins/PCG/Source/PCG/Private/Utils/PCGGeneratedResourcesLogging.cpp
+++ b/Engine/Plugins/PCG/Source/PCG/Private/Utils/PCGGeneratedResourcesLogging.cpp
@@ -323,7 +323,7 @@ namespace PCGGeneratedResourcesLogging
 			return;
 		}
 
-		const IPCGGraphExecutionSource* ExecutionSource = InResource->GetImplementingOuter<IPCGGraphExecutionSource>();
+		const IPCGGraphExecutionSource* ExecutionSource = InResource ? InResource->GetImplementingOuter<IPCGGraphExecutionSource>() : nullptr;
 
 		UE_LOG(LogPCG, Log, TEXT("[PCGMANAGEDRESOURCES] [%s/%s] UPCGManagedResource::Release, SOFT release:"),
 			*GetExecutionSourceString(ExecutionSource),
@@ -341,7 +341,7 @@ namespace PCGGeneratedResourcesLogging
 			return;
 		}
 
-		const IPCGGraphExecutionSource* ExecutionSource = InResource->GetImplementingOuter<IPCGGraphExecutionSource>();
+		const IPCGGraphExecutionSource* ExecutionSource = InResource ? InResource->GetImplementingOuter<IPCGGraphExecutionSource>() : nullptr;
 
 		UE_LOG(LogPCG, Log, TEXT("[PCGMANAGEDRESOURCES] [%s/%s] UPCGManagedResource::Release, HARD release:"),
 			*GetExecutionSourceString(ExecutionSource),
@@ -359,7 +359,7 @@ namespace PCGGeneratedResourcesLogging
 			return;
 		}
 
-		const IPCGGraphExecutionSource* ExecutionSource = InResource->GetImplementingOuter<IPCGGraphExecutionSource>();
+		const IPCGGraphExecutionSource* ExecutionSource = InResource ? InResource->GetImplementingOuter<IPCGGraphExecutionSource>() : nullptr;
 		const bool bMarkedTransientOnLoad = InResource && InResource->IsMarkedTransientOnLoad();
 
 		UE_LOG(LogPCG, Log, TEXT("[PCGMANAGEDRESOURCES] [%s/%s] UPCGManagedActors::Release, %s release, bMarkedTransientOnLoad: %d, bOnlyMarkedForCleanup: %d, actor count %d scheduled for delete"),
@@ -408,7 +408,7 @@ namespace PCGGeneratedResourcesLogging
 			return;
 		}
 
-		const IPCGGraphExecutionSource* ExecutionSource = InResource->GetImplementingOuter<IPCGGraphExecutionSource>();
+		const IPCGGraphExecutionSource* ExecutionSource = InResource ? InResource->GetImplementingOuter<IPCGGraphExecutionSource>() : nullptr;
 
 		UE_LOG(LogPCG, Log, TEXT("[PCGMANAGEDRESOURCES] [%s/%s] UPCGManagedComponent::Release, hidden, component:"),
 			*GetExecutionSourceString(ExecutionSource),
@@ -426,7 +426,7 @@ namespace PCGGeneratedResourcesLogging
 			return;
 		}
 
-		const IPCGGraphExecutionSource* ExecutionSource = InResource->GetImplementingOuter<IPCGGraphExecutionSource>();
+		const IPCGGraphExecutionSource* ExecutionSource = InResource ? InResource->GetImplementingOuter<IPCGGraphExecutionSource>() : nullptr;
 
 		UE_LOG(LogPCG, Log, TEXT("[PCGMANAGEDRESOURCES] [%s/%s] UPCGManagedComponent::Release, delete null component"),
 			*GetExecutionSourceString(ExecutionSource),
