[PCG] Fix intended logical operation.

#rnx

diff --git a/Engine/Plugins/PCG/Source/PCG/Private/PCGComponent.cpp b/Engine/Plugins/PCG/Source/PCG/Private/PCGComponent.cpp
--- a/Engine/Plugins/PCG/Source/PCG/Private/PCGComponent.cpp
+++ b/Engine/Plugins/PCG/Source/PCG/Private/PCGComponent.cpp
@@ -1590,7 +1590,7 @@ void UPCGComponent::BeginPlay()
 	}
 
 	// Procedural instances are never persisted so always require generation.
-	const bool bAlreadyGenerated = bGenerated & !bProceduralInstancesInUse;
+	const bool bAlreadyGenerated = bGenerated && !bProceduralInstancesInUse;
 
 	// Either this is the original component and it is non-null or this is a local component and we need the original component to be loaded to trigger a generation
 	UPCGComponent* OriginalComponent = GetOriginalComponent();
