[PCG] RGS optimization - reduce GT time spent in RGS in runtime grass test level by 30%.

[RN] [PCG] Optimized runtime generation scheduler to reduce per frame cost by up to 30%.

diff --git a/Engine/Plugins/PCG/Source/PCG/Private/RuntimeGen/PCGRuntimeGenScheduler.cpp b/Engine/Plugins/PCG/Source/PCG/Private/RuntimeGen/PCGRuntimeGenScheduler.cpp
--- a/Engine/Plugins/PCG/Source/PCG/Private/RuntimeGen/PCGRuntimeGenScheduler.cpp
+++ b/Engine/Plugins/PCG/Source/PCG/Private/RuntimeGen/PCGRuntimeGenScheduler.cpp
@@ -195,6 +195,8 @@ FPCGRuntimeGenScheduler::FPCGRuntimeGenScheduler(UWorld* InWorld, FPCGActorAndCo
 
 	World = InWorld;
 	Subsystem = UPCGSubsystem::GetInstance(World);
+	WorldPartitionSubsystem = World ? UWorld::GetSubsystem<UWorldPartitionSubsystem>(World) : nullptr;
+
 	ActorAndComponentMapping = InActorAndComponentMapping;
 	GenSourceManager = new FPCGGenSourceManager(InWorld);
 	bPoolingWasEnabledLastFrame = PCGRuntimeGenSchedulerHelpers::CVarRuntimeGenerationEnablePooling.GetValueOnAnyThread();
@@ -315,7 +317,7 @@ void FPCGRuntimeGenScheduler::Tick(APCGWorldActor* InPCGWorldActor, double InEnd
 		Inputs.PCGWorldActor = InPCGWorldActor;
 		Inputs.AllPartitionedComponents = ActorAndComponentMapping->GetAllRegisteredPartitionedComponents();
 		Inputs.AllNonPartitionedComponents = ActorAndComponentMapping->GetAllRegisteredNonPartitionedComponents();
-		Inputs.GeneratedComponents = GeneratedComponents;
+		Inputs.GeneratedComponents = &GeneratedComponents;
 
 		ChangeDetector.PreTick();
 		TickQueueComponentsForGeneration(Inputs, ComponentsToGenerate);
@@ -406,9 +408,7 @@ bool FPCGRuntimeGenScheduler::ShouldTick()
 	// Check if the non-spatially streamed cell is ready (only if world streaming queries are enabled). We should not do any processing until the world is ready.
 	if (!bNonSpatialCellFinishedStreaming)
 	{
-		const UWorldPartitionSubsystem* WorldPartitionSubsystem = PCGRuntimeGenSchedulerHelpers::CVarEnableWorldStreamingQueries.GetValueOnGameThread() ? UWorld::GetSubsystem<UWorldPartitionSubsystem>(World) : nullptr;
-
-		if (WorldPartitionSubsystem)
+		if (WorldPartitionSubsystem && PCGRuntimeGenSchedulerHelpers::CVarEnableWorldStreamingQueries.GetValueOnGameThread())
 		{
 			TRACE_CPUPROFILER_EVENT_SCOPE(IsNonSpatialWorldStreamingComplete);
 
@@ -473,22 +473,18 @@ void FPCGRuntimeGenScheduler::TickQueueComponentsForGeneration(
 	TArray<FWorldPartitionStreamingQuerySource> StreamingQuerySources;
 	StreamingQuerySources.Reserve(1);
 	
-	bool bCheckStreaming = PCGRuntimeGenSchedulerHelpers::CVarEnableWorldStreamingQueries.GetValueOnGameThread();
-	const UWorldPartitionSubsystem* WorldPartitionSubsystem = bCheckStreaming ? UWorld::GetSubsystem<UWorldPartitionSubsystem>(Inputs.PCGWorldActor->GetWorld()) : nullptr;
+	bool bCheckStreaming = false;
 
-	if (WorldPartitionSubsystem)
+	if (WorldPartitionSubsystem && PCGRuntimeGenSchedulerHelpers::CVarEnableWorldStreamingQueries.GetValueOnGameThread())
 	{
 		FWorldPartitionStreamingQuerySource& QuerySource = StreamingQuerySources.Emplace_GetRef();
 		QuerySource.bSpatialQuery = true;
 		QuerySource.bUseGridLoadingRange = false;
 		QuerySource.bDataLayersOnly = false;
-	}
-	else
-	{
-		bCheckStreaming = false;
+		bCheckStreaming = true;
 	}
 
-	auto IsWorldStreamingComplete = [WorldPartitionSubsystem, &StreamingQuerySources, &CachedStreamingQueryResults=CachedStreamingQueryResults](const FVector& InLocation, float InGridSize)
+	auto IsWorldStreamingComplete = [WorldPartitionSubsystem=WorldPartitionSubsystem, &StreamingQuerySources, &CachedStreamingQueryResults=CachedStreamingQueryResults](const FVector& InLocation, float InGridSize)
 	{
 		FStreamingCompleteQueryKey Key{ InLocation, InGridSize };
 		if (bool* FoundResult = CachedStreamingQueryResults.Find(Key))
@@ -670,7 +666,7 @@ void FPCGRuntimeGenScheduler::TickQueueComponentsForGeneration(
 					
 					// Ignore components that have already been generated or marked for generation. Unbounded grid size means not-partitioned.
 					FGridGenerationKey Key(PCGHiGenGrid::UnboundedGridSize(), FIntVector(0), OriginalComponent);
-					if (!Inputs.GeneratedComponents.Contains(Key) &&
+					if (!Inputs.GeneratedComponents->Contains(Key) &&
 						(!Policy || Policy->ShouldGenerate(GenSource, OriginalComponentBounds, bIs2DGrid)))
 					{
 						check(Key.GetGridDescriptor().Is2DGrid() == bIs2DGrid);
@@ -713,7 +709,7 @@ void FPCGRuntimeGenScheduler::TickQueueComponentsForGeneration(
 								FGridGenerationKey Key(GridSize, GridCoords, OriginalComponent);
 
 								// Ignore components that have already been generated or marked for generation.
-								if (Inputs.GeneratedComponents.Find({ GridSize, GridCoords, OriginalComponent }))
+								if (Inputs.GeneratedComponents->Find({ GridSize, GridCoords, OriginalComponent }))
 								{
 									continue;
 								}
@@ -804,7 +800,7 @@ void FPCGRuntimeGenScheduler::TickQueueComponentsForGeneration(
 		// todo_pcg: Worth adding change detection here? Decide on quantization - perhaps a proportion of minimum volume extent or such?
 		
 		// The generation key for a non-partitioned component should always have unbounded grid size and 0,0,0 cell coord.
-		if (Inputs.GeneratedComponents.Contains({ PCGHiGenGrid::UnboundedGridSize(), FIntVector(0), OriginalComponent }))
+		if (Inputs.GeneratedComponents->Contains({ PCGHiGenGrid::UnboundedGridSize(), FIntVector(0), OriginalComponent }))
 		{
 			continue;
 		}
@@ -923,81 +919,85 @@ void FPCGRuntimeGenScheduler::TickCleanup(const TSet<IPCGGenSourceBase*>& InGenS
 	std::atomic<int> ComponentsToCleanCounter = 0;
 	std::atomic<int> InvalidKeysCounter = 0;
 
-	ParallelFor(NumGeneratedComponents, [&GeneratedComponentsArray, &ComponentsToClean, &InvalidKeys, &ComponentsToCleanCounter, &InvalidKeysCounter, ActorAndComponentMapping=ActorAndComponentMapping, &CheckIfAllGenSourcesWantToCleanup](int32 Index)
-	{
-		TRACE_CPUPROFILER_EVENT_SCOPE(SelectComponentForCleanup);
-		LLM_SCOPE_BYTAG(PCG);
-
-		const FGridGenerationKey& GenerationKey = GeneratedComponentsArray[Index];
-
-		if (!GenerationKey.IsValid())
+	ParallelFor(
+		TEXT("FPCGRuntimeGenScheduler::TickCleanup"),
+		NumGeneratedComponents,
+		/*MinBatchSize*/6, // Optimal value found in testing.
+		[&GeneratedComponentsArray, &ComponentsToClean, &InvalidKeys, &ComponentsToCleanCounter, &InvalidKeysCounter, ActorAndComponentMapping=ActorAndComponentMapping, &CheckIfAllGenSourcesWantToCleanup](int32 Index)
 		{
-			const int WriteIndex = InvalidKeysCounter.fetch_add(1);
-			InvalidKeys[WriteIndex] = GenerationKey;
-			return;
-		}
+			TRACE_CPUPROFILER_EVENT_SCOPE(SelectComponentForCleanup);
+			LLM_SCOPE_BYTAG(PCG);
 
-		const FPCGGridDescriptor GridDescriptor = GenerationKey.GetGridDescriptor();
-		const EPCGHiGenGrid Grid = GridDescriptor.GetHiGenGrid();
-		const FIntVector& GridCoords = GenerationKey.GetGridCoords();
-		UPCGComponent* OriginalComponent = GenerationKey.GetOriginalComponent();
-		check(OriginalComponent);
+			const FGridGenerationKey& GenerationKey = GeneratedComponentsArray[Index];
 
-		const UPCGSchedulingPolicyBase* Policy = OriginalComponent->GetRuntimeGenSchedulingPolicy();
-		ensure(Policy);
-
-		const double CleanupRadius = OriginalComponent->GetCleanupRadiusFromGrid(Grid);
-		const double CleanupRadiusSquared = CleanupRadius * CleanupRadius;
-
-		// If the Grid is unbounded, we have a non-partitioned or unbounded component.
-		if (Grid == EPCGHiGenGrid::Unbounded)
-		{
-			if (!OriginalComponent->bActivated)
+			if (!GenerationKey.IsValid())
 			{
-				const int WriteIndex = ComponentsToCleanCounter.fetch_add(1);
-				ComponentsToClean[WriteIndex] = { GenerationKey, OriginalComponent };
+				const int WriteIndex = InvalidKeysCounter.fetch_add(1);
+				InvalidKeys[WriteIndex] = GenerationKey;
 				return;
 			}
 
-			const FBox GridBounds = OriginalComponent->GetGridBounds();
+			const FPCGGridDescriptor GridDescriptor = GenerationKey.GetGridDescriptor();
+			const EPCGHiGenGrid Grid = GridDescriptor.GetHiGenGrid();
+			const FIntVector& GridCoords = GenerationKey.GetGridCoords();
+			UPCGComponent* OriginalComponent = GenerationKey.GetOriginalComponent();
+			check(OriginalComponent);
+
+			const UPCGSchedulingPolicyBase* Policy = OriginalComponent->GetRuntimeGenSchedulingPolicy();
+			ensure(Policy);
 
-			// Only clean up if all generation sources agreed to clean up.
-			if(CheckIfAllGenSourcesWantToCleanup(Policy, GridDescriptor, GridBounds, CleanupRadiusSquared))
+			const double CleanupRadius = OriginalComponent->GetCleanupRadiusFromGrid(Grid);
+			const double CleanupRadiusSquared = CleanupRadius * CleanupRadius;
+
+			// If the Grid is unbounded, we have a non-partitioned or unbounded component.
+			if (Grid == EPCGHiGenGrid::Unbounded)
 			{
-				const int WriteIndex = ComponentsToCleanCounter.fetch_add(1);
-				ComponentsToClean[WriteIndex] = { GenerationKey, OriginalComponent };
-				return;
+				if (!OriginalComponent->bActivated)
+				{
+					const int WriteIndex = ComponentsToCleanCounter.fetch_add(1);
+					ComponentsToClean[WriteIndex] = { GenerationKey, OriginalComponent };
+					return;
+				}
+
+				const FBox GridBounds = GenerationKey.GetCachedBounds().IsSet() ? GenerationKey.GetCachedBounds().GetValue() : OriginalComponent->GetGridBounds();
+
+				// Only clean up if all generation sources agreed to clean up.
+				if (CheckIfAllGenSourcesWantToCleanup(Policy, GridDescriptor, GridBounds, CleanupRadiusSquared))
+				{
+					const int WriteIndex = ComponentsToCleanCounter.fetch_add(1);
+					ComponentsToClean[WriteIndex] = { GenerationKey, OriginalComponent };
+					return;
+				}
 			}
-		}
-		// Otherwise, we have a local component.
-		else
-		{
-			IPCGGraphExecutionSource* LocalComponent = GenerationKey.GetCachedLocalComponent();
-			if (!LocalComponent)
+			// Otherwise, we have a local component.
+			else
 			{
-				LocalComponent = ActorAndComponentMapping->GetLocalComponent(GridDescriptor, GridCoords, OriginalComponent);
-			}
+				IPCGGraphExecutionSource* LocalComponent = GenerationKey.GetCachedLocalComponent();
+				if (!LocalComponent)
+				{
+					LocalComponent = ActorAndComponentMapping->GetLocalComponent(GridDescriptor, GridCoords, OriginalComponent);
+				}
 
-			APCGPartitionActor* PartitionActor = LocalComponent ? Cast<APCGPartitionActor>(CastChecked<UPCGComponent>(LocalComponent)->GetOwner()) : nullptr;
-			if (!PartitionActor || !OriginalComponent->bActivated)
-			{
-				// Attempt to clean even in failure case to avoid leaking resources.
-				const int WriteIndex = ComponentsToCleanCounter.fetch_add(1);
-				ComponentsToClean[WriteIndex] = { GenerationKey, LocalComponent };
-				return;
-			}
+				APCGPartitionActor* PartitionActor = LocalComponent ? Cast<APCGPartitionActor>(CastChecked<UPCGComponent>(LocalComponent)->GetOwner()) : nullptr;
+				if (!PartitionActor || !OriginalComponent->bActivated)
+				{
+					// Attempt to clean even in failure case to avoid leaking resources.
+					const int WriteIndex = ComponentsToCleanCounter.fetch_add(1);
+					ComponentsToClean[WriteIndex] = { GenerationKey, LocalComponent };
+					return;
+				}
 
-			const FBox GridBounds = LocalComponent->GetExecutionState().GetBounds();
+				const FBox GridBounds = GenerationKey.GetCachedBounds().IsSet() ? GenerationKey.GetCachedBounds().GetValue() : LocalComponent->GetExecutionState().GetBounds();
 
-			// Only clean up if all generation sources agreed to clean up.
-			if (CheckIfAllGenSourcesWantToCleanup(Policy, GridDescriptor, GridBounds, CleanupRadiusSquared))
-			{
-				const int WriteIndex = ComponentsToCleanCounter.fetch_add(1);
-				ComponentsToClean[WriteIndex] = { GenerationKey, LocalComponent };
-				return;
+				// Only clean up if all generation sources agreed to clean up.
+				if (CheckIfAllGenSourcesWantToCleanup(Policy, GridDescriptor, GridBounds, CleanupRadiusSquared))
+				{
+					const int WriteIndex = ComponentsToCleanCounter.fetch_add(1);
+					ComponentsToClean[WriteIndex] = { GenerationKey, LocalComponent };
+					return;
+				}
 			}
-		}
-	});
+		});
 
 	const int NumInvalidKeys = InvalidKeysCounter.load();
 	const int NumComponentsToClean = ComponentsToCleanCounter.load();
@@ -1493,9 +1493,9 @@ void FPCGRuntimeGenScheduler::OnObjectsReplaced(const TMap<UObject*, UObject*>&
 		{
 			for (FGridGenerationKey& GenerationKey : GeneratedComponents)
 			{
-				if (GenerationKey.OriginalComponent == OldComponent)
+				if (GenerationKey.GetOriginalComponent() == OldComponent)
 				{
-					GenerationKey.OriginalComponent = NewComponent;
+					GenerationKey.SetOriginalComponent(NewComponent);
 				}
 			}
 
@@ -1512,6 +1512,30 @@ void FPCGRuntimeGenScheduler::OnObjectsReplaced(const TMap<UObject*, UObject*>&
 }
 #endif
 
+FPCGRuntimeGenScheduler::FGridGenerationKey::FGridGenerationKey(uint32 InGridSize, const FIntVector& InGridCoords, UPCGComponent* InOriginalComponent, IPCGGraphExecutionSource* InLocalComponent)
+	: GridSize(InGridSize)
+	, GridCoords(InGridCoords)
+	, OriginalComponent(InOriginalComponent)
+	, CachedLocalComponent(InLocalComponent)
+{
+	bUse2DGrid = InOriginalComponent ? InOriginalComponent->Use2DGrid() : true;
+
+	if (InGridSize == PCGHiGenGrid::UnboundedGridSize())
+	{
+		const AActor* Owner = InOriginalComponent ? InOriginalComponent->GetOwner() : nullptr;
+		const USceneComponent* RootComponent = Owner ? Owner->GetRootComponent() : nullptr;
+
+		if (RootComponent && RootComponent->GetMobility() != EComponentMobility::Movable)
+		{
+			CachedBounds = InOriginalComponent->GetGridBounds();
+		}
+	}
+	else if (InLocalComponent)
+	{
+		CachedBounds = InLocalComponent->GetExecutionState().GetBounds();
+	}
+}
+
 void FPCGRuntimeGenScheduler::CleanupRemainingComponents(UPCGComponent* InOriginalComponent)
 {
 	// Check for remaining PAs to cleanup
@@ -1770,8 +1794,11 @@ void FPCGRuntimeGenScheduler::CleanupDelayedRefreshComponents()
 	}
 
 	// Remove any remaining generation keys that have been registered for deferred removal.
-	GeneratedComponents = GeneratedComponents.Difference(GeneratedComponentsToRemove);
-	GeneratedComponentsToRemove.Empty();
+	if (!GeneratedComponentsToRemove.IsEmpty())
+	{
+		GeneratedComponents = GeneratedComponents.Difference(GeneratedComponentsToRemove);
+		GeneratedComponentsToRemove.Reset();
+	}
 }
 
 void FPCGRuntimeGenScheduler::RefreshComponent(IPCGGraphExecutionSource* InComponent, bool bRemovePartitionActors)
diff --git a/Engine/Plugins/PCG/Source/PCG/Public/RuntimeGen/PCGRuntimeGenScheduler.h b/Engine/Plugins/PCG/Source/PCG/Public/RuntimeGen/PCGRuntimeGenScheduler.h
--- a/Engine/Plugins/PCG/Source/PCG/Public/RuntimeGen/PCGRuntimeGenScheduler.h
+++ b/Engine/Plugins/PCG/Source/PCG/Public/RuntimeGen/PCGRuntimeGenScheduler.h
@@ -20,6 +20,7 @@ class ULevelStreaming;
 class UPCGComponent;
 class UPCGSubsystem;
 class UWorld;
+class UWorldPartitionSubsystem;
 enum class ELevelStreamingState : uint8;
 enum class ELevelStreamingTargetState : uint8;
 enum class EPCGGraphParameterEvent;
@@ -95,14 +96,7 @@ protected:
 		{
 		}
 
-		FGridGenerationKey(uint32 InGridSize, const FIntVector& InGridCoords, UPCGComponent* InOriginalComponent, IPCGGraphExecutionSource* InLocalComponent)
-			: GridSize(InGridSize)
-			, GridCoords(InGridCoords)
-			, OriginalComponent(InOriginalComponent)
-			, CachedLocalComponent(InLocalComponent)
-		{
-			bUse2DGrid = InOriginalComponent ? InOriginalComponent->Use2DGrid() : true;
-		}
+		FGridGenerationKey(uint32 InGridSize, const FIntVector& InGridCoords, UPCGComponent* InOriginalComponent, IPCGGraphExecutionSource* InLocalComponent);
 
 		bool operator==(const FGridGenerationKey& Other) const
 		{
@@ -117,13 +111,26 @@ protected:
 		bool Use2DGrid() const { return bUse2DGrid; }
 		uint32 GetGridSize() const { return GridSize; }
 		FIntVector GetGridCoords() const { return GridCoords; }
+
 		UPCGComponent* GetOriginalComponent() const { return OriginalComponent.ResolveObjectPtr(); }
+		void SetOriginalComponent(UPCGComponent* InComponent) { OriginalComponent = InComponent; }
 
 		IPCGGraphExecutionSource* GetCachedLocalComponent() const { return CachedLocalComponent.Get(); }
-		void SetCachedLocalComponent(IPCGGraphExecutionSource* InLocalComponent) const { CachedLocalComponent = InLocalComponent; }
+		void SetCachedLocalComponent(IPCGGraphExecutionSource* InLocalComponent) const
+		{
+			CachedLocalComponent = InLocalComponent;
+
+			if (InLocalComponent)
+			{
+				CachedBounds = InLocalComponent->GetExecutionState().GetBounds();
+			}
+		}
+
+		const TOptional<FBox>& GetCachedBounds() const { return CachedBounds; }
 
 		FPCGGridDescriptor GetGridDescriptor() const;
 
+	private:
 		bool bUse2DGrid = false;
 		uint32 GridSize = 0;
 		FIntVector GridCoords = FIntVector(0);
@@ -132,6 +139,8 @@ protected:
 		// Optional/opportunistic cached local execution source if one has been created.
 		mutable TWeakInterfacePtr<IPCGGraphExecutionSource> CachedLocalComponent;
 
+		mutable TOptional<FBox> CachedBounds;
+
 		friend uint32 GetTypeHash(const FGridGenerationKey& In)
 		{
 			return HashCombine(GetTypeHash(In.GridCoords), GetTypeHash(In.GridSize), GetTypeHash(In.bUse2DGrid), GetTypeHash(In.OriginalComponent));
@@ -144,11 +153,11 @@ protected:
 	// Add constructor, get moves for the sets
 	struct FTickQueueComponentsForGenerationInputs
 	{
-		const TSet<IPCGGenSourceBase*>* GenSources;
-		const APCGWorldActor* PCGWorldActor;
+		const TSet<IPCGGenSourceBase*>* GenSources = nullptr;
+		const APCGWorldActor* PCGWorldActor = nullptr;
 		TSet<UPCGComponent*> AllPartitionedComponents;
 		TSet<UPCGComponent*> AllNonPartitionedComponents;
-		TSet<FGridGenerationKey> GeneratedComponents;
+		TSet<FGridGenerationKey>* GeneratedComponents = nullptr;
 	};
 
 	/** Queue nearby components for generation. */
@@ -238,6 +247,7 @@ private:
 	FPCGGenSourceManager* GenSourceManager = nullptr;
 	UPCGSubsystem* Subsystem = nullptr;
 	UWorld* World = nullptr;
+	const UWorldPartitionSubsystem* WorldPartitionSubsystem = nullptr;
 
 	bool bPoolingWasEnabledLastFrame = true;
 	uint32 BasePoolSizeLastFrame = 0;
