[PCG] Fix regression with Blueprints and new FManagedResourceContainer (serialization would clear owner)

#rnx

diff --git a/Engine/Plugins/PCG/Source/PCG/Private/PCGComponent.cpp b/Engine/Plugins/PCG/Source/PCG/Private/PCGComponent.cpp
--- a/Engine/Plugins/PCG/Source/PCG/Private/PCGComponent.cpp
+++ b/Engine/Plugins/PCG/Source/PCG/Private/PCGComponent.cpp
@@ -1828,6 +1828,9 @@ void UPCGComponent::Serialize(FArchive& Ar)
 		GeneratedGraphOutput = GeneratedGraphOutputCopy;
 	}
 #endif // WITH_EDITOR
+
+	// Fixup Owner after serialization
+	ManagedResourceContainer.SetOwner(this);
 }
 
 void UPCGComponent::PostLoad()
@@ -2333,6 +2336,9 @@ void UPCGComponent::PostEditImport()
 	LLM_SCOPE_BYTAG(PCG);
 	Super::PostEditImport();
 
+	// Update owner post import which doesn't use object serialization
+	ManagedResourceContainer.SetOwner(this);
+
 	SetupCallbacksOnCreation();
 }
 
