[PCG] Make sure the bounds in the spline sampler interior are including all the points on the spline

* Local bounds on the spline only take into account the control points.
* But with the tangents, points on the spline can be outside of those bounds

#rnx

diff --git a/Engine/Plugins/PCG/Source/PCG/Private/Elements/PCGSplineSampler.cpp b/Engine/Plugins/PCG/Source/PCG/Private/Elements/PCGSplineSampler.cpp
--- a/Engine/Plugins/PCG/Source/PCG/Private/Elements/PCGSplineSampler.cpp
+++ b/Engine/Plugins/PCG/Source/PCG/Private/Elements/PCGSplineSampler.cpp
@@ -1076,8 +1076,8 @@ namespace PCGSplineSamplerHelpers
 		}
 
 		const FBoxSphereBounds SplineLocalBounds = Spline->LocalBounds;
-		const FVector MinPoint = SplineLocalBounds.Origin - SplineLocalBounds.BoxExtent;
-		const FVector MaxPoint = SplineLocalBounds.Origin + SplineLocalBounds.BoxExtent;
+		FVector MinPoint = SplineLocalBounds.Origin - SplineLocalBounds.BoxExtent;
+		FVector MaxPoint = SplineLocalBounds.Origin + SplineLocalBounds.BoxExtent;
 
 		const FVector::FReal MaxDimension = FMath::Max(Spline->Bounds.BoxExtent.X, Spline->Bounds.BoxExtent.Y) * 2.f;
 		const FVector::FReal MaxDimensionSquared = MaxDimension * MaxDimension;
@@ -1125,6 +1125,9 @@ namespace PCGSplineSamplerHelpers
 		for (const FVector& Point : SplineSamplePoints)
 		{
 			SplineSamplePoints2D.Add(FVector2D(Point));
+
+			MinPoint = FVector::Min(MinPoint, Point);
+			MaxPoint = FVector::Max(MaxPoint, Point);
 		}
 
 		TArray<TTuple<FVector2D, FVector2D>> MedialAxisEdges;
