[PCG] More work to move from UPCGComponent to IPCGGraphExecutionSource
- FPCGActorComponentMapping renamed to FPCGTrackingManager
- Move most of the API to IPCGGraphExecutionSource except for partitioning
- IPCGGraphExecutionState already supports partitioning in its interface but actual implementation is still very much APartitionActor dependant (needs to be abstracted away)
- Some of the partitioning API was moved to IPCGGraphExecutionSource but implementers should not expect functional tracking for anything other than PCGComponents

diff --git a/Engine/Plugins/PCG/Source/PCG/Private/Elements/PCGActorSelector.cpp b/Engine/Plugins/PCG/Source/PCG/Private/Elements/PCGActorSelector.cpp
--- a/Engine/Plugins/PCG/Source/PCG/Private/Elements/PCGActorSelector.cpp
+++ b/Engine/Plugins/PCG/Source/PCG/Private/Elements/PCGActorSelector.cpp
@@ -8,6 +8,7 @@
 #include "Helpers/PCGActorHelpers.h"
 
 #include "Algo/AnyOf.h"
+#include "Algo/Transform.h"
 #include "Components/ActorComponent.h"
 #include "GameFramework/Actor.h"
 
@@ -368,18 +369,6 @@ void FPCGSelectionKey::SetExtraDependency(const UClass* InExtraDependency)
 	OptionalExtraDependency = InExtraDependency;
 }
 
-// DEPRECATED 5.6
-bool FPCGSelectionKey::IsMatching(const UObject* InObject, const UPCGComponent* InComponent) const
-{
-	// This signature of IsMatching does not support custom keys
-	if (CustomKey.IsValid())
-	{
-		return false;
-	}
-
-	return IsMatching(FPCGSelectionKey::CreateFromObjectPtr(InObject), InComponent);
-}
-
 // DEPRECATED 5.7
 bool FPCGSelectionKey::IsMatching(const UObject* InObject, const FSoftObjectPath& InObjectPath, const UPCGComponent* InComponent) const
 {
@@ -394,7 +383,7 @@ bool FPCGSelectionKey::IsMatching(const UObject* InObject, const FSoftObjectPath
 	return IsMatching(SelectionKey, InComponent);
 }
 
-bool FPCGSelectionKey::IsMatching(const FPCGSelectionKey& InSelectionKey, const UPCGComponent* InComponent) const
+bool FPCGSelectionKey::IsMatching(const FPCGSelectionKey& InSelectionKey, const IPCGGraphExecutionSource* InExecutionSource) const
 {
 	const UObject* SelectionObject = InSelectionKey.GetObjectFromPath();
 	if (!SelectionObject)
@@ -413,12 +402,15 @@ bool FPCGSelectionKey::IsMatching(const FPCGSelectionKey& InSelectionKey, const
 			return false;
 		}
 
+		// @todo_pcg: support execution source
+		const UPCGComponent* PCGComponent = Cast<UPCGComponent>(InExecutionSource);
+
 		// InKey provide the info for selecting a given actor.
 		// We reconstruct the selector settings from this key, and we also force it to SelectMultiple, since
 		// we want to gather all the actors that matches this given key, to find if ours matches.
 		FPCGActorSelectorSettings SelectorSettings = FPCGActorSelectorSettings::ReconstructFromKey(*this);
 		SelectorSettings.bSelectMultiple = true;
-		TArray<AActor*> AllActors = PCGActorSelector::FindActors(SelectorSettings, InComponent, [](const AActor*) { return true; }, [](const AActor*) { return true; });
+		TArray<AActor*> AllActors = PCGActorSelector::FindActors(SelectorSettings, PCGComponent, [](const AActor*) { return true; }, [](const AActor*) { return true; });
 		return AllActors.Contains(InActor);
 	}
 
@@ -438,19 +430,7 @@ bool FPCGSelectionKey::IsMatching(const FPCGSelectionKey& InSelectionKey, const
 	}
 }
 
-// DEPRECATED 5.6
-bool FPCGSelectionKey::IsMatching(const UObject* InObject, const TSet<FName>& InRemovedTags, const TSet<UPCGComponent*>& InComponents, TSet<UPCGComponent*>* OptionalMatchedComponents) const
-{
-	// This signature of IsMatching does not support custom keys
-	if (CustomKey.IsValid())
-	{
-		return false;
-	}
-
-	return IsMatching(FPCGSelectionKey::CreateFromObjectPtr(InObject), InRemovedTags, InComponents, OptionalMatchedComponents);
-}
-
-// DEPRECATED 5.7
+// deprecated
 bool FPCGSelectionKey::IsMatching(const UObject* InObject, const FSoftObjectPath& InObjectPath, const TSet<FName>& InRemovedTags, const TSet<UPCGComponent*>& InComponents, TSet<UPCGComponent*>* OptionalMatchedComponents) const
 {
 	// This signature of IsMatching does not support custom keys
@@ -462,10 +442,32 @@ bool FPCGSelectionKey::IsMatching(const UObject* InObject, const FSoftObjectPath
 	FPCGSelectionKey SelectionKey = FPCGSelectionKey::CreateFromPath(InObjectPath);
 	SelectionKey.CachedObject = InObject;
 
+PRAGMA_DISABLE_DEPRECATION_WARNINGS
 	return IsMatching(SelectionKey, InRemovedTags, InComponents, OptionalMatchedComponents);
+PRAGMA_ENABLE_DEPRECATION_WARNINGS
 }
 
+// deprecated
 bool FPCGSelectionKey::IsMatching(const FPCGSelectionKey& InSelectionKey, const TSet<FName>& InRemovedTags, const TSet<UPCGComponent*>& InComponents, TSet<UPCGComponent*>* OptionalMatchedComponents) const
+{
+	TSet<IPCGGraphExecutionSource*> OptionalMatchedExecutionSources;
+	TSet<IPCGGraphExecutionSource*> ExecutionSources;
+	ExecutionSources.Reserve(InComponents.Num());
+	Algo::Transform(InComponents, ExecutionSources, [](UPCGComponent* InComponent) { return InComponent; });
+
+	bool bReturn = IsMatching(InSelectionKey, InRemovedTags, ExecutionSources, OptionalMatchedComponents ? &OptionalMatchedExecutionSources : nullptr);
+	if (OptionalMatchedComponents)
+	{
+		OptionalMatchedComponents->Reserve(OptionalMatchedExecutionSources.Num());
+		Algo::TransformIf(OptionalMatchedExecutionSources, *OptionalMatchedComponents,
+			[](IPCGGraphExecutionSource* InExecutionSource) { return Cast<UPCGComponent>(InExecutionSource) != nullptr; },
+			[](IPCGGraphExecutionSource* InExecutionSource) { return Cast<UPCGComponent>(InExecutionSource); });
+	}
+
+	return bReturn;
+}
+
+bool FPCGSelectionKey::IsMatching(const FPCGSelectionKey& InSelectionKey, const TSet<FName>& InRemovedTags, const TSet<IPCGGraphExecutionSource*>& InExecutionSources, TSet<IPCGGraphExecutionSource*>* OptionalMatchedExecutionSources) const
 {
 	const UObject* SelectionObject = InSelectionKey.GetObjectFromPath();
 
@@ -522,14 +524,14 @@ bool FPCGSelectionKey::IsMatching(const FPCGSelectionKey& InSelectionKey, const
 
 		for (UActorComponent* Component : ActorComponents)
 		{
-			if (UPCGComponent* PCGComponent = Cast<UPCGComponent>(Component))
+			if (IPCGGraphExecutionSource* ExecutionSource = Cast<IPCGGraphExecutionSource>(Component))
 			{
-				if (InComponents.Contains(PCGComponent))
+				if (InExecutionSources.Contains(ExecutionSource))
 				{
 					bFoundMatch = true;
-					if (OptionalMatchedComponents)
+					if (OptionalMatchedExecutionSources)
 					{
-						OptionalMatchedComponents->Add(PCGComponent);
+						OptionalMatchedExecutionSources->Add(ExecutionSource);
 					}
 					else
 					{
@@ -563,9 +565,9 @@ bool FPCGSelectionKey::IsMatching(const FPCGSelectionKey& InSelectionKey, const
 		break;
 	}
 
-	if (bIsMatched && OptionalMatchedComponents)
+	if (bIsMatched && OptionalMatchedExecutionSources)
 	{
-		OptionalMatchedComponents->Append(InComponents);
+		OptionalMatchedExecutionSources->Append(InExecutionSources);
 	}
 
 	return bIsMatched;
diff --git a/Engine/Plugins/PCG/Source/PCG/Private/Elements/PCGDataFromActor.cpp b/Engine/Plugins/PCG/Source/PCG/Private/Elements/PCGDataFromActor.cpp
--- a/Engine/Plugins/PCG/Source/PCG/Private/Elements/PCGDataFromActor.cpp
+++ b/Engine/Plugins/PCG/Source/PCG/Private/Elements/PCGDataFromActor.cpp
@@ -2,11 +2,11 @@
 
 #include "Elements/PCGDataFromActor.h"
 
-#include "PCGActorAndComponentMapping.h"
 #include "PCGComponent.h"
 #include "PCGCustomVersion.h"
 #include "PCGParamData.h"
 #include "PCGSubsystem.h"
+#include "PCGTrackingManager.h"
 #include "Data/PCGBasePointData.h"
 #include "Data/PCGSpatialData.h"
 #include "Data/PCGVolumeData.h"
@@ -426,26 +426,32 @@ bool FPCGDataFromActorElement::ExecuteInternal(FPCGContext* InContext) const
 				if (Settings->ActorSelector.bMustOverlapSelf)
 				{
 					FBox ActorBounds = PCGHelpers::GetGridBounds(Self, PCGComponent);
-					for (UPCGComponent* Component : Subsystem->GetAllIntersectingComponents(ActorBounds))
+					for (IPCGGraphExecutionSource* ExecutionSource : Subsystem->GetAllIntersectingExecutionSources(ActorBounds))
 					{
-						if (AActor* Actor = Component->GetOwner())
+						if (UPCGComponent* Component = Cast<UPCGComponent>(ExecutionSource))
 						{
-							if (Key.IsMatching(FPCGSelectionKey::CreateFromObjectPtr(Actor), Component))
+							if (AActor* Actor = Component->GetOwner())
 							{
-								Context->FoundActors.AddUnique(Actor);
+								if (Key.IsMatching(FPCGSelectionKey::CreateFromObjectPtr(Actor), Component))
+								{
+									Context->FoundActors.AddUnique(Actor);
+								}
 							}
 						}
 					}
 				}
 				else
 				{
-					for (UPCGComponent* Component : Subsystem->GetAllRegisteredComponents())
+					for (IPCGGraphExecutionSource* ExecutionSource : Subsystem->GetAllRegisteredExecutionSources())
 					{
-						if (AActor* Actor = Component->GetOwner())
+						if (UPCGComponent* Component = Cast<UPCGComponent>(ExecutionSource))
 						{
-							if (Key.IsMatching(FPCGSelectionKey::CreateFromObjectPtr(Actor), Component))
+							if (AActor* Actor = Component->GetOwner())
 							{
-								Context->FoundActors.AddUnique(Actor);
+								if (Key.IsMatching(FPCGSelectionKey::CreateFromObjectPtr(Actor), Component))
+								{
+									Context->FoundActors.AddUnique(Actor);
+								}
 							}
 						}
 					}
@@ -666,7 +672,7 @@ void FPCGDataFromActorElement::GatherWaitTasks(AActor* FoundActor, FPCGContext*
 #if WITH_EDITOR
 			// Signal that any change notifications from generating upstream component should not trigger re-executions of this component.
 			// Such change notifications can cancel the current execution.
-			// Note: Uses owner because FPCGActorAndComponentMapping::OnPCGGraphGeneratedOrCleaned reports change on owner.
+			// Note: Uses owner because FPCGTrackingManager::OnPCGGraphGeneratedOrCleaned reports change on owner.
 			SourceOriginalComponent->StartIgnoringChangeOriginDuringGeneration(Component->GetOwner());
 			Context->IgnoredChangeOrigins.Add(Component->GetOwner());
 #endif
diff --git a/Engine/Plugins/PCG/Source/PCG/Private/Elements/PCGDataFromTool.cpp b/Engine/Plugins/PCG/Source/PCG/Private/Elements/PCGDataFromTool.cpp
--- a/Engine/Plugins/PCG/Source/PCG/Private/Elements/PCGDataFromTool.cpp
+++ b/Engine/Plugins/PCG/Source/PCG/Private/Elements/PCGDataFromTool.cpp
@@ -1,10 +1,10 @@
 // Copyright Epic Games, Inc. All Rights Reserved.
 
 #include "Elements/PCGDataFromTool.h"
-#include "PCGActorAndComponentMapping.h"
 #include "PCGComponent.h"
 #include "PCGCustomVersion.h"
 #include "PCGSubsystem.h"
+#include "PCGTrackingManager.h"
 #include "Data/PCGBasePointData.h"
 #include "Data/Tool/PCGToolBaseData.h"
 #include "Elements/PCGAddComponent.h"
diff --git a/Engine/Plugins/PCG/Source/PCG/Private/Elements/PCGSpawnActor.cpp b/Engine/Plugins/PCG/Source/PCG/Private/Elements/PCGSpawnActor.cpp
--- a/Engine/Plugins/PCG/Source/PCG/Private/Elements/PCGSpawnActor.cpp
+++ b/Engine/Plugins/PCG/Source/PCG/Private/Elements/PCGSpawnActor.cpp
@@ -1080,7 +1080,7 @@ void FPCGSpawnActorElement::SpawnActors(FPCGSubgraphContext* Context, AActor* Ta
 				{
 					if (Subsystem)
 					{
-						Subsystem->RegisterOrUpdatePCGComponent(PCGComponent);
+						Subsystem->RegisterOrUpdateExecutionSource(PCGComponent);
 					}
 
 					// TODO: use ScheduleGraph if we want to pass user parameters
diff --git a/Engine/Plugins/PCG/Source/PCG/Private/Grid/PCGComponentOctree.cpp b/Engine/Plugins/PCG/Source/PCG/Private/Grid/PCGComponentOctree.cpp
deleted file mode 100644
diff --git a/Engine/Plugins/PCG/Source/PCG/Private/Grid/PCGExecutionSourceOctree.cpp b/Engine/Plugins/PCG/Source/PCG/Private/Grid/PCGExecutionSourceOctree.cpp
new file mode 100644
--- /dev/null
+++ b/Engine/Plugins/PCG/Source/PCG/Private/Grid/PCGExecutionSourceOctree.cpp
@@ -0,0 +1,182 @@
+// Copyright Epic Games, Inc. All Rights Reserved.
+
+#include "Grid/PCGExecutionSourceOctree.h"
+#include "PCGGraphExecutionStateInterface.h"
+
+FPCGExecutionSourceRef::FPCGExecutionSourceRef(IPCGGraphExecutionSource* InExecutionSource, const FPCGExecutionSourceOctreeIDSharedRef& InIdShared)
+	: IdShared(InIdShared)
+{
+	check(InExecutionSource);
+
+	ExecutionSource = InExecutionSource;
+
+	UpdateBounds();
+}
+
+void FPCGExecutionSourceRef::UpdateBounds()
+{
+	check(ExecutionSource);
+
+	Bounds = ExecutionSource->GetExecutionState().GetBounds();
+}
+
+FPCGExecutionSourceOctreeAndMap::FPCGExecutionSourceOctreeAndMap(const FVector& InOrigin, FVector::FReal InExtent)
+	: Octree(InOrigin, InExtent)
+{
+}
+
+void FPCGExecutionSourceOctreeAndMap::Reset(const FVector& InOrigin, FVector::FReal InExtent)
+{
+	UE::TWriteScopeLock WriteLock(Lock);
+	Octree = FPCGExecutionSourceOctree(InOrigin, InExtent);
+	ExecutionSourceToIdMap.Empty();
+}
+
+bool FPCGExecutionSourceOctreeAndMap::Contains(const IPCGGraphExecutionSource* InExecutionSource) const
+{
+	UE::TReadScopeLock ReadLock(Lock);
+
+	return ExecutionSourceToIdMap.Find(InExecutionSource) != nullptr;
+}
+
+FBox FPCGExecutionSourceOctreeAndMap::GetBounds(const IPCGGraphExecutionSource* InExecutionSource) const
+{
+	FBox Bounds(EForceInit::ForceInit);
+
+	UE::TReadScopeLock ReadLock(Lock);
+	if (const FPCGExecutionSourceOctreeIDSharedRef* ElementIdPtr = ExecutionSourceToIdMap.Find(InExecutionSource))
+	{
+		const FPCGExecutionSourceRef& ExecutionSourceRef = Octree.GetElementById((*ElementIdPtr)->Id);
+		Bounds = ExecutionSourceRef.Bounds.GetBox();
+	}
+
+	return Bounds;
+}
+
+void FPCGExecutionSourceOctreeAndMap::AddOrUpdateExecutionSource(IPCGGraphExecutionSource* InExecutionSource, FBox& OutBounds, bool& bOutExecutionSourceHasChanged, bool& bOutExecutionSourceWasAdded)
+{
+	bOutExecutionSourceHasChanged = false;
+	bOutExecutionSourceWasAdded = false;
+
+	UE::TWriteScopeLock WriteLock(Lock);
+
+	FPCGExecutionSourceOctreeIDSharedRef* ElementIdPtr = ExecutionSourceToIdMap.Find(InExecutionSource);
+
+	if (!ElementIdPtr)
+	{
+		// Does not exist yet, add it.
+		FPCGExecutionSourceOctreeIDSharedRef IdShared = MakeShared<FPCGExecutionSourceOctreeID>();
+		FPCGExecutionSourceRef ExecutionSourceRef(InExecutionSource, IdShared);
+		OutBounds = ExecutionSourceRef.Bounds.GetBox();
+		check(OutBounds.IsValid);
+
+		// If the ExecutionSource is already generated, it probably mean we are in loading. The ExecutionSource bounds and last
+		// generated bounds should be the same.
+		// If the bounds depends on other ExecutionSources on the owner however, it might not be the same, because of the registration order.
+		// In this case, override the bounds by the last generated ones.
+		FBox LastGeneratedBounds = InExecutionSource->GetExecutionState().GetGeneratedBounds();
+		if (InExecutionSource->GetExecutionState().IsGenerated() && !LastGeneratedBounds.Equals(OutBounds))
+		{
+			OutBounds = LastGeneratedBounds;
+			ExecutionSourceRef.Bounds = OutBounds;
+		}
+
+		Octree.AddElement(ExecutionSourceRef);
+
+		bOutExecutionSourceHasChanged = true;
+		bOutExecutionSourceWasAdded = true;
+
+		// Store the shared ptr, because if we add/remove ExecutionSources in the octree, the id might change.
+		// We need to make sure that we always have the latest id for the given ExecutionSource.
+		ExecutionSourceToIdMap.Add(InExecutionSource, MoveTemp(IdShared));
+	}
+	else
+	{
+		// It already exists, update it if the bounds changed.
+
+		// Do a copy here.
+		FPCGExecutionSourceRef ExecutionSourceRef = Octree.GetElementById((*ElementIdPtr)->Id);
+		FBox PreviousBounds = ExecutionSourceRef.Bounds.GetBox();
+		ExecutionSourceRef.UpdateBounds();
+		OutBounds = ExecutionSourceRef.Bounds.GetBox();
+		check(OutBounds.IsValid);
+
+		// If bounds changed, remove and re-add to the octree
+		if (!PreviousBounds.Equals(OutBounds))
+		{
+			Octree.RemoveElement((*ElementIdPtr)->Id);
+			Octree.AddElement(ExecutionSourceRef);
+			bOutExecutionSourceHasChanged = true;
+		}
+	}
+}
+
+bool FPCGExecutionSourceOctreeAndMap::RemapExecutionSource(const IPCGGraphExecutionSource* InOldExecutionSource, IPCGGraphExecutionSource* InNewExecutionSource, bool& bOutBoundsHasChanged)
+{
+	bOutBoundsHasChanged = false;
+
+	// First verification we have the old ExecutionSource registered
+	if (!Contains(InOldExecutionSource))
+	{
+		return false;
+	}
+
+	// If so, lock again in write and recheck if it has not been remapped already.
+	{
+		UE::TWriteScopeLock WriteLock(Lock);
+
+		FPCGExecutionSourceOctreeIDSharedRef* ElementIdPtr = ExecutionSourceToIdMap.Find(InOldExecutionSource);
+
+		if (!ElementIdPtr)
+		{
+			// Nothing done
+			return false;
+		}
+
+		FPCGExecutionSourceOctreeIDSharedRef ElementId = *ElementIdPtr;
+
+		FPCGExecutionSourceRef ExecutionSourceRef = Octree.GetElementById(ElementId->Id);
+		FBox PreviousBounds = ExecutionSourceRef.Bounds.GetBox();
+		ExecutionSourceRef.ExecutionSource = InNewExecutionSource;
+		ExecutionSourceRef.UpdateBounds();
+		FBox Bounds = ExecutionSourceRef.Bounds.GetBox();
+		check(Bounds.IsValid);
+
+		// If bounds changed, we need to update the mapping
+		bOutBoundsHasChanged = !Bounds.Equals(PreviousBounds);
+
+		Octree.RemoveElement((*ElementIdPtr)->Id);
+		Octree.AddElement(ExecutionSourceRef);
+		ExecutionSourceToIdMap.Remove(InOldExecutionSource);
+		ExecutionSourceToIdMap.Add(InNewExecutionSource, MoveTemp(ElementId));
+	}
+
+	return true;
+}
+
+bool FPCGExecutionSourceOctreeAndMap::RemoveExecutionSource(IPCGGraphExecutionSource* InExecutionSource)
+{
+	UE::TWriteScopeLock WriteLock(Lock);
+
+	FPCGExecutionSourceOctreeIDSharedRef* ElementIdPtr = ExecutionSourceToIdMap.Find(InExecutionSource);
+
+	if (!ElementIdPtr)
+	{
+		return false;
+	}
+
+	Octree.RemoveElement((*ElementIdPtr)->Id);
+	ExecutionSourceToIdMap.Remove(InExecutionSource);
+
+	return true;
+}
+
+TSet<IPCGGraphExecutionSource*> FPCGExecutionSourceOctreeAndMap::GetAllExecutionSources() const
+{
+	UE::TReadScopeLock ReadLock(Lock);
+
+	TSet<IPCGGraphExecutionSource*> ExecutionSources;
+	ExecutionSourceToIdMap.GetKeys(ExecutionSources);
+
+	return ExecutionSources;
+}
diff --git a/Engine/Plugins/PCG/Source/PCG/Private/Grid/PCGExecutionSourceOctree.h b/Engine/Plugins/PCG/Source/PCG/Private/Grid/PCGExecutionSourceOctree.h
new file mode 100644
--- /dev/null
+++ b/Engine/Plugins/PCG/Source/PCG/Private/Grid/PCGExecutionSourceOctree.h
@@ -0,0 +1,95 @@
+// Copyright Epic Games, Inc. All Rights Reserved.
+
+#pragma once
+
+#include "Math/BoxSphereBounds.h"
+#include "Math/GenericOctree.h"
+#include "Misc/ScopeRWLock.h"
+#include "Misc/TransactionallySafeRWLock.h"
+#include "UObject/ObjectPtr.h"
+#include "Templates/SharedPointer.h"
+#include "Templates/SharedPointerFwd.h"
+
+template <typename ElementType, typename OctreeSemantics> class TOctree2;
+
+class IPCGGraphExecutionSource;
+
+struct FPCGExecutionSourceOctreeID : public TSharedFromThis<FPCGExecutionSourceOctreeID, ESPMode::ThreadSafe>
+{
+	FOctreeElementId2 Id;
+};
+
+using FPCGExecutionSourceOctreeIDSharedRef = TSharedRef<struct FPCGExecutionSourceOctreeID, ESPMode::ThreadSafe>;
+
+struct FPCGExecutionSourceRef
+{
+	FPCGExecutionSourceRef(IPCGGraphExecutionSource* InExecutionSource, const FPCGExecutionSourceOctreeIDSharedRef& InIdShared);
+
+	void UpdateBounds();
+
+	FPCGExecutionSourceOctreeIDSharedRef IdShared;
+	IPCGGraphExecutionSource* ExecutionSource = nullptr;
+	FBoxSphereBounds Bounds;
+};
+
+struct FPCGExecutionSourceRefSemantics
+{
+	enum { MaxElementsPerLeaf = 16 };
+	enum { MinInclusiveElementsPerNode = 7 };
+	enum { MaxNodeDepth = 12 };
+
+	typedef TInlineAllocator<MaxElementsPerLeaf> ElementAllocator;
+
+	inline static const FBoxSphereBounds& GetBoundingBox(const FPCGExecutionSourceRef& InVolume)
+	{
+		return InVolume.Bounds;
+	}
+
+	inline static const bool AreElementsEqual(const FPCGExecutionSourceRef& A, const FPCGExecutionSourceRef& B)
+	{
+		return A.ExecutionSource == B.ExecutionSource;
+	}
+
+	inline static void ApplyOffset(FPCGExecutionSourceRef& InVolume, const FVector& Offset)
+	{
+		InVolume.Bounds.Origin += Offset;
+	}
+
+	inline static void SetElementId(const FPCGExecutionSourceRef& Element, FOctreeElementId2 OctreeElementID)
+	{
+		Element.IdShared->Id = OctreeElementID;
+	}
+};
+
+using FPCGExecutionSourceOctree = TOctree2<FPCGExecutionSourceRef, FPCGExecutionSourceRefSemantics> ;
+using FPCGExecutionSourceToIdMap = TMap<IPCGGraphExecutionSource*, FPCGExecutionSourceOctreeIDSharedRef>;
+
+class FPCGExecutionSourceOctreeAndMap
+{
+public:
+	FPCGExecutionSourceOctreeAndMap() = default;
+	FPCGExecutionSourceOctreeAndMap(const FVector& InOrigin, FVector::FReal InExtent);
+
+	void Reset(const FVector& InOrigin, FVector::FReal InExtent);
+
+	TSet<IPCGGraphExecutionSource*> GetAllExecutionSources() const;
+
+	template<typename IterateBoundsFunc>
+	inline void FindElementsWithBoundsTest(const FBoxCenterAndExtent& BoxBounds, const IterateBoundsFunc& Func) const
+	{
+		UE::TReadScopeLock ReadLock(Lock);
+		return Octree.FindElementsWithBoundsTest(BoxBounds, Func);
+	}
+
+	bool Contains(const IPCGGraphExecutionSource* InExecutionSource) const;
+	FBox GetBounds(const IPCGGraphExecutionSource* InExecutionSource) const;
+
+	void AddOrUpdateExecutionSource(IPCGGraphExecutionSource* InExecutionSource, FBox& OutBounds, bool& bOutExecutionSourceHasChanged, bool& bOutExecutionSourcetWasAdded);
+	bool RemapExecutionSource(const IPCGGraphExecutionSource* InOldExecutionSource, IPCGGraphExecutionSource* InNewExecutionSource, bool& bOutBoundsHasChanged);
+	bool RemoveExecutionSource(IPCGGraphExecutionSource* InExecutionSource);
+
+private:
+	FPCGExecutionSourceOctree Octree;
+	FPCGExecutionSourceToIdMap ExecutionSourceToIdMap;
+	mutable FTransactionallySafeRWLock Lock;
+};
\ No newline at end of file
diff --git a/Engine/Plugins/PCG/Source/PCG/Private/PCGActorAndComponentMapping.cpp b/Engine/Plugins/PCG/Source/PCG/Private/PCGActorAndComponentMapping.cpp
deleted file mode 100644
diff --git a/Engine/Plugins/PCG/Source/PCG/Private/PCGActorAndComponentMapping.h b/Engine/Plugins/PCG/Source/PCG/Private/PCGActorAndComponentMapping.h
deleted file mode 100644
diff --git a/Engine/Plugins/PCG/Source/PCG/Private/PCGComponent.cpp b/Engine/Plugins/PCG/Source/PCG/Private/PCGComponent.cpp
--- a/Engine/Plugins/PCG/Source/PCG/Private/PCGComponent.cpp
+++ b/Engine/Plugins/PCG/Source/PCG/Private/PCGComponent.cpp
@@ -153,7 +153,7 @@ void UPCGComponent::SetIsPartitioned(bool bIsNowPartitioned)
 
 		// Update the component on the subsystem
 		bIsComponentPartitioned = bIsNowPartitioned;
-		Subsystem->RegisterOrUpdatePCGComponent(this, bDoActorMapping);
+		Subsystem->RegisterOrUpdateExecutionSource(this, bDoActorMapping);
 	}
 	else
 	{
@@ -1586,7 +1586,7 @@ void UPCGComponent::BeginPlay()
 	if (GetSubsystem())
 	{
 		ensure(!bGenerated || LastGeneratedBounds.IsValid);
-		GetSubsystem()->RegisterOrUpdatePCGComponent(this);
+		GetSubsystem()->RegisterOrUpdateExecutionSource(this);
 	}
 
 	// Procedural instances are never persisted so always require generation.
@@ -1611,7 +1611,7 @@ void UPCGComponent::EndPlay(const EEndPlayReason::Type EndPlayReason)
 	if (UPCGSubsystem* Subsystem = GetSubsystem())
 	{
 		Subsystem->CancelGeneration(this);
-		Subsystem->UnregisterPCGComponent(this);
+		Subsystem->UnregisterExecutionSource(this);
 	}
 
 	ManagedResourceContainer.EndPlay();
@@ -1637,7 +1637,7 @@ void UPCGComponent::OnUnregister()
 		// If Original Component is managed by runtime and being unloaded, unregister immediatly so that runtime generated Local components also get cleaned up. (BeginDestroy is too late)
 		if (bUnregisteredThroughLoading && IsManagedByRuntimeGenSystem() && !IsLocalComponent())
 		{
-			Subsystem->UnregisterPCGComponent(this, /*bForce=*/true);
+			Subsystem->UnregisterExecutionSource(this, /*bForce=*/true);
 		}
 	}
 
@@ -1748,7 +1748,7 @@ void UPCGComponent::OnComponentDestroyed(bool bDestroyingHierarchy)
 	{
 		if (!PCGHelpers::IsRuntimeOrPIE())
 		{
-			Subsystem->UnregisterPCGComponent(this);
+			Subsystem->UnregisterExecutionSource(this);
 		}
 
 		if (IsCreatedByConstructionScript())
@@ -1994,7 +1994,7 @@ void UPCGComponent::BeginDestroy()
 	// but gets destroyed immediately, we need to force the unregistering. 
 	if (UPCGSubsystem* PCGSubsystem = UPCGSubsystem::GetSubsystemForCurrentWorld())
 	{
-		PCGSubsystem->UnregisterPCGComponent(this, /*bForce=*/true);
+		PCGSubsystem->UnregisterExecutionSource(this, /*bForce=*/true);
 	}
 #endif // WITH_EDITOR
 
@@ -2047,7 +2047,7 @@ void UPCGComponent::OnRegister()
 		if (UWorld* World = GetWorld())
 		{
 			ensure(!bGenerated || LastGeneratedBounds.IsValid);
-			GetSubsystem()->RegisterOrUpdatePCGComponent(this, bGenerated);
+			GetSubsystem()->RegisterOrUpdateExecutionSource(this, bGenerated);
 		}
 	}
 #endif //WITH_EDITOR
@@ -2666,7 +2666,7 @@ void UPCGComponent::Refresh(EPCGChangeType ChangeType, bool bCancelExistingRefre
 		// Note that we pass in false here to remove everything when moving a partitioned graph because we would otherwise need to do a reversible stamp to support this
 		if (UPCGSubsystem* Subsystem = GetSubsystem())
 		{
-			Subsystem->RegisterOrUpdatePCGComponent(this, bGenerated);
+			Subsystem->RegisterOrUpdateExecutionSource(this, bGenerated);
 		}
 
 		return;
@@ -2824,7 +2824,7 @@ void UPCGComponent::OnRefresh(bool bForceRefresh)
 
 	if (Subsystem)
 	{
-		Subsystem->RegisterOrUpdatePCGComponent(this, /*bDoActorMapping=*/ bWasGeneratedOrGenerating);
+		Subsystem->RegisterOrUpdateExecutionSource(this, /*bDoActorMapping=*/ bWasGeneratedOrGenerating);
 	}
 
 	// Following a change in some properties or in some spatial information related to this component,
@@ -4109,7 +4109,7 @@ void FPCGComponentInstanceData::ApplyToComponent(UActorComponent* Component, con
 		UPCGSubsystem* Subsystem = PCGComponent->GetSubsystem();
 		if (Subsystem && ConstructionSourceComponent)
 		{
-			Subsystem->RemapPCGComponent(ConstructionSourceComponent, PCGComponent, bDoActorMapping);
+			Subsystem->RemapExecutionSource(ConstructionSourceComponent, PCGComponent, bDoActorMapping);
 		}
 
 #if WITH_EDITOR
diff --git a/Engine/Plugins/PCG/Source/PCG/Private/PCGComponentExecutionState.cpp b/Engine/Plugins/PCG/Source/PCG/Private/PCGComponentExecutionState.cpp
--- a/Engine/Plugins/PCG/Source/PCG/Private/PCGComponentExecutionState.cpp
+++ b/Engine/Plugins/PCG/Source/PCG/Private/PCGComponentExecutionState.cpp
@@ -121,6 +121,18 @@ bool FPCGComponentExecutionState::IsGenerating() const
 	return Component->IsGenerating();
 }
 
+bool FPCGComponentExecutionState::IsGenerated() const
+{
+	check(Component);
+	return Component->bGenerated;
+}
+
+FBox FPCGComponentExecutionState::GetGeneratedBounds() const
+{
+	check(Component);
+	return Component->GetLastGeneratedBounds();
+}
+
 void FPCGComponentExecutionState::ExecutePreGraph(FPCGContext* InContext)
 {
 	check(Component);
diff --git a/Engine/Plugins/PCG/Source/PCG/Private/PCGElement.cpp b/Engine/Plugins/PCG/Source/PCG/Private/PCGElement.cpp
--- a/Engine/Plugins/PCG/Source/PCG/Private/PCGElement.cpp
+++ b/Engine/Plugins/PCG/Source/PCG/Private/PCGElement.cpp
@@ -899,14 +899,14 @@ bool IPCGElement::DebugDisplay(FPCGContext* Context) const
 		return true;
 	}
 
-	// If graph is being inspected, only display Debug if the component is being inspected, or in the HiGen case also display if
+	// If graph is being inspected, only display Debug if the execution source is being inspected, or in the HiGen case also display if
 	// this component is a parent of an inspected component (because this data is available to child components).
 
 	// If the graph is not being inspected, or the current component is being inspected, then we know we should display
 	// debug, if not then we do further checks.
-	UPCGComponent* SourceComponent = Cast<UPCGComponent>(Context->ExecutionSource.Get());
-	UPCGGraph* Graph = SourceComponent ? SourceComponent->GetGraph() : nullptr;
-	if (Graph && Graph->IsInspecting() && !SourceComponent->GetExecutionState().GetInspection().IsInspecting() && Graph->DebugFlagAppliesToIndividualComponents())
+	IPCGGraphExecutionSource* ExecutionSource = Context->ExecutionSource.Get();
+	UPCGGraph* Graph = ExecutionSource ? ExecutionSource->GetExecutionState().GetGraph() : nullptr;
+	if (Graph && Graph->IsInspecting() && !ExecutionSource->GetExecutionState().GetInspection().IsInspecting() && Graph->DebugFlagAppliesToIndividualComponents())
 	{
 		// If we're no doing HiGen, or if the current component is not a local component (and therefore will not have children),
 		// then do not display debug.
@@ -917,21 +917,26 @@ bool IPCGElement::DebugDisplay(FPCGContext* Context) const
 
 		// If a child of this component is being inspected (a local component on smaller grid and overlapping) then we still show debug,
 		// because this data is available to that child for use.
-		if (UPCGSubsystem* Subsystem = UPCGSubsystem::GetInstance(SourceComponent->GetExecutionState().GetWorld()))
+		if (UPCGSubsystem* Subsystem = Cast<UPCGSubsystem>(ExecutionSource->GetExecutionState().GetSubsystem()))
 		{
-			const uint32 ThisGenerationGridSize = SourceComponent->GetGenerationGridSize();
+			const uint32 ThisGenerationGridSize = ExecutionSource->GetExecutionState().GetGenerationGridSize();
 
-			bool bFoundInspectedChildComponent = false;
-			Subsystem->ForAllOverlappingComponentsInHierarchy(SourceComponent, [ThisGenerationGridSize, &bFoundInspectedChildComponent](UPCGComponent* InLocalComponent)
+			bool bFoundInspectedChild = false;
+
+			// @todo_pcg: support execution source
+			if(UPCGComponent* OriginalComponent = Cast<UPCGComponent>(ExecutionSource))
 			{
-				if (InLocalComponent->GetGenerationGridSize() < ThisGenerationGridSize && InLocalComponent->GetExecutionState().GetInspection().IsInspecting())
+				Subsystem->ForAllOverlappingComponentsInHierarchy(OriginalComponent, [ThisGenerationGridSize, &bFoundInspectedChild](UPCGComponent* InLocalComponent)
 				{
-					bFoundInspectedChildComponent = true;
-				}
-			});
+					if (InLocalComponent->GetExecutionState().GetGenerationGridSize() < ThisGenerationGridSize && InLocalComponent->GetExecutionState().GetInspection().IsInspecting())
+					{
+						bFoundInspectedChild = true;
+					}
+				});
+			}
 
 			// If no inspected child component then don't display debug.
-			if (!bFoundInspectedChildComponent)
+			if (!bFoundInspectedChild)
 			{
 				return true;
 			}
diff --git a/Engine/Plugins/PCG/Source/PCG/Private/PCGTrackingManager.cpp b/Engine/Plugins/PCG/Source/PCG/Private/PCGTrackingManager.cpp
new file mode 100644
--- /dev/null
+++ b/Engine/Plugins/PCG/Source/PCG/Private/PCGTrackingManager.cpp
@@ -0,0 +1,2462 @@
+// Copyright Epic Games, Inc. All Rights Reserved.
+
+#include "PCGTrackingManager.h"
+
+#include "PCGComponent.h"
+#include "PCGDataAsset.h"
+#include "PCGGraph.h"
+#include "PCGManagedResource.h"
+#include "PCGModule.h"
+#include "PCGSubsystem.h"
+#include "PCGWorldActor.h"
+#include "Grid/PCGPartitionActor.h"
+#include "Helpers/PCGActorHelpers.h"
+#include "Helpers/PCGHelpers.h"
+#include "Utils/PCGGraphExecutionLogging.h"
+
+#include "Landscape.h"
+#include "LandscapeSubsystem.h"
+#include "StaticMeshCompiler.h"
+#include "TextureCompiler.h"
+#include "Algo/AnyOf.h"
+#include "Algo/ForEach.h"
+#include "Engine/DataTable.h"
+#include "Engine/Engine.h"
+#include "Engine/Level.h"
+#include "Engine/StaticMesh.h"
+#include "Kismet/GameplayStatics.h"
+#include "PackedLevelActor/PackedLevelActor.h"
+#include "LevelInstance/LevelInstanceEditorInstanceActor.h"
+#include "LevelInstance/LevelInstanceInterface.h"
+#include "LevelInstance/LevelInstanceEditorPivotInterface.h"
+#include "LevelInstance/LevelInstanceSubsystem.h"
+#include "Materials/MaterialInterface.h"
+
+#if WITH_EDITOR
+#include "Grid/PCGPartitionActorDesc.h"
+#include "UObject/PackageReload.h"
+#include "WorldPartition/WorldPartition.h"
+#include "WorldPartition/WorldPartitionActorDescInstance.h"
+#include "WorldPartition/WorldPartitionHelpers.h"
+#endif
+
+namespace PCGTrackingManager
+{
+#if WITH_EDITOR
+	static TAutoConsoleVariable<bool> CVarDisableObjectDependenciesTracking(
+		TEXT("pcg.DisableObjectDependenciesTracking"),
+		false,
+		TEXT("If depencencies are being unstable, disable the tracking, allowing people to continue working while we investigate."));
+
+	static TAutoConsoleVariable<bool> CVarDisableDelayedActorRegistering(
+		TEXT("pcg.DisableDelayedActorRegistering"),
+		false,
+		TEXT("If delayed actor registering when their components aren't registered yet is introducing bad behavior, disables it, allowing people to continue working while we investigate."));
+
+	static TAutoConsoleVariable<bool> CVarLandscapeDisableRefreshTracking(
+		TEXT("pcg.LandscapeDisableRefreshTracking"),
+		false,
+		TEXT("Completely disable landscape refresh when it changes."));
+
+	static TAutoConsoleVariable<bool> CVarLandscapeDisableRefreshTrackingInLandscapeEditingMode(
+		TEXT("pcg.LandscapeDisableRefreshTrackingInLandscapeEditingMode"),
+		false,
+		TEXT("Disable landscape refresh when it changes in landscape editing mode."));
+
+	static TAutoConsoleVariable<int> CVarLandscapeRefreshTimeDelay(
+		TEXT("pcg.LandscapeRefreshTimeDelayMS"),
+		1000,
+		TEXT("Time in MS between a landscape change and PCG refresh. Set it to 0 or negative value to disable the delay."));
+
+	static TAutoConsoleVariable<int> CVarActorModifiedPreviousDataCleanupDelayMS(
+		TEXT("pcg.ActorModifiedPreviousDataCleanupDelayMS"),
+		10000,
+		TEXT("Time in MS between a cached previous data from modified actor and the time we remove it from our cache."));
+
+	static TAutoConsoleVariable<int> CVarActorModifiedPreviousDataCheckDelayMS(
+		TEXT("pcg.ActorModifiedPreviousDataCheckDelayMS"),
+		1000,
+		TEXT("Delay in MS between cleanup checks on previous data from modified actors."));
+
+	static TAutoConsoleVariable<bool> CVarDisablePCGDataInterdependencyOptimization(
+		TEXT("pcg.DisablePCGDataInterdependencyOptimization"),
+		false,
+		TEXT("Disable the optimization that keep track of components depending on others, as a safety measure."));
+#endif // WITH_EDITOR
+
+	bool IsValidOriginalExecutionSource(const IPCGGraphExecutionSource* InExecutionSource)
+	{
+		const UObject* ExecutionSourceObject = Cast<UObject>(InExecutionSource);
+
+		// Ignore invalid, templates and local sources
+		if (!IsValid(ExecutionSourceObject) || ExecutionSourceObject->IsTemplate() || InExecutionSource->GetExecutionState().IsLocalSource())
+		{
+			return false;
+		}
+		// Keep existing PCG Component conditions for now (should probably be moved inside execution source api)
+		else if (const UPCGComponent* PCGComponent = Cast<UPCGComponent>(InExecutionSource))
+		{
+			return PCGComponent->GetOwner() && !PCGComponent->GetOwner()->IsA<APCGPartitionActor>();
+		}
+		else
+		{
+			return true;
+		}
+	}
+
+	bool IsValidExecutionSource(const IPCGGraphExecutionSource* InExecutionSource)
+	{
+		const UObject* ExecutionSourceObject = Cast<UObject>(InExecutionSource);
+	
+		// Ignore invalid and templates
+		if (!IsValid(ExecutionSourceObject) || ExecutionSourceObject->IsTemplate())
+		{
+			return false;
+		}
+		// Keep existing PCG Component conditions for now (should probably be moved inside execution source api)
+		else if (const UPCGComponent* PCGComponent = Cast<UPCGComponent>(InExecutionSource))
+		{
+			return PCGComponent->GetOwner() != nullptr;
+		}
+		else
+		{
+			return true;
+		}
+	}
+
+#if WITH_EDITOR
+	bool ShouldIgnoreActor(AActor* InActor)
+	{
+		if(!InActor)
+		{
+			return true;
+		}
+				
+		return InActor->bIsEditorPreviewActor || InActor->IsA<ALevelInstanceEditorInstanceActor>() || InActor->IsA<APCGWorldActor>() || InActor->Implements<ULevelInstanceEditorPivotInterface>();
+	}
+
+	void PropagateToLevelInstanceActors(ILevelInstanceInterface* InLevelInstance, TFunctionRef<bool(AActor* LevelActor)> InFunc)
+	{
+		if (InLevelInstance)
+		{
+			if (const ULevelInstanceSubsystem* LevelInstanceSubsystem = InLevelInstance->GetLevelInstanceSubsystem())
+			{
+				LevelInstanceSubsystem->ForEachActorInLevelInstance(InLevelInstance, [InFunc](AActor* LevelActor)
+				{
+					if (!ShouldIgnoreActor(LevelActor))
+					{
+						return InFunc(LevelActor);
+					}
+					return true;
+				});
+			}
+		}
+	}
+	
+	void PropagateToLevelInstanceActors(AActor* InActor, TFunctionRef<bool(AActor* LevelActor)> InFunc)
+	{
+		PropagateToLevelInstanceActors(Cast<ILevelInstanceInterface>(InActor), InFunc);
+	}
+		
+	bool ShouldDiscardLandscapeRefresh(const ALandscapeProxy* InLandscape, bool& bIsInEditingMode, bool bIsExitingEditingMode)
+	{
+		bIsInEditingMode = false;
+		// If it is not a landscape, we should refresh.
+		if (!InLandscape)
+		{
+			return false;
+		}
+
+		// If refresh is globably disabled, never refresh
+		if (PCGTrackingManager::CVarLandscapeDisableRefreshTracking.GetValueOnAnyThread())
+		{
+			return true;
+		}
+
+		// If refresh is not disabled in editing, always refresh
+		if (!PCGTrackingManager::CVarLandscapeDisableRefreshTrackingInLandscapeEditingMode.GetValueOnAnyThread())
+		{
+			return false;
+		}
+
+		// Refresh only if we are not editing.
+		const ALandscape* Landscape = InLandscape->GetLandscapeActor();
+		bIsInEditingMode = Landscape && Landscape->HasLandscapeEdMode() && !bIsExitingEditingMode;
+		return bIsInEditingMode;
+	}
+#endif
+}
+
+FPCGTrackingManager::FPCGTrackingManager(UPCGSubsystem* InPCGSubsystem)
+	: PCGSubsystem(InPCGSubsystem)
+{
+	check(PCGSubsystem);
+
+	// TODO: For now we set our octree to be 2km wide, but it would be perhaps better to
+	// scale it to the size of our world.
+	constexpr FVector::FReal OctreeExtent = 200000; // 2km
+	PartitionedOctree.Reset(FVector::ZeroVector, OctreeExtent);
+	NonPartitionedOctree.Reset(FVector::ZeroVector, OctreeExtent);
+}
+
+void FPCGTrackingManager::Initialize()
+{
+#if WITH_EDITOR
+	BuildPartitionActorRecords();
+	RegisterDelegates();
+#endif // WITH_EDITOR
+}
+
+void FPCGTrackingManager::Deinitialize()
+{
+#if WITH_EDITOR
+	UnregisterDelegates();
+#endif
+}
+
+void FPCGTrackingManager::Tick()
+{
+#if WITH_EDITOR
+
+	// Process delayed Registrations
+	TSet<UPCGComponent*> ComponentsToUnregister;
+	TSet<UPCGComponent*> ComponentsToRegister;
+	{
+		PCG::TScopeLock Lock(DelayedComponentsRegistrationLock);
+		ComponentsToUnregister = MoveTemp(DelayedComponentsToUnregister);
+		ComponentsToRegister = MoveTemp(DelayedComponentsToRegister);
+	}
+
+	Algo::ForEach(ComponentsToRegister, [this](UPCGComponent* Component) { RegisterOrUpdateExecutionSource(Component, Component->bGenerated, /*bForce=*/true); });
+	Algo::ForEach(ComponentsToUnregister, [this](UPCGComponent* Component) { UnregisterExecutionSource(Component, /*bForce=*/true); });
+
+	ProcessDelayedEvents();
+
+	const double CurrentTime = FApp::GetCurrentTime();
+
+	if (!DelayedModifiedLandscapes.IsEmpty() && LastLandscapeDirtyTime > 0.0 && ((CurrentTime - LastLandscapeDirtyTime) * 1000.0) > PCGTrackingManager::CVarLandscapeRefreshTimeDelay.GetValueOnAnyThread())
+	{
+		LastLandscapeDirtyTime = -1.0;
+		for (TObjectKey<ALandscapeProxy> Landscape : DelayedModifiedLandscapes)
+		{
+			ApplyLandscapeChanges(Landscape.ResolveObjectPtr());
+		}
+
+		DelayedModifiedLandscapes.Empty();
+	}
+
+	// Cleaning up previous data gathered by the OnObjectModified function but not consumed.
+	if (LastPreviousActorDataCleanup < 0.0 || ((CurrentTime - LastPreviousActorDataCleanup) * 1000) > PCGTrackingManager::CVarActorModifiedPreviousDataCheckDelayMS.GetValueOnAnyThread())
+	{
+		TRACE_CPUPROFILER_EVENT_SCOPE(FPCGTrackingManager::Tick::PreviousActorDataCleanupCheck);
+
+		LastPreviousActorDataCleanup = CurrentTime;
+		TArray<TObjectKey<AActor>> Keys;
+		ActorToPreviousDataMap.GetKeys(Keys);
+		for (const TObjectKey<AActor>& Key : Keys)
+		{
+			if (((CurrentTime - ActorToPreviousDataMap[Key].Get<2>()) * 1000) > PCGTrackingManager::CVarActorModifiedPreviousDataCleanupDelayMS.GetValueOnAnyThread())
+			{
+				ActorToPreviousDataMap.Remove(Key);
+			}
+		}
+	}
+
+	if (PCGTrackingManager::CVarDisablePCGDataInterdependencyOptimization.GetValueOnAnyThread() && !ExecutionSourcesToDependencyMap.IsEmpty())
+	{
+		ExecutionSourcesToDependencyMap.Empty();
+	}
+#endif // WITH_EDITOR
+}
+
+UWorld* FPCGTrackingManager::GetWorld() const
+{
+	return PCGSubsystem->GetWorld();
+}
+
+TArray<FPCGTaskId> FPCGTrackingManager::DispatchToRegisteredLocalComponents(UPCGComponent* OriginalComponent, const TFunctionRef<FPCGTaskId(UPCGComponent*, const TArray<FPCGTaskId>&)>& InFunc) const
+{
+	TRACE_CPUPROFILER_EVENT_SCOPE(FPCGTrackingManager::DispatchToRegisteredLocalComponents);
+	if (!ensure(OriginalComponent))
+	{
+		return {};
+	}
+
+	// TODO: Might be more interesting to copy the set and release the lock.
+	PCG::TSharedScopeLock ReadLock(ComponentToPartitionActorsMapLock);
+	const TSet<TObjectPtr<APCGPartitionActor>>* PartitionActorsPtr = ComponentToPartitionActorsMap.Find(OriginalComponent);
+
+	if (!PartitionActorsPtr)
+	{
+		return TArray<FPCGTaskId>();
+	}
+
+	return DispatchToLocalComponents(OriginalComponent, *PartitionActorsPtr, InFunc);
+}
+
+TArray<FPCGTaskId> FPCGTrackingManager::DispatchToLocalComponents(UPCGComponent* OriginalComponent, const TSet<TObjectPtr<APCGPartitionActor>>& PartitionActors, const TFunctionRef<FPCGTaskId(UPCGComponent*, const TArray<FPCGTaskId>&)>& InFunc) const
+{
+	TArray<FPCGTaskId> TaskIds;
+
+	auto DispatchLambda = [this, OriginalComponent, &InFunc, &TaskIds](const TObjectPtr<APCGPartitionActor>& PartitionActor)
+	{
+		if (PartitionActor)
+		{
+			if (UPCGComponent* LocalComponent = PartitionActor->GetLocalComponent(OriginalComponent))
+			{
+				// Add check to avoid infinite loop
+				if (ensure(!LocalComponent->IsPartitioned()))
+				{
+					FPCGTaskId LocalTask = InvalidPCGTaskId;
+#if WITH_EDITOR
+					if (bChainedDispatchToLocalComponents && TaskIds.Num() > 0)
+					{
+						LocalTask = InFunc(LocalComponent, { TaskIds.Last() });
+					}
+					else
+#endif // WITH_EDITOR
+					{
+						LocalTask = InFunc(LocalComponent, { });
+					}
+
+					if (LocalTask != InvalidPCGTaskId)
+					{
+						TaskIds.Add(LocalTask);
+					}
+				}
+			}
+		}
+	};
+
+#if WITH_EDITOR
+	if (bChainedDispatchToLocalComponents)
+	{
+		TArray<TObjectPtr<APCGPartitionActor>> SortedPartitionActors = PartitionActors.Array();
+		SortedPartitionActors.Remove(nullptr);
+		SortedPartitionActors.Sort([](const APCGPartitionActor& ActorA, const APCGPartitionActor& ActorB)
+		{
+			return ActorA.GetPathName() < ActorB.GetPathName();
+		});
+		Algo::ForEach(SortedPartitionActors, DispatchLambda);
+	}
+	else
+#endif // WITH_EDITOR
+	{
+		Algo::ForEach(PartitionActors, DispatchLambda);
+	}
+
+	return TaskIds;
+}
+
+bool FPCGTrackingManager::RegisterOrUpdateExecutionSource(IPCGGraphExecutionSource* InExecutionSource, bool bDoPartitionMapping)
+{
+	return RegisterOrUpdateExecutionSource(InExecutionSource, bDoPartitionMapping, /*bForce=*/false);
+}
+
+bool FPCGTrackingManager::RegisterOrUpdateExecutionSource(IPCGGraphExecutionSource* InExecutionSource, bool bDoPartitionMapping, bool bForce)
+{
+	check(InExecutionSource);
+
+	if (!PCGTrackingManager::IsValidOriginalExecutionSource(InExecutionSource))
+	{
+		return false;
+	}
+
+	// Check also that the bounds are valid. If not early out.
+	if (!InExecutionSource->GetExecutionState().GetBounds().IsValid)
+	{
+		UE_LOG(LogPCG, Error, TEXT("[RegisterOrUpdateExecutionSource] Component has invalid bounds, not registered nor updated."));
+		return false;
+	}
+
+	// @todo_pcg: support execution source
+	UPCGComponent* PCGComponent = Cast<UPCGComponent>(InExecutionSource);
+#if WITH_EDITOR
+	if(PCGComponent)
+	{
+		PCG::TScopeLock Lock(DelayedComponentsRegistrationLock);
+
+		DelayedComponentsToUnregister.Remove(PCGComponent);
+		DelayedComponentsToRegister.Remove(PCGComponent);
+
+		// When we are running Construction Script there are 2 possibilities
+		// 1. This component is replacing an existing PCG Component and we don't want to register, we want to replace in the RemapPCGComponent call (remap will remove this component from the Delayed list)
+		// 2. This comopnent is new and so it will be registered after the potential replacement call on the next frame in FPCGTrackingManager::Tick
+		if (!bForce && GetWorld()->bIsRunningConstructionScript && PCGComponent->IsCreatedByConstructionScript())
+		{
+			DelayedComponentsToRegister.Add(PCGComponent);
+			return true;
+		}
+	}
+#endif // WITH_EDITOR
+
+	const bool bWasAlreadyRegistered = IsExecutionSourceRegistered(InExecutionSource);
+
+	// First check if the execution source has changed its partitioned flag.
+	const bool bIsPartitioned = InExecutionSource->GetExecutionState().IsPartitioned();
+	if (bIsPartitioned && NonPartitionedOctree.Contains(InExecutionSource))
+	{
+		UnregisterNonPartitionedExecutionSource(InExecutionSource);
+	}
+	else if (!bIsPartitioned && PartitionedOctree.Contains(InExecutionSource))
+	{
+		UnregisterPartitionedExecutionSource(InExecutionSource);
+	}
+
+	PCGSubsystem->OnOriginalExecutionSourceRegistered(InExecutionSource);
+
+	// Then register/update accordingly
+	bool bHasChanged = false;
+	if (bIsPartitioned)
+	{
+		bHasChanged = RegisterOrUpdatePartitionedExecutionSource(InExecutionSource, bDoPartitionMapping);
+	}
+	else
+	{
+		bHasChanged = RegisterOrUpdateNonPartitionedExecutionSource(InExecutionSource);
+	}
+
+#if WITH_EDITOR
+	// And finally handle the tracking. Only do it when the execution source is registered for the first time.
+	if (!bWasAlreadyRegistered && bHasChanged)
+	{
+		RegisterTracking(InExecutionSource);
+	}
+#endif // WITH_EDITOR
+
+	return bHasChanged;
+}
+
+bool FPCGTrackingManager::RegisterOrUpdatePartitionedExecutionSource(IPCGGraphExecutionSource* InExecutionSource, bool bDoActorMapping)
+{
+	FBox Bounds(EForceInit::ForceInit);
+	bool bExecutionSourceHasChanged = false;
+	bool bExecutionSourceWasAdded = false;
+
+	PartitionedOctree.AddOrUpdateExecutionSource(InExecutionSource, Bounds, bExecutionSourceHasChanged, bExecutionSourceWasAdded);
+
+	// @todo_pcg: support execution source
+	UPCGComponent* PCGComponent = Cast<UPCGComponent>(InExecutionSource);
+	if (!PCGComponent)
+	{
+		return bExecutionSourceHasChanged;
+	}
+
+#if WITH_EDITOR
+	// In Editor only, we will create new partition actors depending on the new bounds and generation trigger. Runtime managed execution sources should not create PAs here
+	if ((bExecutionSourceHasChanged || bExecutionSourceWasAdded) && !InExecutionSource->GetExecutionState().IsManagedByRuntimeGenSystem())
+	{
+		bool bHasUnbounded = false;
+		PCGHiGenGrid::FSizeArray GridSizes;
+		ensure(PCGHelpers::GetGenerationGridSizes(InExecutionSource->GetExecutionState().GetGraph(), PCGSubsystem->GetPCGWorldActor(), GridSizes, bHasUnbounded));
+		PCGSubsystem->CreatePartitionActorsWithinBounds(PCGComponent, Bounds, GridSizes);
+	}
+#endif // WITH_EDITOR
+
+	// After adding/updating, try to do the mapping (if we asked for it and the component changed)
+	if (bDoActorMapping)
+	{
+		if (bExecutionSourceHasChanged)
+		{
+			UpdateMappingPCGComponentPartitionActor(PCGComponent);
+		}
+	}
+	else
+	{
+		if (!bExecutionSourceWasAdded)
+		{
+			// If we do not want a mapping, delete the existing one
+			DeleteMappingPCGComponentPartitionActor(PCGComponent);
+		}
+		else
+		{
+			// Here bDoActorMapping is false so we just want to update the ComponentToPartitionActorsMap without adding/removing graph instances
+			UpdateMappingPCGComponentPartitionActor(PCGComponent, /*bChangeGraphInstances=*/false);
+		}
+	}
+
+	return bExecutionSourceHasChanged;
+}
+
+bool FPCGTrackingManager::RegisterOrUpdateNonPartitionedExecutionSource(IPCGGraphExecutionSource* InExecutionSource)
+{
+	FBox Bounds(EForceInit::ForceInit);
+	bool bExecutionSourceHasChanged = false;
+	bool bExecutionSourceWasAdded = false;
+
+	NonPartitionedOctree.AddOrUpdateExecutionSource(InExecutionSource, Bounds, bExecutionSourceHasChanged, bExecutionSourceWasAdded);
+
+	return bExecutionSourceHasChanged;
+}
+
+bool FPCGTrackingManager::RemapExecutionSource(const IPCGGraphExecutionSource* InOldExecutionSource, IPCGGraphExecutionSource* InNewExecutionSource, bool bDoPartitionMapping)
+{
+	check(InOldExecutionSource && InNewExecutionSource);
+
+	bool bBoundsChanged = false;
+
+	if (InOldExecutionSource->GetExecutionState().IsPartitioned())
+	{
+		if (!PartitionedOctree.RemapExecutionSource(InOldExecutionSource, InNewExecutionSource, bBoundsChanged))
+		{
+			return false;
+		}
+	}
+	else
+	{
+		if (!NonPartitionedOctree.RemapExecutionSource(InOldExecutionSource, InNewExecutionSource, bBoundsChanged))
+		{
+			return false;
+		}
+	}
+
+	// @todo_pcg: support execution source
+	const UPCGComponent* OldComponent = Cast<UPCGComponent>(InOldExecutionSource);
+	UPCGComponent* NewComponent = Cast<UPCGComponent>(InNewExecutionSource);
+	
+	if (OldComponent && NewComponent)
+	{
+#if WITH_EDITOR
+		// Remove it from the delayed
+		{
+			PCG::TScopeLock Lock(DelayedComponentsRegistrationLock);
+			DelayedComponentsToUnregister.Remove(OldComponent);
+			DelayedComponentsToRegister.Remove(NewComponent);
+		}
+#endif
+
+		// Remap all previous instances
+		auto RemapPreviousInstances = [OldComponent, NewComponent](TMap<const UPCGComponent*, TSet<TObjectPtr<APCGPartitionActor>>>& Map, PCG::FSharedLock& Lock)
+		{
+			PCG::TUniqueScopeLock WriteLock(Lock);
+
+			if (TSet<TObjectPtr<APCGPartitionActor>>* PartitionActorsPtr = Map.Find(OldComponent))
+			{
+				TSet<TObjectPtr<APCGPartitionActor>> PartitionActorsToRemap = MoveTemp(*PartitionActorsPtr);
+				Map.Remove(OldComponent);
+
+				for (APCGPartitionActor* Actor : PartitionActorsToRemap)
+				{
+					Actor->RemapGraphInstance(OldComponent, NewComponent);
+				}
+
+				Map.Add(NewComponent, MoveTemp(PartitionActorsToRemap));
+			}
+		};
+
+		RemapPreviousInstances(ComponentToPartitionActorsMap, ComponentToPartitionActorsMapLock);
+
+		// And update the mapping if bounds changed and we want to do actor mapping
+		if (bBoundsChanged && NewComponent->IsPartitioned() && bDoPartitionMapping)
+		{
+			UpdateMappingPCGComponentPartitionActor(NewComponent);
+		}
+	}
+
+#if WITH_EDITOR
+	RemapTracking(InOldExecutionSource, InNewExecutionSource);
+#endif // WITH_EDITOR
+
+	return true;
+}
+
+void FPCGTrackingManager::UnregisterExecutionSource(IPCGGraphExecutionSource* InExecutionSource, bool bForce)
+{
+	if (!InExecutionSource)
+	{
+		return;
+	}
+
+#if WITH_EDITOR
+	UPCGComponent* PCGComponent = Cast<UPCGComponent>(InExecutionSource);
+
+	if(PCGComponent)
+	{
+		PCG::TScopeLock Lock(DelayedComponentsRegistrationLock);
+		DelayedComponentsToRegister.Remove(PCGComponent);
+		DelayedComponentsToUnregister.Remove(PCGComponent);
+	}
+
+	if ((PartitionedOctree.Contains(InExecutionSource) || NonPartitionedOctree.Contains(InExecutionSource)))
+	{
+		// We also need to check that our current PCG Component is not deleted while being reconstructed by a construction script.
+		// If so, it will be "re-created" at some point with the same properties.
+		// In this particular case, we don't remove the PCG component from the octree and we won't delete the mapping, but mark it to be removed
+		// at next Subsystem tick. If we call "RemapPCGComponent" before, we will re-connect everything correctly.
+		// Ignore this if we force (aka when we actually unregister the delayed one)		
+		if (!bForce && PCGComponent && PCGComponent->IsCreatedByConstructionScript())
+		{
+			PCG::TScopeLock Lock(DelayedComponentsRegistrationLock);
+			DelayedComponentsToUnregister.Add(PCGComponent);
+			return;
+		}
+
+		UnregisterTracking(InExecutionSource);
+	}
+#endif // WITH_EDITOR
+
+	UnregisterPartitionedExecutionSource(InExecutionSource);
+	UnregisterNonPartitionedExecutionSource(InExecutionSource);
+}
+
+void FPCGTrackingManager::UnregisterPartitionedExecutionSource(IPCGGraphExecutionSource* InExecutionSource)
+{
+	PCGSubsystem->OnOriginalExecutionSourceUnregistered(InExecutionSource);
+
+	if (!PartitionedOctree.RemoveExecutionSource(InExecutionSource) || InExecutionSource->GetExecutionState().IsManagedByRuntimeGenSystem())
+	{
+		return;
+	}
+
+	// @todo_pcg: support execution source
+	if (UPCGComponent* PCGComponent = Cast<UPCGComponent>(InExecutionSource))
+	{
+		// Because of recursive component deletes actors that has components, we cannot do RemoveGraphInstance
+		// inside a lock. So copy the actors to clean up and release the lock before doing RemoveGraphInstance.
+		TSet<TObjectPtr<APCGPartitionActor>> PartitionActorsToCleanUp;
+		{
+			PCG::TUniqueScopeLock WriteLock(ComponentToPartitionActorsMapLock);
+			TSet<TObjectPtr<APCGPartitionActor>>* PartitionActorsPtr = ComponentToPartitionActorsMap.Find(PCGComponent);
+
+			if (PartitionActorsPtr)
+			{
+				PartitionActorsToCleanUp = MoveTemp(*PartitionActorsPtr);
+				ComponentToPartitionActorsMap.Remove(PCGComponent);
+			}
+		}
+
+#if WITH_EDITOR
+		// Avoid removing graph instances when we are being unloaded (non-destructive)
+		if (!PCGComponent->bUnregisteredThroughLoading)
+#endif
+		{
+			for (APCGPartitionActor* Actor : PartitionActorsToCleanUp)
+			{
+				Actor->RemoveGraphInstance(PCGComponent);
+			}
+		}
+	}
+}
+
+void FPCGTrackingManager::UnregisterNonPartitionedExecutionSource(IPCGGraphExecutionSource* InExecutionSource)
+{
+	PCGSubsystem->OnOriginalExecutionSourceUnregistered(InExecutionSource);
+
+	NonPartitionedOctree.RemoveExecutionSource(InExecutionSource);
+}
+
+void FPCGTrackingManager::ForAllIntersectingPartitionedExecutionSources(const FBoxCenterAndExtent& InBounds, TFunctionRef<void(IPCGGraphExecutionSource*)> InFunc) const
+{
+	PartitionedOctree.FindElementsWithBoundsTest(InBounds, [&InFunc](const FPCGExecutionSourceRef& ExecutionSourceRef)
+	{
+		InFunc(ExecutionSourceRef.ExecutionSource);
+	});
+}
+
+void FPCGTrackingManager::ForAllOriginalExecutionSources(TFunctionRef<void(IPCGGraphExecutionSource*)> InFunc)
+{
+	for (IPCGGraphExecutionSource* ExecutionSource : PartitionedOctree.GetAllExecutionSources())
+	{
+		InFunc(ExecutionSource);
+	}
+
+	for (IPCGGraphExecutionSource* ExecutionSource : NonPartitionedOctree.GetAllExecutionSources())
+	{
+		InFunc(ExecutionSource);
+	}
+}
+
+TArray<IPCGGraphExecutionSource*> FPCGTrackingManager::GetAllIntersectingExecutionSources(const FBoxCenterAndExtent& InBounds) const
+{
+	TArray<IPCGGraphExecutionSource*> Result;
+	auto AddToResult = [&Result](const FPCGExecutionSourceRef& ExecutionSourceRef)
+	{
+		Result.Add(ExecutionSourceRef.ExecutionSource);
+	};
+
+	PartitionedOctree.FindElementsWithBoundsTest(InBounds, AddToResult);
+	NonPartitionedOctree.FindElementsWithBoundsTest(InBounds, AddToResult);
+
+	return Result;
+}
+
+void FPCGTrackingManager::RegisterPartitionActor(APCGPartitionActor* InActor)
+{
+	check(InActor);
+		
+#if WITH_EDITOR
+	if (InActor->IsInvalidForPCG())
+	{
+		// No need to log actor was already flagged to be invalid (and logged a message for it), we can ignore it.
+		return;
+	}
+#endif
+
+	const FPCGGridDescriptor PartitionActorDescriptor = InActor->GetGridDescriptor();
+
+#if WITH_EDITOR
+	// Ignore Invalid Actors 
+	if (InvalidPartitionActors.Contains(InActor->GetActorGuid()))
+	{
+		InActor->SetInvalidForPCG();
+		UE_LOG(LogPCG, Warning, TEXT("[RegisterPartitionActor] Invalid PCG Partiton Actor '%s' (%s). Please delete actor to remove warning."), *InActor->GetName(), *InActor->GetPackage()->GetName());
+		return;
+	} // Invalidate actor if there is a duplicate for the same grid information
+	else if (FGuid* FoundGuid = PartitionActorRecords.Find({ PartitionActorDescriptor, InActor->GetGridCoord()}); FoundGuid && *FoundGuid != InActor->GetActorGuid())
+	{
+		InActor->SetInvalidForPCG();
+		UE_LOG(LogPCG, Warning, TEXT("[RegisterPartitionActor] Duplicate PCG Partition Actor '%s' (%s). Please delete actor to remove warning."), *InActor->GetName(), *InActor->GetPackage()->GetName());
+		return;
+	}
+#endif
+	const FIntVector GridCoord = InActor->GetGridCoord();
+
+	check(PartitionActorDescriptor.GetGridSize() > 0);
+
+	{
+		PCG::TUniqueScopeLock WriteLock(PartitionActorsMapLock);
+
+		TMap<FIntVector, TObjectPtr<APCGPartitionActor>>& PartitionActorsMapGrid = PartitionActorsMap.FindOrAdd(PartitionActorDescriptor);
+		if (PartitionActorsMapGrid.Contains(GridCoord))
+		{
+			return;
+		}
+
+		PartitionActorsMapGrid.Add(GridCoord, InActor);
+	}
+
+	APCGWorldActor* WorldActor = PCGSubsystem->GetPCGWorldActorForPartitionActor(InActor);
+	if (!ensureMsgf(WorldActor, TEXT("When registering %s (Level: %s), the PCG World Actor does not exist in the world of the subsystem nor the partition actor level.\n"
+		  "This is ill-formed, make sure the PCG World Actor exists if you use PCG Partition actors."),
+		  *InActor->GetName(), InActor->GetLevel() ? *InActor->GetLevel()->GetName() : TEXT("Unknown")))
+	{
+		return;
+	}
+
+#if WITH_EDITOR
+PRAGMA_DISABLE_DEPRECATION_WARNINGS
+	// Handle move of the Use2DGrid flag for old partition actors that haven't been resaved
+	InActor->UpdateUse2DGridIfNeeded(WorldActor->bUse2DGrid);
+PRAGMA_ENABLE_DEPRECATION_WARNINGS
+#endif // WITH_EDITOR
+
+	// Register to all the components that intersect with the PA. Ignore for runtime generated, it is handled manually
+	if (!PartitionActorDescriptor.IsRuntime())
+	{
+		PCG::TUniqueScopeLock WriteLock(ComponentToPartitionActorsMapLock);
+		ForAllIntersectingPartitionedExecutionSources(FBoxCenterAndExtent(InActor->GetFixedBounds()), [this, InActor, WorldActor, PartitionActorDescriptor](IPCGGraphExecutionSource* ExecutionSource)
+		{
+			// For each component, do the mapping if the component is generated or if the Partition Actor already has a local component
+			// Partition Actors could have a local component while the original doesn't have a mapping if for instance a cleanup/save on the original component happened
+			// while not saving the Partition Actors
+			if (UPCGComponent* PCGComponent = Cast<UPCGComponent>(ExecutionSource))
+			{
+				FPCGGridDescriptor ComponentGridDescriptor = PCGComponent->GetGridDescriptor(InActor->GetPCGGridSize());
+				if (ComponentGridDescriptor == PartitionActorDescriptor)
+				{
+					if (PCGComponent->bGenerated || (InActor->GetLocalComponent(PCGComponent) != nullptr))
+					{
+						bool bHasUnbounded = false;
+						PCGHiGenGrid::FSizeArray GridSizes;
+						ensure(PCGHelpers::GetGenerationGridSizes(PCGComponent->GetGraph(), WorldActor, GridSizes, bHasUnbounded));
+
+						if (GridSizes.Contains(PartitionActorDescriptor.GetGridSize()))
+						{
+							// In editor we might load/create partition actors while the component is registering. Because of that,
+							// the mapping might not already exists, even if the component is marked generated.
+							if (!UE::GetIsEditorLoadingPackage())
+							{
+								InActor->AddGraphInstance(PCGComponent);
+							}
+							
+							// Either we already had a local component or one was just added
+							if (InActor->GetLocalComponent(PCGComponent))
+							{
+								// Add mapping
+								ComponentToPartitionActorsMap.FindOrAdd(PCGComponent).Add(InActor);
+							}
+						}
+					}
+				}
+			}
+		});
+	}
+}
+
+void FPCGTrackingManager::UnregisterPartitionActor(APCGPartitionActor* Actor)
+{
+	check(Actor);
+
+#if WITH_EDITOR
+	// Ignore Invalid Actors 
+	if (Actor->IsInvalidForPCG())
+	{
+		return;
+	}
+#endif
+
+	const FIntVector GridCoord = Actor->GetGridCoord();
+	const FPCGGridDescriptor GridDescriptor = Actor->GetGridDescriptor();
+	if (!ensure(GridDescriptor.GetGridSize() > 0))
+	{
+		return;
+	}
+
+	if (TMap<FIntVector, TObjectPtr<APCGPartitionActor>>* PartitionActorsMapGrid = PartitionActorsMap.Find(GridDescriptor))
+	{
+		PCG::TUniqueScopeLock WriteLock(PartitionActorsMapLock);
+		PartitionActorsMapGrid->Remove(GridCoord);
+	}
+
+	// Unregister from all intersecting components. Ignore for runtime generated, it is handled manually
+	if (!GridDescriptor.IsRuntime())
+	{
+		PCG::TUniqueScopeLock WriteLock(ComponentToPartitionActorsMapLock);
+		ForAllIntersectingPartitionedExecutionSources(FBoxCenterAndExtent(Actor->GetFixedBounds()), [this, Actor](IPCGGraphExecutionSource* ExecutionSource)
+		{
+			if(UPCGComponent* PCGComponent = Cast<UPCGComponent>(ExecutionSource))
+			{
+				TSet<TObjectPtr<APCGPartitionActor>>* PartitionActorsPtr = ComponentToPartitionActorsMap.Find(PCGComponent);
+				if (PartitionActorsPtr)
+				{
+					PartitionActorsPtr->Remove(Actor);
+				}
+			}
+		});
+	}
+}
+
+void FPCGTrackingManager::ForAllIntersectingPartitionActors(const FBox& InBounds, TFunctionRef<void(APCGPartitionActor*)> InFunc) const
+{
+	if (!InBounds.IsValid)
+	{
+		return;
+	}
+
+	const bool bInBoundsAreEmpty = FMath::IsNearlyZero(InBounds.GetVolume());
+
+	auto ForAllIntersectingPartitionActorsOfGridSize = [InFunc, &InBounds, this, bInBoundsAreEmpty](const FPCGGridDescriptor& GridDescriptor)
+	{
+		PCG::TSharedScopeLock ReadLock(PartitionActorsMapLock);
+
+		const TMap<FIntVector, TObjectPtr<APCGPartitionActor>>* PartitionActorsMapGrid = PartitionActorsMap.Find(GridDescriptor);
+		if (!PartitionActorsMapGrid || PartitionActorsMapGrid->IsEmpty())
+		{
+			return;
+		}
+
+		const FIntVector MinCellCoords = UPCGActorHelpers::GetCellCoord(InBounds.Min, GridDescriptor.GetGridSize(), GridDescriptor.Is2DGrid());
+		const FIntVector MaxCellCoords = UPCGActorHelpers::GetCellCoord(InBounds.Max, GridDescriptor.GetGridSize(), GridDescriptor.Is2DGrid());
+
+		for (int32 z = MinCellCoords.Z; z <= MaxCellCoords.Z; z++)
+		{
+			for (int32 y = MinCellCoords.Y; y <= MaxCellCoords.Y; y++)
+			{
+				for (int32 x = MinCellCoords.X; x <= MaxCellCoords.X; x++)
+				{
+					FIntVector CellCoords(x, y, z);
+					if (const TObjectPtr<APCGPartitionActor>* ActorPtr = PartitionActorsMapGrid->Find(CellCoords))
+					{
+						if (APCGPartitionActor* Actor = ActorPtr->Get())
+						{
+							const FBox ActorBounds = Actor->GetFixedBounds();
+							// Exclude any surrounding cells which are touching this cell but not meaningfully overlapping, only if the InBounds are
+							// not "empty" (volume > 0). If the InBounds are "empty" (just a single point), the overlap will always be zero,
+							// so just do an intersection check.
+							const bool bShouldProcess = bInBoundsAreEmpty
+								? ActorBounds.Intersect(InBounds)
+								: !FMath::IsNearlyZero(Actor->GetFixedBounds().Overlap(InBounds).GetVolume());
+							
+							if (bShouldProcess)
+							{
+								InFunc(Actor);
+							}
+						}
+					}
+				}
+			}
+		}
+	};
+
+	TArray<FPCGGridDescriptor> GridDescriptors;
+	{
+		PCG::TSharedScopeLock ReadLock(PartitionActorsMapLock);
+
+		PartitionActorsMap.GenerateKeyArray(GridDescriptors);
+	}
+
+	for (const FPCGGridDescriptor& GridDescriptor : GridDescriptors)
+	{
+		ForAllIntersectingPartitionActorsOfGridSize(GridDescriptor);
+	}
+}
+
+void FPCGTrackingManager::UpdateMappingPCGComponentPartitionActor(UPCGComponent* InComponent, bool bInChangeGraphInstances)
+{
+	if (!PCGSubsystem->IsInitialized())
+	{
+		return;
+	}
+
+	check(InComponent);
+
+	// Do not modify Partition Actors when loading
+	const bool bChangeGraphInstances = bInChangeGraphInstances && !UE::GetIsEditorLoadingPackage();
+
+	if (!InComponent->GetGraph())
+	{
+		return;
+	}
+
+	// Get the bounds
+	FBox Bounds = PartitionedOctree.GetBounds(InComponent);
+
+	if (!Bounds.IsValid)
+	{
+		return;
+	}
+
+	if (const APCGWorldActor* WorldActor = PCGSubsystem->GetPCGWorldActor())
+	{
+		bool bHasUnbounded = false;
+		PCGHiGenGrid::FSizeArray GridSizes;
+		ensure(PCGHelpers::GetGenerationGridSizes(InComponent->GetGraph(), WorldActor, GridSizes, bHasUnbounded));
+				
+		TSet<FPCGGridDescriptor> GridDescriptors;
+		Algo::Transform(GridSizes, GridDescriptors, [InComponent](uint32 GridSize)
+		{ 
+			return InComponent->GetGridDescriptor(GridSize);
+		});
+
+		auto UpdateMapping = [this, InComponent, &Bounds, WorldActor, &GridDescriptors, bChangeGraphInstances]()
+		{
+			TSet<TObjectPtr<APCGPartitionActor>> RemovedActors;
+
+			{
+				PCG::TUniqueScopeLock WriteLock(ComponentToPartitionActorsMapLock);
+				TSet<TObjectPtr<APCGPartitionActor>>& PartitionActors = ComponentToPartitionActorsMap.FindOrAdd(InComponent);
+
+				TSet<TObjectPtr<APCGPartitionActor>> NewMapping;
+				ForAllIntersectingPartitionActors(Bounds, [&NewMapping, InComponent, WorldActor, &GridDescriptors, bChangeGraphInstances](APCGPartitionActor* Actor)
+				{
+					if (!Actor)
+					{
+						return;
+					}
+
+					// Only add a graph instance to partition actors with compatible Grid Descriptors
+					if (GridDescriptors.Contains(Actor->GetGridDescriptor()))
+					{
+						if (bChangeGraphInstances)
+						{
+							Actor->AddGraphInstance(InComponent);
+						}
+
+						// Either we already had a local component or one was just added
+						if (Actor->GetLocalComponent(InComponent))
+						{
+							NewMapping.Add(Actor);
+						}
+					}
+				});
+
+				// Find the ones that were removed
+				RemovedActors = PartitionActors.Difference(NewMapping);
+
+				PartitionActors = MoveTemp(NewMapping);
+			}
+
+			// Here we can only remove graph instances if bChangeGraphInstances, which means in some edge cases
+			// we might have Partition Actors that have a PCG Component which is no longer linked to the original component.
+			// Those need to be manually cleaned up
+			if (bChangeGraphInstances)
+			{
+				// No need to be locked to do this.
+				for (APCGPartitionActor* RemovedActor : RemovedActors)
+				{
+					if (RemovedActor)
+					{
+						RemovedActor->RemoveGraphInstance(InComponent);
+					}
+				}
+			}
+		};
+
+		UpdateMapping();
+	}
+}
+
+TSet<TObjectPtr<APCGPartitionActor>> FPCGTrackingManager::GetPCGComponentPartitionActorMappings(UPCGComponent* InComponent) const
+{
+	if (const TSet<TObjectPtr<APCGPartitionActor>>* PartitionActorsPtr = ComponentToPartitionActorsMap.Find(InComponent))
+	{
+		return *PartitionActorsPtr;
+	}
+
+	return TSet<TObjectPtr<APCGPartitionActor>>();
+}
+
+void FPCGTrackingManager::DeleteMappingPCGComponentPartitionActor(UPCGComponent* InComponent)
+{
+	check(InComponent);
+
+	if (!InComponent->IsPartitioned())
+	{
+		return;
+	}
+
+	auto DeleteMapping = [this, InComponent](TMap<const UPCGComponent*, TSet<TObjectPtr<APCGPartitionActor>>>& Map, PCG::FSharedLock& Lock)
+	{
+		PCG::TUniqueScopeLock WriteLock(Lock);
+
+		if (TSet<TObjectPtr<APCGPartitionActor>>* PartitionActorsPtr = Map.Find(InComponent))
+		{
+			for (APCGPartitionActor* Actor : *PartitionActorsPtr)
+			{
+				Actor->RemoveGraphInstance(InComponent);
+			}
+
+			PartitionActorsPtr->Empty();
+		}
+	};
+
+	DeleteMapping(ComponentToPartitionActorsMap, ComponentToPartitionActorsMapLock);
+}
+
+bool FPCGTrackingManager::IsExecutionSourceRegistered(const IPCGGraphExecutionSource* InExecutionSource) const
+{
+	return PartitionedOctree.Contains(InExecutionSource) || NonPartitionedOctree.Contains(InExecutionSource);
+}
+
+bool FPCGTrackingManager::AnyRuntimeGenExecutionSourcesExist() const
+{
+	for (IPCGGraphExecutionSource* ExecutionSource : PartitionedOctree.GetAllExecutionSources())
+	{
+		if (ExecutionSource && ExecutionSource->GetExecutionState().IsManagedByRuntimeGenSystem())
+		{
+			return true;
+		}
+	}
+
+	for (IPCGGraphExecutionSource* ExecutionSource : NonPartitionedOctree.GetAllExecutionSources())
+	{
+		if (ExecutionSource && ExecutionSource->GetExecutionState().IsManagedByRuntimeGenSystem())
+		{
+			return true;
+		}
+	}
+
+	return false;
+}
+
+TSet<IPCGGraphExecutionSource*> FPCGTrackingManager::GetAllRegisteredPartitionedExecutionSources() const
+{
+	return PartitionedOctree.GetAllExecutionSources();
+}
+
+TSet<IPCGGraphExecutionSource*> FPCGTrackingManager::GetAllRegisteredNonPartitionedExecutionSources() const
+{
+	return NonPartitionedOctree.GetAllExecutionSources();
+}
+
+TSet<IPCGGraphExecutionSource*> FPCGTrackingManager::GetAllRegisteredExecutionSources() const
+{
+	TSet<IPCGGraphExecutionSource*> Res = GetAllRegisteredPartitionedExecutionSources();
+	Res.Append(GetAllRegisteredNonPartitionedExecutionSources());
+	return Res;
+}
+
+UPCGComponent* FPCGTrackingManager::GetLocalComponent(const FPCGGridDescriptor& GridDescriptor, const FIntVector& CellCoords, const UPCGComponent* InOriginalComponent) const
+{	
+	PCG::TSharedScopeLock ReadLock(PartitionActorsMapLock);
+
+	if (const TMap<FIntVector, TObjectPtr<APCGPartitionActor>>* PartitionActorsOnGrid = PartitionActorsMap.Find(GridDescriptor))
+	{
+		const TObjectPtr<APCGPartitionActor>* PartitionActor = PartitionActorsOnGrid->Find(CellCoords);
+		if (PartitionActor && *PartitionActor)
+		{
+			return (*PartitionActor)->GetLocalComponent(InOriginalComponent);
+		}
+	}
+
+	return nullptr;
+}
+
+APCGPartitionActor* FPCGTrackingManager::GetPartitionActor(const FPCGGridDescriptor& GridDescriptor, const FIntVector& CellCoords) const
+{
+	PCG::TSharedScopeLock ReadLock(PartitionActorsMapLock);
+
+	if (const TMap<FIntVector, TObjectPtr<APCGPartitionActor>>* PartitionActorsOnGrid = PartitionActorsMap.Find(GridDescriptor))
+	{
+		const TObjectPtr<APCGPartitionActor>* PartitionActor = PartitionActorsOnGrid->Find(CellCoords);
+		return PartitionActor ? *PartitionActor : nullptr;
+	}
+
+	return nullptr;
+}
+
+#if WITH_EDITOR
+bool FPCGTrackingManager::DoesPartitionActorRecordExist(const FPCGGridDescriptor& GridDescriptor, const FIntVector& GridCoords) const
+{
+	// In Editor this will not be empty (including -game)
+	if (PartitionActorRecords.Contains({ GridDescriptor, GridCoords }))
+	{
+		return true;
+	}
+	
+	return false;
+}
+
+void FPCGTrackingManager::RegisterTracking(IPCGGraphExecutionSource* InExecutionSource)
+{
+	if (!PCGTrackingManager::IsValidOriginalExecutionSource(InExecutionSource))
+	{
+		return;
+	}
+
+	UWorld* World = GetWorld();
+	check(World);
+
+	// Components owner needs to be always tracked
+	AlwaysTrackedKeysToExecutionSourcesMap.FindOrAdd(FPCGSelectionKey(EPCGActorFilter::Self)).Add(InExecutionSource);
+
+	UpdateTracking(InExecutionSource, /*bInShouldDirtyActors=*/ false);
+}
+
+void FPCGTrackingManager::UpdateTracking(IPCGGraphExecutionSource* InExecutionSource, bool bInShouldDirtyActors, const TArray<FPCGSelectionKey>* ChangedKeys)
+{
+	if (!PCGTrackingManager::IsValidOriginalExecutionSource(InExecutionSource))
+	{
+		return;
+	}
+
+	UWorld* World = GetWorld();
+	check(World);
+
+	// @todo_pcg: support execution source
+	UPCGComponent* PCGComponent = Cast<UPCGComponent>(InExecutionSource);
+	if (!PCGComponent)
+	{
+		return;
+	}
+
+	// If no keys are provided, update all tracking keys.
+	TArray<FPCGSelectionKey> AllKeys;
+	if (ChangedKeys == nullptr)
+	{
+		AllKeys = PCGComponent->GatherTrackingKeys();
+		ChangedKeys = &AllKeys;
+	}
+
+	check(ChangedKeys);
+
+	auto RemoveFromMap = [InExecutionSource](TMap<FPCGSelectionKey, TSet<IPCGGraphExecutionSource*>>& InMap, const FPCGSelectionKey& InKey)
+	{
+		if (TSet<IPCGGraphExecutionSource*>* Components = InMap.Find(InKey))
+		{
+			Components->Remove(InExecutionSource);
+			if (Components->IsEmpty())
+			{
+				InMap.Remove(InKey);
+			}
+		}
+	};
+
+	for (const FPCGSelectionKey& Key : *ChangedKeys)
+	{
+		// bShouldBeCulled is modified in IsKeyTrackedAndCulled
+		bool bShouldBeCulled = false;
+		if (!PCGComponent->IsKeyTrackedAndCulled(Key, bShouldBeCulled))
+		{
+			// Untrack
+			RemoveFromMap(CulledTrackedKeysToExecutionSourcesMap, Key);
+			RemoveFromMap(AlwaysTrackedKeysToExecutionSourcesMap, Key);
+		}
+		else if (bShouldBeCulled)
+		{
+			RemoveFromMap(AlwaysTrackedKeysToExecutionSourcesMap, Key);
+			CulledTrackedKeysToExecutionSourcesMap.FindOrAdd(Key).Add(PCGComponent);
+		}
+		else
+		{
+			RemoveFromMap(CulledTrackedKeysToExecutionSourcesMap, Key);
+			AlwaysTrackedKeysToExecutionSourcesMap.FindOrAdd(Key).Add(PCGComponent);
+		}
+	}
+}
+
+void FPCGTrackingManager::RemapTracking(const IPCGGraphExecutionSource* InOldExecutionSource, IPCGGraphExecutionSource* InNewExecutionSource)
+{
+	auto ReplaceInMap = [InOldExecutionSource, InNewExecutionSource](TMap<FPCGSelectionKey, TSet<IPCGGraphExecutionSource*>>& InMap)
+	{
+		for (TPair<FPCGSelectionKey, TSet<IPCGGraphExecutionSource*>>& It : InMap)
+		{
+			if (It.Value.Remove(InOldExecutionSource) > 0)
+			{
+				It.Value.Add(InNewExecutionSource);
+			}
+		}
+	};
+
+	ReplaceInMap(CulledTrackedKeysToExecutionSourcesMap);
+	ReplaceInMap(AlwaysTrackedKeysToExecutionSourcesMap);
+
+	const FObjectKey OldKey = Cast<UObject>(InOldExecutionSource);
+	const FObjectKey NewKey = Cast<UObject>(InNewExecutionSource);
+
+	// If this execution source has dependencies, we transfer them.
+	if (ExecutionSourcesToDependencyMap.Contains(OldKey))
+	{
+		TArray<FObjectKey> Temp;
+		ExecutionSourcesToDependencyMap.RemoveAndCopyValue(OldKey, Temp);
+		ExecutionSourcesToDependencyMap.Emplace(NewKey, std::move(Temp));
+	}
+
+	// If this execution source was a dependency to any other execution source, just remove it. New one will register itself when generating if needed.
+	for (TPair<FObjectKey, TArray<FObjectKey>>& It : ExecutionSourcesToDependencyMap)
+	{
+		It.Value.RemoveSwap(OldKey);
+	}
+}
+
+void FPCGTrackingManager::UnregisterTracking(IPCGGraphExecutionSource* InExecutionSource, const TSet<FPCGSelectionKey>* OptionalKeysToUntrack)
+{
+	if (!InExecutionSource)
+	{
+		return;
+	}
+
+	TSet<FPCGSelectionKey> KeysToRemove;
+	auto RemoveAllFromMap = [InExecutionSource, &KeysToRemove](TMap<FPCGSelectionKey, TSet<IPCGGraphExecutionSource*>>& InMap)
+	{
+		for (auto& It : InMap)
+		{
+			It.Value.Remove(InExecutionSource);
+			if (It.Value.IsEmpty())
+			{
+				KeysToRemove.Add(It.Key);
+			}
+		}
+	};
+
+	auto RemoveKeysFromMap = [InExecutionSource, &KeysToRemove](TMap<FPCGSelectionKey, TSet<IPCGGraphExecutionSource*>>& InMap, const TSet<FPCGSelectionKey>& SetToIterateOn)
+	{
+		for (const FPCGSelectionKey& KeyIt : SetToIterateOn)
+		{
+			if (TSet<IPCGGraphExecutionSource*>* ExecutionSourceSetPtr = InMap.Find(KeyIt))
+			{
+				ExecutionSourceSetPtr->Remove(InExecutionSource);
+				if (ExecutionSourceSetPtr->IsEmpty())
+				{
+					KeysToRemove.Add(KeyIt);
+				}
+			}
+		}
+	};
+
+	if (OptionalKeysToUntrack)
+	{
+		RemoveKeysFromMap(CulledTrackedKeysToExecutionSourcesMap, *OptionalKeysToUntrack);
+		RemoveKeysFromMap(AlwaysTrackedKeysToExecutionSourcesMap, *OptionalKeysToUntrack);
+	}
+	else
+	{
+		RemoveAllFromMap(CulledTrackedKeysToExecutionSourcesMap);
+		RemoveAllFromMap(AlwaysTrackedKeysToExecutionSourcesMap);
+	}
+
+	for (const FPCGSelectionKey& Key : KeysToRemove)
+	{
+		CulledTrackedKeysToExecutionSourcesMap.Remove(Key);
+		AlwaysTrackedKeysToExecutionSourcesMap.Remove(Key);
+	}
+}
+
+void FPCGTrackingManager::UnregisterTracking(IPCGGraphExecutionSource* InExecutionSource)
+{
+	if (!InExecutionSource)
+	{
+		return;
+	}
+
+	UnregisterTracking(InExecutionSource, nullptr);
+}
+
+bool FPCGTrackingManager::IsKeyTracked(const FPCGSelectionKey& InKey) const
+{
+	return CulledTrackedKeysToExecutionSourcesMap.Contains(InKey) || AlwaysTrackedKeysToExecutionSourcesMap.Contains(InKey);
+}
+
+bool FPCGTrackingManager::IsActorTracked(const AActor* InActor) const
+{
+	TRACE_CPUPROFILER_EVENT_SCOPE(FPCGTrackingManager::IsActorTracked);
+
+	TSet<FName> EmptySet;
+
+	const FPCGSelectionKey ActorKey = FPCGSelectionKey::CreateFromObjectPtr(InActor);
+	auto Matching = [&ActorKey, &EmptySet](const TPair<FPCGSelectionKey, TSet<IPCGGraphExecutionSource*>>& It) -> bool
+	{
+		return It.Key.IsMatching(ActorKey, EmptySet, It.Value, nullptr);
+	};
+
+	return Algo::AnyOf(CulledTrackedKeysToExecutionSourcesMap, Matching) || Algo::AnyOf(AlwaysTrackedKeysToExecutionSourcesMap, Matching);
+}
+
+void FPCGTrackingManager::ResetPartitionActorsMap()
+{
+	PCG::TUniqueScopeLock WriteLock(PartitionActorsMapLock);
+	PartitionActorsMap.Empty();
+}
+
+void FPCGTrackingManager::RegisterDelegates()
+{
+	GEngine->OnLevelActorAdded().AddRaw(this, &FPCGTrackingManager::OnActorAdded);
+	GEngine->OnLevelActorDeleted().AddRaw(this, &FPCGTrackingManager::OnActorDeleted);
+	FCoreUObjectDelegates::OnObjectPropertyChanged.AddRaw(this, &FPCGTrackingManager::OnObjectPropertyChanged);
+	FCoreUObjectDelegates::OnObjectModified.AddRaw(this, &FPCGTrackingManager::OnObjectModified);
+	FCoreUObjectDelegates::OnObjectPreSave.AddRaw(this, &FPCGTrackingManager::OnObjectSaved);
+	FCoreUObjectDelegates::OnPackageReloaded.AddRaw(this, &FPCGTrackingManager::OnPackageReloaded);
+
+	UWorld* World = GetWorld();
+	check(World);
+
+	// Listen to added/removed ActorDescs to keep track of unloaded Partiton Actors
+	if (UWorldPartition* WorldPartition = World->GetWorldPartition())
+	{
+		WorldPartition->OnActorDescInstanceAddedEvent.AddRaw(this, &FPCGTrackingManager::OnActorDescInstanceAdded);
+		WorldPartition->OnActorDescInstanceRemovedEvent.AddRaw(this, &FPCGTrackingManager::OnActorDescInstanceRemoved);
+	}
+
+	if (World->PersistentLevel)
+	{
+		World->PersistentLevel->OnLoadedActorAddedToLevelEvent.AddRaw(this, &FPCGTrackingManager::OnActorLoaded);
+		World->PersistentLevel->OnLoadedActorRemovedFromLevelEvent.AddRaw(this, &FPCGTrackingManager::OnActorUnloaded);
+
+		if (ULevelInstanceSubsystem* LevelInstanceSubsystem = UWorld::GetSubsystem<ULevelInstanceSubsystem>(World))
+		{
+			LevelInstanceSubsystem->OnLevelInstancesUpdated().AddRaw(this, &FPCGTrackingManager::OnLevelInstancesUpdated);
+			LevelInstanceSubsystem->OnLevelInstanceEditCancelled().AddRaw(this, &FPCGTrackingManager::OnLevelInstanceEditCancelled);
+		}
+	}
+
+	if (ULandscapeSubsystem* LandscapeSubsystem = World->GetSubsystem<ULandscapeSubsystem>())
+	{
+		LandscapeSubsystem->OnLandscapeProxyComponentDataChanged().AddRaw(this, &FPCGTrackingManager::OnLandscapeChanged);
+	}
+}
+
+void FPCGTrackingManager::UnregisterDelegates()
+{
+	GEngine->OnLevelActorAdded().RemoveAll(this);
+	GEngine->OnLevelActorDeleted().RemoveAll(this);
+	FCoreUObjectDelegates::OnObjectPropertyChanged.RemoveAll(this);
+	FCoreUObjectDelegates::OnObjectModified.RemoveAll(this);
+	FCoreUObjectDelegates::OnObjectPreSave.RemoveAll(this);
+	FCoreUObjectDelegates::OnPackageReloaded.RemoveAll(this);
+
+	UWorld* World = GetWorld();
+	check(World);
+
+	if (UWorldPartition* WorldPartition = World->GetWorldPartition())
+	{
+		WorldPartition->OnActorDescInstanceAddedEvent.RemoveAll(this);
+		WorldPartition->OnActorDescInstanceRemovedEvent.RemoveAll(this);
+	}
+
+	if (World->PersistentLevel)
+	{
+		World->PersistentLevel->OnLoadedActorAddedToLevelEvent.RemoveAll(this);
+		World->PersistentLevel->OnLoadedActorRemovedFromLevelEvent.RemoveAll(this);
+
+		if (ULevelInstanceSubsystem* LevelInstanceSubsystem = UWorld::GetSubsystem<ULevelInstanceSubsystem>(World))
+		{
+			LevelInstanceSubsystem->OnLevelInstancesUpdated().RemoveAll(this);
+			LevelInstanceSubsystem->OnLevelInstanceEditCancelled().RemoveAll(this);
+		}
+	}
+
+	if (ULandscapeSubsystem* LandscapeSubsystem = World->GetSubsystem<ULandscapeSubsystem>())
+	{
+		LandscapeSubsystem->OnLandscapeProxyComponentDataChanged().RemoveAll(this);
+	}
+}
+
+bool FPCGTrackingManager::ShouldDelayActor(AActor* InActor) const
+{
+	const bool bDisableDelayedActorRegistering = PCGTrackingManager::CVarDisableDelayedActorRegistering.GetValueOnAnyThread();
+	if (!InActor->HasActorRegisteredAllComponents() && !bDisableDelayedActorRegistering)
+	{
+		return true;
+	}
+
+	bool bShouldDelay = false;
+
+	// Check that whole Level Instance hierarchy is loaded
+	if (ILevelInstanceInterface* LevelInstance = Cast<ILevelInstanceInterface>(InActor))
+	{
+		// If Level Instance isn't loaded yet it should be delayed
+		if (!LevelInstance->IsLoaded() && !InActor->IsA<APackedLevelActor>())
+		{
+			return true;
+		}
+			
+		PCGTrackingManager::PropagateToLevelInstanceActors(LevelInstance, [this, &bShouldDelay](AActor* InChildActor)
+		{
+			if(ShouldDelayActor(InChildActor))
+			{
+				bShouldDelay = true;
+				return false;
+			}
+			return true;
+		});
+	}
+
+	return bShouldDelay;
+}
+
+void FPCGTrackingManager::ProcessDelayedEvents()
+{
+	// Safeguard, we can't add delayed actors if the subsystem is not initialized
+	if (!PCGSubsystem->IsInitialized() || DelayedChangedActors.IsEmpty())
+	{
+		return;
+	}
+			
+	auto LocalDelayedChangedActors = MoveTemp(DelayedChangedActors);
+
+	for (const auto& ActorToObject : LocalDelayedChangedActors)
+	{
+		AActor* Actor = ActorToObject.Key.ResolveObjectPtr();
+		UObject* Object = ActorToObject.Value.Get<0>().ResolveObjectPtr();
+		const EChangedObjectReason Reason = ActorToObject.Value.Get<1>();
+		
+		if (!Actor)
+		{
+			continue;
+		}
+
+		OnActorChanged_Internal(Actor, Object, Reason);
+	}
+}
+
+void FPCGTrackingManager::OnLevelInstancesUpdated(const TArray<ILevelInstanceInterface*>& InLevelInstances)
+{
+	for(ILevelInstanceInterface* LevelInstance : InLevelInstances)
+	{
+		OnActorAdded(Cast<AActor>(LevelInstance));
+	}
+}
+
+void FPCGTrackingManager::OnLevelInstanceEditCancelled(ILevelInstanceInterface* InLevelInstance, bool bInHasDiscardedChanges)
+{
+	if (bInHasDiscardedChanges)
+	{
+		OnActorAdded(Cast<AActor>(InLevelInstance));
+	}
+}
+
+void FPCGTrackingManager::OnActorLoaded(AActor& InActor)
+{
+	// We have to make sure to not create a infinite loop
+	if (PCGTrackingManager::ShouldIgnoreActor(&InActor) || InActor.GetWorld() != GetWorld())
+	{
+		return;
+	}
+
+	// Loaded actors should not dirty.
+	OnActorAdded_Internal(&InActor, /*bShouldDirty=*/ false);
+}
+
+void FPCGTrackingManager::OnActorAdded(AActor* InActor)
+{
+	// We have to make sure to not create a infinite loop
+	if (PCGTrackingManager::ShouldIgnoreActor(InActor) || InActor->GetWorld() != GetWorld())
+	{
+		return;
+	}
+			
+	// Implementation note: We delay adding because OnActorAdded fires before an actor's properties are set,
+	// so the actor is not ready for processing until the next tick.
+	// Only dirty if we aren't currently inside a loading path
+	OnActorAdded_Internal(InActor, /*bShouldDirty=*/ !UE::GetIsEditorLoadingPackage());
+}
+
+void FPCGTrackingManager::OnPackageReloaded(const EPackageReloadPhase InPackageReloadPhase, FPackageReloadedEvent* InPackageReloadedEvent)
+{
+	if (InPackageReloadPhase == EPackageReloadPhase::PrePackageFixup)
+	{
+		for (const auto& RepointedObjectPair : InPackageReloadedEvent->GetRepointedObjects())
+		{
+			if (RepointedObjectPair.Key && RepointedObjectPair.Key->IsAsset() && RepointedObjectPair.Key->IsA<UPCGGraphInterface>())
+			{
+				ForEachObjectOfClass(UPCGGraphInstance::StaticClass(), [OldGraphInterface = RepointedObjectPair.Key](UObject* InObj)
+				{
+					if (UPCGGraphInstance* GraphInstance = Cast<UPCGGraphInstance>(InObj); GraphInstance && GraphInstance->Graph == OldGraphInterface)
+					{
+						GraphInstance->TeardownCallbacks();
+					}
+				});
+			}
+		}
+	}
+	else if (InPackageReloadPhase == EPackageReloadPhase::PostPackageFixup)
+	{
+		TSet<UPCGGraphInstance*> GraphInstancesToRefresh;
+
+		for (const auto& RepointedObjectPair : InPackageReloadedEvent->GetRepointedObjects())
+		{
+			if (RepointedObjectPair.Value && RepointedObjectPair.Value->IsAsset() && RepointedObjectPair.Value->IsA<UPCGGraphInterface>())
+			{
+				ForEachObjectOfClass(UPCGGraphInstance::StaticClass(), [&GraphInstancesToRefresh, NewGraphInterface = RepointedObjectPair.Value](UObject* InObj)
+				{
+					if (UPCGGraphInstance* GraphInstance = Cast<UPCGGraphInstance>(InObj); GraphInstance && GraphInstance->Graph == NewGraphInterface)
+					{
+						GraphInstancesToRefresh.Add(GraphInstance);
+					}
+				});
+			}
+		}
+
+		for (UPCGGraphInstance* GraphInstance : GraphInstancesToRefresh)
+		{
+			if (GraphInstance)
+			{
+				GraphInstance->SetupCallbacks();
+				GraphInstance->OnGraphParametersChanged(EPCGGraphParameterEvent::GraphChanged, NAME_None);
+			}
+		}
+	}
+}
+
+void FPCGTrackingManager::OnActorDescInstanceAdded(FWorldPartitionActorDescInstance* InActorDescInstance)
+{
+	if (InActorDescInstance->GetActorNativeClass()->IsChildOf(APCGPartitionActor::StaticClass()))
+	{
+		const FPCGPartitionActorDesc* ActorDesc = static_cast<const FPCGPartitionActorDesc*>(InActorDescInstance->GetActorDesc());
+		FPCGGridCellDescriptor PartitionActorRecord = { ActorDesc->GetGridDescriptor(InActorDescInstance), FIntVector3(ActorDesc->GridIndexX, ActorDesc->GridIndexY, ActorDesc->GridIndexZ)};
+		if (!PartitionActorRecords.Contains(PartitionActorRecord))
+		{
+			PartitionActorRecords.Add(MoveTemp(PartitionActorRecord), ActorDesc->GetGuid());
+		}
+	}
+}
+
+void FPCGTrackingManager::OnActorDescInstanceRemoved(FWorldPartitionActorDescInstance* InActorDescInstance)
+{
+	if (InActorDescInstance->GetActorNativeClass()->IsChildOf(APCGPartitionActor::StaticClass()))
+	{
+		const FPCGPartitionActorDesc* ActorDesc = static_cast<const FPCGPartitionActorDesc*>(InActorDescInstance->GetActorDesc());
+		FPCGGridCellDescriptor PartitionActorRecord = { ActorDesc->GetGridDescriptor(InActorDescInstance), FIntVector3(ActorDesc->GridIndexX, ActorDesc->GridIndexY, ActorDesc->GridIndexZ)};
+		if (FGuid* FoundGuid = PartitionActorRecords.Find(PartitionActorRecord); FoundGuid && *FoundGuid == InActorDescInstance->GetGuid())
+		{
+			PartitionActorRecords.Remove(PartitionActorRecord);
+		}
+	}
+}
+
+void FPCGTrackingManager::BuildPartitionActorRecords(APCGWorldActor* PCGWorldActor, UWorldPartition* WorldPartition, TMap<FPCGGridCellDescriptor, FGuid>& OutPartitionActorRecords, TSet<FGuid>& OutInvalidPartitionActors)
+{
+	check(WorldPartition);
+	
+	FWorldPartitionHelpers::ForEachActorDescInstance<APCGPartitionActor>(WorldPartition, [PCGWorldActor, &OutPartitionActorRecords, &OutInvalidPartitionActors](const FWorldPartitionActorDescInstance* InActorDescInstance)
+	{
+		const FPCGPartitionActorDesc* ActorDesc = static_cast<const FPCGPartitionActorDesc*>(InActorDescInstance->GetActorDesc());
+		
+		// Fixup data
+		if (ActorDesc->bRequiresUse2DGridFixup && PCGWorldActor)
+		{
+			FPCGPartitionActorDesc* NonConstActorDesc = const_cast<FPCGPartitionActorDesc*>(ActorDesc);
+PRAGMA_DISABLE_DEPRECATION_WARNINGS
+			NonConstActorDesc->bUse2DGrid = PCGWorldActor->bUse2DGrid;
+PRAGMA_ENABLE_DEPRECATION_WARNINGS
+			NonConstActorDesc->bRequiresUse2DGridFixup = false;
+		}
+
+		// Add to invalid list if we couldn't properly fixup or if it was already marked invalid
+		if (ActorDesc->bInvalid)
+		{
+			OutInvalidPartitionActors.Add(ActorDesc->GetGuid());
+		}
+		else // Add to existing valid PCG Partition actors
+		{
+			FPCGGridCellDescriptor PartitionActorRecord = { ActorDesc->GetGridDescriptor(InActorDescInstance), FIntVector3(ActorDesc->GridIndexX, ActorDesc->GridIndexY, ActorDesc->GridIndexZ)};
+			if (!OutPartitionActorRecords.Contains(PartitionActorRecord))
+			{
+				OutPartitionActorRecords.Add(MoveTemp(PartitionActorRecord), ActorDesc->GetGuid());
+			}
+		}
+		return true;
+	});
+}
+
+void FPCGTrackingManager::BuildPartitionActorRecords()
+{
+	UWorld* World = GetWorld();
+	if (UWorldPartition* WorldPartition = World->GetWorldPartition())
+	{
+		APCGWorldActor* PCGWorldActor = PCGSubsystem->FindPCGWorldActor();
+
+		InvalidPartitionActors.Empty();
+		PartitionActorRecords.Empty();
+		BuildPartitionActorRecords(PCGWorldActor, WorldPartition, PartitionActorRecords, InvalidPartitionActors);
+	}
+}
+
+void FPCGTrackingManager::OnActorChanged_Recursive(AActor* InActor, UObject* InOriginatingChangeObject, const EChangedObjectReason InChangeReason)
+{
+	check(!DelayedChangedActors.Contains(InActor));
+
+	// Process actors recursively
+	auto OnActorChanged = [this, InOriginatingChangeObject, InChangeReason](AActor* InActor, auto&& OnActorChangedRecursive) -> void
+	{
+		if (InActor)
+		{
+			FActorPreviousData* PreviousData = ActorToPreviousDataMap.Find(InActor);
+			OnObjectChanged(InActor, InChangeReason, PreviousData, InOriginatingChangeObject);
+			if (PreviousData)
+			{
+				ActorToPreviousDataMap.Remove(InActor);
+			}
+		}
+		
+		PCGTrackingManager::PropagateToLevelInstanceActors(InActor, [this, OnActorChangedRecursive](AActor* LevelActor)
+		{
+			OnActorChangedRecursive(LevelActor, OnActorChangedRecursive);
+			return true;
+		});
+	};
+
+	OnActorChanged(InActor, OnActorChanged);
+}
+
+void FPCGTrackingManager::OnActorAdded_Internal(AActor* InActor, bool bShouldDirty)
+{
+	TRACE_CPUPROFILER_EVENT_SCOPE(FPCGTrackingManager::OnActorAdded);
+	check(!PCGTrackingManager::ShouldIgnoreActor(InActor) && InActor->GetWorld() == GetWorld());
+
+
+	if (!bShouldDirty || DelayedChangedActors.Contains(InActor))
+	{
+		// Nothing to do
+		return;
+	}
+
+	// We delay adding because OnActorAdded fires before an actor's properties are set,
+	// so the actor is not ready for processing until the next tick.
+	DelayedChangedActors.Emplace(InActor, {nullptr, EChangedObjectReason::ActorAdded});
+}
+
+void FPCGTrackingManager::OnActorChanged_Internal(AActor* InActor, UObject* InOriginatingChangeObject, const EChangedObjectReason InChangeReason)
+{
+	// Currently delayed, return
+	if (DelayedChangedActors.Contains(InActor))
+	{
+		return;
+	}
+
+	// Should still be delayed, add to delayed actors, return
+	if (ShouldDelayActor(InActor))
+	{
+		DelayedChangedActors.Emplace(InActor, {InOriginatingChangeObject, InChangeReason});
+		return;
+	}
+	
+	// Call Recursive method
+	OnActorChanged_Recursive(InActor, InOriginatingChangeObject, InChangeReason);
+}
+
+void FPCGTrackingManager::OnActorUnloaded(AActor& InActor)
+{
+	// Don't dirty on unload (to mirror the behavior in load)
+	OnActorDeleted_Internal(&InActor, /*bShouldDirty=*/false);
+}
+
+void FPCGTrackingManager::OnActorDeleted(AActor* InActor)
+{
+	if (PCGTrackingManager::ShouldIgnoreActor(InActor) || InActor->GetWorld() != GetWorld())
+	{
+		return;
+	}
+
+	// Only dirty if we aren't currently inside a loading path
+	OnActorDeleted_Internal(InActor, /*bShouldDirty=*/ !UE::GetIsEditorLoadingPackage());
+}
+
+void FPCGTrackingManager::OnActorDeleted_Internal(AActor* InActor, bool bShouldDirty)
+{
+	TRACE_CPUPROFILER_EVENT_SCOPE(UPCGTrackingManager::OnActorDeleted_Internal);
+	check(InActor && InActor->GetWorld() == GetWorld());
+
+	if (!bShouldDirty)
+	{
+		// Nothing to do
+		return;
+	}
+
+#if WITH_EDITOR
+	PCGTrackingManager::PropagateToLevelInstanceActors(InActor, [this, bShouldDirty](AActor* LevelActor)
+	{
+		OnActorDeleted_Internal(LevelActor, bShouldDirty);
+		return true;
+	});
+#endif // WITH_EDITOR
+
+	// Notify all components that the actor has changed (was removed), but the Refresh will only happen AFTER the actor was actually removed from the world (because of delayed refresh).
+	OnObjectChanged(InActor, EChangedObjectReason::ActorDeleted, /*InPreviousData=*/nullptr, /*InOriginatingObject=*/nullptr);
+}
+
+void FPCGTrackingManager::OnObjectModified(UObject* InObject)
+{
+	TRACE_CPUPROFILER_EVENT_SCOPE(FPCGTrackingManager::OnObjectModified);
+
+	// Nothing to do if we track nothing
+	if (CulledTrackedKeysToExecutionSourcesMap.IsEmpty() && AlwaysTrackedKeysToExecutionSourcesMap.IsEmpty())
+	{
+		return;
+	}
+
+	AActor* Actor = Cast<AActor>(InObject);
+
+	// Otherwise, if it's an actor component, track it as well
+	if (!Actor)
+	{
+		if (UActorComponent* ActorComponent = Cast<UActorComponent>(InObject))
+		{
+			Actor = ActorComponent->GetOwner();
+		}
+	}
+
+	if (!Actor || (Actor->GetWorld() != GetWorld()))
+	{
+		return;
+	}
+
+#if WITH_EDITOR
+	if (PCGTrackingManager::ShouldIgnoreActor(Actor))
+	{
+		return;
+	}
+#endif
+
+	auto StorePreviousData = [this](AActor* InActor, auto&& RecursiveCall) -> void
+	{
+		if (!ActorToPreviousDataMap.Contains(InActor))
+		{
+			if (IsActorTracked(InActor))
+			{
+				FActorPreviousData& PreviousData = ActorToPreviousDataMap.Add(InActor);
+				PreviousData.Get<0>() = PCGHelpers::GetActorBounds(InActor, /*bIgnorePCGCreatedComponents=*/false);
+				PreviousData.Get<1>() = TSet<FName>(InActor->Tags);
+				PreviousData.Get<2>() = FApp::GetCurrentTime();
+			}
+
+			// Also propagate the pre-change to all child actors if it is within a level instance.
+			PCGTrackingManager::PropagateToLevelInstanceActors(InActor, [this, RecursiveCall](AActor* LevelActor)
+			{
+				RecursiveCall(LevelActor, RecursiveCall);
+				return true;
+			});
+		}
+	};
+
+	StorePreviousData(Actor, StorePreviousData);
+}
+
+void FPCGTrackingManager::OnObjectSaved(UObject* InObject, FObjectPreSaveContext InObjectSaveContext)
+{
+	// Nothing to do if we track nothing
+	if (!InObject || (CulledTrackedKeysToExecutionSourcesMap.IsEmpty() && AlwaysTrackedKeysToExecutionSourcesMap.IsEmpty()))
+	{
+		return;
+	}
+
+	// Only trigger a refresh on save for limited data classes.
+	// At this point in time, We only track data tables and PCG data assets because in most cases we probably will catch other changes with OnObjectPropertyChanged.
+	// This is especially important to make sure we don't trigger refresh multiple times (less a problem for PCG assets, but a big problem for data tables).
+	if (!InObjectSaveContext.IsProceduralSave() && (InObject->IsA<UDataTable>() || InObject->IsA<UPCGDataAsset>() || InObject->IsA<UStaticMesh>()))
+	{
+		FPropertyChangedEvent Event{ nullptr };
+		OnObjectChanged(InObject, EChangedObjectReason::ObjectSaved, /*InPreviousData=*/nullptr, /*InOriginatingChangeObject=*/ InObject);
+	}
+}
+
+void FPCGTrackingManager::OnObjectPropertyChanged(UObject* InObject, FPropertyChangedEvent& InEvent)
+{
+	TRACE_CPUPROFILER_EVENT_SCOPE(FPCGTrackingManager::OnObjectPropertyChanged);
+
+	// Nothing to do if we track nothing
+	if (CulledTrackedKeysToExecutionSourcesMap.IsEmpty() && AlwaysTrackedKeysToExecutionSourcesMap.IsEmpty())
+	{
+		return;
+	}
+
+	const bool bValueNotInteractive = (InEvent.ChangeType != EPropertyChangeType::Interactive);
+	// Special exception for actor tags, as we can't track otherwise an actor "losing" a tag
+	const bool bActorTagChange = (InEvent.Property && InEvent.Property->GetFName() == GET_MEMBER_NAME_CHECKED(AActor, Tags));
+	// Another special exception for texture/mesh compilation. If the InEvent is empty and the object is a texture/mesh, we ignore it.
+	const bool bEventIsEmpty = (InEvent.Property == nullptr) && (InEvent.ChangeType == EPropertyChangeType::Unspecified);
+	const bool bIsTextureCompilationResult = bEventIsEmpty && InObject && InObject->IsA<UTexture>()
+		&& FTextureCompilingManager::Get().IsCompilingTexture(Cast<UTexture>(InObject));
+	// There is no equivalent for StaticMesh to know if we are in PostCompilation, so we assume there are still some meshes to compile (including this one).
+	// Might be an over-optimistic approach, might need a revisit.
+	const bool bIsStaticMeshCompilationResult = bEventIsEmpty && InObject && InObject->IsA<UStaticMesh>()
+		&& FStaticMeshCompilingManager::Get().GetNumRemainingMeshes() > 0;
+
+	// First check if it is an actor
+	AActor* Actor = Cast<AActor>(InObject);
+
+	// Otherwise, if it's an actor component, track it as well
+	if (!Actor)
+	{
+		if (UActorComponent* ActorComponent = Cast<UActorComponent>(InObject))
+		{
+			Actor = ActorComponent->GetOwner();
+		}
+	}
+
+	if (Actor && Actor->GetWorld() != GetWorld())
+	{
+		return;
+	}
+
+	if (Actor && PCGTrackingManager::ShouldIgnoreActor(Actor))
+	{
+		return;
+	}
+
+	const bool bNoOperation = (!bValueNotInteractive && !bActorTagChange)
+		|| bIsTextureCompilationResult
+		|| bIsStaticMeshCompilationResult
+		|| (Actor && (DelayedChangedActors.Contains(Actor)));
+
+	if (bNoOperation)
+	{
+		// Don't remove it as it will be needed when the delay will be called.
+		if (Actor && !(DelayedChangedActors.Contains(Actor)))
+		{
+			auto RemoveAll = [this](AActor* InActor, auto&& RecursiveCall) -> void
+			{
+				ActorToPreviousDataMap.Remove(InActor);
+				PCGTrackingManager::PropagateToLevelInstanceActors(InActor, [this, RecursiveCall](AActor* LevelActor)
+				{
+					RecursiveCall(LevelActor, RecursiveCall);
+					return true;
+				});
+			};
+
+			RemoveAll(Actor, RemoveAll);
+		}
+		
+		return;
+	}
+
+	if (Actor)
+	{
+		OnActorChanged_Internal(Actor, /*InOriginatingChangeObject=*/ InObject, EChangedObjectReason::ActorChanged);
+	}
+	else
+	{
+		OnObjectChanged(InObject, EChangedObjectReason::ObjectChanged, /*InPreviousData=*/nullptr, /*InOriginatingChangeObject=*/ InObject);
+	}
+}
+
+void FPCGTrackingManager::OnObjectChanged(UObject* InObject, const EChangedObjectReason InChangeReason, const FActorPreviousData* InPreviousData, const UObject* InOriginatingChangeObject)
+{
+	TRACE_CPUPROFILER_EVENT_SCOPE(FPCGTrackingManager::OnObjectChanged);
+
+	// Discard everything related to self for actor added/delete and changes coming from a PCG Component.
+	const bool bNoRefreshOwner = InChangeReason == EChangedObjectReason::ActorAdded
+		|| InChangeReason == EChangedObjectReason::ActorDeleted
+		|| InChangeReason == EChangedObjectReason::PCGComponentGenerated;
+
+	// Nothing to do if we track nothing or there is no object.
+	if (!InObject || (CulledTrackedKeysToExecutionSourcesMap.IsEmpty() && AlwaysTrackedKeysToExecutionSourcesMap.IsEmpty()))
+	{
+		return;
+	}
+
+	// Don't react to what the PCG Component is already reacting to.
+	static const TArray<const UClass*> ExcludedClasses =
+	{
+		UPCGComponent::StaticClass()
+	};
+
+	if (Algo::AnyOf(ExcludedClasses, [InObject](const UClass* Class) -> bool { return InObject->IsA(Class); }))
+	{
+		return;
+	}
+
+	const AActor* Actor = Cast<AActor>(InObject);
+
+	ensure(!Actor || Actor->GetWorld() == GetWorld());
+
+	TSet<FName> RemovedTags;
+	if (Actor && InPreviousData && !InPreviousData->Get<1>().IsEmpty())
+	{
+		RemovedTags = InPreviousData->Get<1>().Difference(TSet<FName>(Actor->Tags));
+	}
+
+	FPCGSelectionKey SelectionKey = FPCGSelectionKey::CreateFromObjectPtr(InObject);
+	
+	TArray<FBox, TInlineAllocator<2>> Bounds;
+	const FBox ActorBounds = Actor ? PCGHelpers::GetActorBounds(Actor, /*bIgnorePCGCreatedComponents=*/false) : FBox{ EForceInit::ForceInit };
+	Bounds.Add(ActorBounds);
+	// Also provide previous bounds if we do have previous data
+	if (InPreviousData)
+	{
+		const FBox& PreviousBounds = InPreviousData->Get<0>();
+		if (PreviousBounds.IsValid && !PreviousBounds.Equals(ActorBounds))
+		{
+			Bounds.Add(PreviousBounds);
+		}
+	}
+
+	OnSelectionKeyChanged(SelectionKey, InOriginatingChangeObject, RemovedTags, bNoRefreshOwner, Bounds);
+}
+
+void FPCGTrackingManager::OnSelectionKeyChanged(const FPCGSelectionKey& InSelectionKey, const UObject* InOriginatingChangeObject, TArrayView<const FBox> InBounds)
+{
+	const UObject* ChangedObject = InSelectionKey.GetObjectFromPath();
+	ensure(ChangedObject);
+
+	const AActor* Actor = Cast<AActor>(ChangedObject);
+
+	OnSelectionKeyChanged(InSelectionKey, InOriginatingChangeObject, {}, /*bInNoRefreshOwner=*/false, InBounds);
+}
+
+void FPCGTrackingManager::OnSelectionKeyChanged(const FPCGSelectionKey& InSelectionKey, const UObject* InOriginatingChangeObject, const TSet<FName>& InRemovedTags, bool bInNoRefreshOwner, TArrayView<const FBox> InBounds)
+{
+	TRACE_CPUPROFILER_EVENT_SCOPE(FPCGTrackingManager::OnSelectionKeyChanged);
+		
+	/** First gather all the execution sources that are tracking this key */
+	TSet<IPCGGraphExecutionSource*> CulledTrackedExecutionSources;
+	TSet<IPCGGraphExecutionSource*> AlwaysTrackedExecutionSources;
+	TArray<FPCGSelectionKey> MatchedKeys;
+	
+	FPCGSelectionKey InOriginatingSelectionKey = FPCGSelectionKey::CreateFromObjectPtr(InOriginatingChangeObject);
+	
+	const UObject* ChangedObject = InSelectionKey.GetObjectFromPath();
+	ensure(ChangedObject);
+
+	auto Gather = [InSelectionKey, InOriginatingSelectionKey, &InRemovedTags, &MatchedKeys](const TMap<FPCGSelectionKey, TSet<IPCGGraphExecutionSource*>>& InMap, TSet<IPCGGraphExecutionSource*>& OutSet)
+	{
+		for (const auto& It : InMap)
+		{
+			if (It.Key.IsMatching(InSelectionKey, InRemovedTags, It.Value, &OutSet) || It.Key.IsMatching(InOriginatingSelectionKey, InRemovedTags, It.Value, &OutSet))
+			{
+				MatchedKeys.Add(It.Key);
+			}
+		}
+	};
+
+	Gather(CulledTrackedKeysToExecutionSourcesMap, CulledTrackedExecutionSources);
+	Gather(AlwaysTrackedKeysToExecutionSourcesMap, AlwaysTrackedExecutionSources);
+
+	// If this actor is not tracked, just early out
+	if (CulledTrackedExecutionSources.IsEmpty() && AlwaysTrackedExecutionSources.IsEmpty())
+	{
+		return;
+	}
+
+	TArray<const UObject*, TInlineAllocator<2>> ChangedObjects = { ChangedObject };
+	ChangedObjects.AddUnique(InOriginatingChangeObject);
+
+	TOptional<FPCGDynamicTrackingPriority> ChangedObjectDynamicTrackingPriority;
+	// @todo_pcg: Find a better way to get this info as this doesn't work for managed objects that were spawned under a different outer
+	if (AActor* OuterActor = InOriginatingChangeObject ? InOriginatingChangeObject->GetTypedOuter<AActor>() : nullptr)
+	{
+		TArray<UPCGComponent*> OriginatingComponents;
+		OuterActor->GetComponents<UPCGComponent>(OriginatingComponents);
+		for (UPCGComponent* OriginatingComponent : OriginatingComponents)
+		{
+			if (!OriginatingComponent->IsCleaningUp() && OriginatingComponent->IsAnyObjectManagedByResource(ChangedObjects))
+			{
+				ChangedObjectDynamicTrackingPriority = OriginatingComponent->GetExecutionState().GetDynamicTrackingPriority();
+				break;
+			}
+		}
+	}
+
+	// We discard the component if:
+	//   * It is null.
+	//   * It is one of the objects that changed.
+	//   * We should not refresh the owner and its owner is one of the objects that changed.
+	//   * The objects that changed are managed by it (safe-guard, would probably induce an infinite loop).
+	//   * The objects that changed are managed by a lower tracking priority component. (lower tracking priority value means higher priority)
+	//   * One of the objects that changed is ignored by its original component.
+	auto ShouldDiscardComponent = [&ChangedObjects, &ChangedObjectDynamicTrackingPriority, bInNoRefreshOwner](const UPCGComponent* InComponent)
+	{
+		if (!InComponent
+			|| ChangedObjects.Contains(InComponent)
+			|| (bInNoRefreshOwner && ChangedObjects.Contains(InComponent->GetOwner()))
+			|| InComponent->IsCleaningUp()
+			|| InComponent->IsAnyObjectManagedByResource(ChangedObjects)
+			|| (ChangedObjectDynamicTrackingPriority.IsSet() && InComponent->GetExecutionState().GetDynamicTrackingPriority() < ChangedObjectDynamicTrackingPriority.GetValue()))
+		{
+			return true;
+		}
+
+		const UObject* ObjectIgnored = nullptr;
+		const UPCGComponent* OriginalComponent = InComponent->GetConstOriginalComponent();
+		EPCGIgnoreChangeOriginReason Reason;
+
+		if (ensure(OriginalComponent) && OriginalComponent->IsIgnoringAnyChangeOrigins(ChangedObjects, ObjectIgnored, Reason))
+		{
+			PCGGraphExecutionLogging::LogChangeOriginIgnoredForComponent(ObjectIgnored, InComponent, Reason);
+			return true;
+		}
+
+		return false;
+	};
+
+	TSet<UPCGComponent*> DirtyComponents;
+
+	EPCGComponentDirtyFlag DirtyFlag = EPCGComponentDirtyFlag::Actor;
+	if (ChangedObject->IsA<ALandscapeProxy>())
+	{
+		DirtyFlag = DirtyFlag | EPCGComponentDirtyFlag::Landscape;
+	}
+
+	{
+		TRACE_CPUPROFILER_EVENT_SCOPE(FPCGTrackingManager::OnObjectChanged::AlwaysTrackedUpdate);
+
+		for (IPCGGraphExecutionSource* ExecutionSource : AlwaysTrackedExecutionSources)
+		{
+			UPCGComponent* PCGComponent = Cast<UPCGComponent>(ExecutionSource);
+			// @todo_pcg: support execution source
+			if (ShouldDiscardComponent(PCGComponent))
+			{
+				continue;
+			}
+
+			// Don't mark "Owner changed" if the change originate from a PCG Component. It will be delegated to the ClearCacheForActor.
+			// It is necessary to avoid infine loops when there are multiple PCG components on one actor, and one component was generated.
+			const bool bOwnerChanged = (PCGComponent->GetOwner() == ChangedObject) && (!InOriginatingChangeObject || !InOriginatingChangeObject->IsA<UPCGComponent>());
+			bool bShouldDirty = bOwnerChanged || (PCGComponent->ShouldTrackLandscape() && ChangedObject->IsA<ALandscapeProxy>());
+
+			// Since when we clear the cache for a settings, we clear it all, it's only necessary to do it once on the original component. And we will dirty all the local components (with dirty dispatch)
+			if (!DirtyComponents.Contains(PCGComponent))
+			{
+				// Technically, a key that is always tracked can also be culled tracked, so maybe having an intersection test here to not
+				// dirty components that culled tracked and don't intersect with the component would be more efficient.
+				// In practice, it's mostly relevent for partitioned graphs and we will dirty all the locals anyway.
+				bShouldDirty |= ClearCacheForKeys(MatchedKeys, PCGComponent, /*bIntersect=*/true, InOriginatingChangeObject);
+				if (bShouldDirty)
+				{
+					PCGComponent->DirtyGenerated(DirtyFlag, /*bDispatchToLocalComponents=*/true);
+					DirtyComponents.Add(PCGComponent);
+				}
+			}
+		}
+	}
+
+	if (!CulledTrackedExecutionSources.IsEmpty())
+	{
+		// Then do an octree find to get all components that intersect with this actor.
+		// If the actor has moved, we also need to find components that intersected with it before
+		// We first do it for non-partitioned, then we do it for partitioned
+		auto UpdateNonPartitioned = [this, &DirtyComponents, ChangedObject, CulledTrackedExecutionSources, &InRemovedTags, DirtyFlag, InOriginatingChangeObject, bInNoRefreshOwner, &MatchedKeys, &ShouldDiscardComponent](const FPCGExecutionSourceRef& ExecutionSourceRef, const FBox& InBounds) -> void
+		{
+			TRACE_CPUPROFILER_EVENT_SCOPE(FPCGTrackingManager::OnObjectChanged::UpdateNonPartitioned);
+
+			if (UPCGComponent* PCGComponent = Cast<UPCGComponent>(ExecutionSourceRef.ExecutionSource))
+			{
+				// Don't dirty if the component was already dirtied, not tracked, the origin of the change or its owner is the changed actor and we should not refresh.
+				if (DirtyComponents.Contains(PCGComponent) ||
+					!CulledTrackedExecutionSources.Contains(PCGComponent) ||
+					ShouldDiscardComponent(PCGComponent))
+				{
+					return;
+				}
+
+				bool bShouldDirty = PCGComponent->ShouldTrackLandscape() && ChangedObject->IsA<ALandscapeProxy>();
+				bShouldDirty |= ClearCacheForKeys(MatchedKeys, PCGComponent, /*bIntersect=*/true, InOriginatingChangeObject);
+				if (bShouldDirty)
+				{
+					PCGComponent->DirtyGenerated(DirtyFlag);
+					DirtyComponents.Add(PCGComponent);
+				}
+			}
+		};
+
+		// For partitioned, we first need check if the original component intersect with the bounds, then forward the dirty call only to locals that intersect with the bounds.
+		// Note: CurrentActorBoundsPtr is passed by reference because it will be modified between lambda calls (cf comment above).
+		auto UpdatePartitioned = [this, &DirtyComponents, ChangedObject, CulledTrackedExecutionSources, &InRemovedTags, DirtyFlag, InOriginatingChangeObject, bInNoRefreshOwner, &MatchedKeys, &ShouldDiscardComponent](const FPCGExecutionSourceRef& ExecutionSourceRef, const FBox& InBounds)  -> void
+		{
+			TRACE_CPUPROFILER_EVENT_SCOPE(FPCGTrackingManager::OnObjectChanged::UpdatePartitioned);
+
+			if (UPCGComponent* PCGComponent = Cast<UPCGComponent>(ExecutionSourceRef.ExecutionSource))
+			{
+				// Don't dirty if the component is not tracked, the origin of the change or its owner is the changed actor and we should not refresh.
+				// We can "re-dirty" it because changes can impact different local components, from the same
+				// original component.
+				if (!CulledTrackedExecutionSources.Contains(PCGComponent) ||
+					ShouldDiscardComponent(PCGComponent))
+				{
+					return;
+				}
+
+				const FBox ComponentBounds = ExecutionSourceRef.Bounds.GetBox();
+				const bool bIntersect = InBounds.Intersect(ComponentBounds);
+				if (!bIntersect)
+				{
+					return;
+				}
+
+				const FBox Overlap = InBounds.Overlap(ComponentBounds);
+				bool bShouldDirty = PCGComponent->ShouldTrackLandscape() && ChangedObject->IsA<ALandscapeProxy>();
+				bShouldDirty |= ClearCacheForKeys(MatchedKeys, PCGComponent, /*bIntersect=*/true, InOriginatingChangeObject);
+				bool bWasDirtied = false;
+
+				// Since when we clear the cache for a settings, we clear it all, it's only necessary to do it once on the original component, then only dirty the local that intersects.
+				if (bShouldDirty)
+				{
+					ForAllIntersectingPartitionActors(Overlap, [ChangedObject, Component = PCGComponent, &bWasDirtied, DirtyFlag, InOriginatingChangeObject](APCGPartitionActor* InPartitionActor) -> void
+						{
+							if (UPCGComponent* LocalComponent = InPartitionActor->GetLocalComponent(Component))
+							{
+								bWasDirtied = true;
+								LocalComponent->DirtyGenerated(DirtyFlag);
+							}
+						});
+				}
+
+				if (bWasDirtied)
+				{
+					// Don't dispatch
+					PCGComponent->DirtyGenerated(DirtyFlag, /*bDispatchToLocalComponents=*/false);
+					DirtyComponents.Add(PCGComponent);
+				}
+			}
+		};
+
+		// Gather for all bounds
+		for (const FBox& Bounds : InBounds)
+		{
+			NonPartitionedOctree.FindElementsWithBoundsTest(Bounds, [&Bounds, &UpdateNonPartitioned](const FPCGExecutionSourceRef& ExecutionSourceRef) { UpdateNonPartitioned(ExecutionSourceRef, Bounds); });
+			PartitionedOctree.FindElementsWithBoundsTest(Bounds, [&Bounds, &UpdatePartitioned](const FPCGExecutionSourceRef& ExecutionSourceRef) { UpdatePartitioned(ExecutionSourceRef, Bounds); });
+		}
+	}
+
+	// Finally, dirty all components that always track this actor that are not yet notified.
+
+	// If it is a landscape and we should discard the refresh, early out.
+	bool bIsInEditingMode = false;
+	const ALandscapeProxy* Landscape = Cast<ALandscapeProxy>(ChangedObject);
+	if (PCGTrackingManager::ShouldDiscardLandscapeRefresh(Landscape, bIsInEditingMode, bIsCurrentlyExitingLandscapeEditMode))
+	{
+		// If we are in editing, keep track of all the dirtied landscape to refresh them when we exit.
+		if (bIsInEditingMode)
+		{
+			DirtiedLandscapes.AddUnique(Landscape);
+		}
+
+		return;
+	}
+
+	ULevelInstanceSubsystem* LevelInstanceSubsystem = GetWorld()->GetSubsystem<ULevelInstanceSubsystem>();
+	const UPCGComponent* OriginatingComponent = Cast<const UPCGComponent>(InOriginatingChangeObject);
+
+	// And refresh all dirtied components
+	for (UPCGComponent* Component : DirtyComponents)
+	{
+		if (!ensure(Component))
+		{
+			continue;
+		}
+
+		// This part checks if the change originates from a PCG Component. If so, we check the dirty component has no more other PCG dependencies.
+		// In that case we can proceed for the refresh, otherwise early out, this will be woken up by another dependency change.
+		if (OriginatingComponent)
+		{
+			if (TArray<FObjectKey>* Dependencies = ExecutionSourcesToDependencyMap.Find(Component))
+			{
+				Dependencies->RemoveSwap(OriginatingComponent);
+				if (Dependencies->IsEmpty())
+				{
+					ExecutionSourcesToDependencyMap.Remove(Component);
+				}
+				else
+				{
+					continue;
+				}
+			}
+		}
+
+		const bool bOwnerHasChanged = Component->GetOwner() == ChangedObject;
+
+		if ((!bInNoRefreshOwner || !bOwnerHasChanged) && (!Component->bOnlyTrackItself || bOwnerHasChanged))
+		{
+			const AActor* Actor = Cast<AActor>(ChangedObject);
+
+			// When an object changes, we need to make sure that we don't trigger a refresh on PCG components that are "higher" in the
+			// level hierarchy, otherwise we will end up generating in the Level Instance level, which is wrong.
+			// Note that in some instances (e.g. when something happens at the Level Instance level) we need to make sure that the
+			// PCG components higher-up are properly updated, hence the level instance depth
+			if (LevelInstanceSubsystem && Actor)
+			{
+				// Immediate Level Instance Parent (Actor)
+				const ILevelInstanceInterface* ActorLevelInstance = LevelInstanceSubsystem->GetParentLevelInstance(Actor);
+
+				// Immediate Level Instance Parent (PCG Component)
+				const ILevelInstanceInterface* ComponentLevelInstance = LevelInstanceSubsystem->GetParentLevelInstance(Component->GetOwner());
+								
+				// Is the Actor under a Level Instance that is the same or under the Components owning Level Instance
+				bool bActorLevelInstanceInComponentLevelInstanceHierarchy = false;
+
+				// Go up Level Instance hierarchy for Actor and see if we are part of the Component's level instance hierarchy, stop if we find an Editing Level instance as we don't want to impact top level 
+				// Components from an editing level instance
+				if (ComponentLevelInstance)
+				{
+					LevelInstanceSubsystem->ForEachLevelInstanceAncestors(Actor, [ComponentLevelInstance, &bActorLevelInstanceInComponentLevelInstanceHierarchy](const ILevelInstanceInterface* Ancestor)
+					{
+						if (Ancestor == ComponentLevelInstance)
+						{
+							bActorLevelInstanceInComponentLevelInstanceHierarchy = true;
+							return false; // stop iterating
+						}
+
+						return !Ancestor->IsEditing();
+					});
+				}
+				
+				// Actor in same level instance as Component or under the Components Level instance
+				if (bActorLevelInstanceInComponentLevelInstanceHierarchy)
+				{
+					// Allow update if Component is inside Editing level instance or if it is only in preview mode (generates transient data because it is in a non editing Level instance)
+					if (!ComponentLevelInstance->IsEditing() && Component->GetEditingMode() != EPCGEditorDirtyMode::Preview)
+					{
+						continue;
+					}
+				} // Allow update of Component from a Editing Level instance contained actor if the Component being refreshed is in preview mode and not itself in a Level Instance
+				else if (Actor->IsInEditLevelInstanceHierarchy() && (Component->GetEditingMode() != EPCGEditorDirtyMode::Preview || ComponentLevelInstance != nullptr))
+				{
+					continue;
+				}
+				// If ComponentLevelInstance is non-null and Actor Level instance isn't in its hierarchy then skip it
+				else if (ActorLevelInstance && ComponentLevelInstance != nullptr)
+				{
+					continue;
+				}
+			}
+
+			Component->Refresh();
+		}
+	}
+}
+
+void FPCGTrackingManager::OnLandscapeChanged(ALandscapeProxy* InLandscape, const FLandscapeProxyComponentDataChangedParams& InChangeParams)
+{
+	if (!InLandscape)
+	{
+		return;
+	}
+
+	if (PCGTrackingManager::CVarLandscapeRefreshTimeDelay.GetValueOnAnyThread() > 0)
+	{
+		LastLandscapeDirtyTime = FApp::GetCurrentTime();
+		DelayedModifiedLandscapes.AddUnique(InLandscape);
+	}
+	else
+	{
+		ApplyLandscapeChanges(InLandscape);
+	}
+}
+
+void FPCGTrackingManager::ApplyLandscapeChanges(ALandscapeProxy* InLandscape)
+{
+	if (!InLandscape)
+	{
+		return;
+	}
+
+	OnObjectChanged(InLandscape, EChangedObjectReason::ActorChanged, /*InPreviousData=*/nullptr, InLandscape);
+}
+
+void FPCGTrackingManager::OnPCGGraphStartsGenerating(IPCGGraphExecutionSource* InExecutionSource)
+{
+	TRACE_CPUPROFILER_EVENT_SCOPE(FPCGTrackingManager::OnPCGGraphStartsGenerating);
+
+	if (!PCGTrackingManager::IsValidExecutionSource(InExecutionSource) || PCGTrackingManager::CVarDisablePCGDataInterdependencyOptimization.GetValueOnAnyThread())
+	{
+		return;
+	}
+
+	// @todo_pcg: support execution source
+	UPCGComponent* PCGComponent = Cast<UPCGComponent>(InExecutionSource);
+	if (!PCGComponent)
+	{
+		return;
+	}
+
+	// When a graph starts generating, look for component that depends on it and keep a count.
+	// When this graph will be done generating, we'll trigger a dependency generation. But if multiple graphs are generated at the same time and 
+	// they all contribute to the same dependency, we'll trigger the dependency only when all the graphs are done generating
+	// We only need to gather components that depends on PCG Data, because other dependency changes will be caught by other engine callbacks (such as OnObjectPropertyChanged)
+	TSet<IPCGGraphExecutionSource*> TrackedExecutionSources;
+	TSet<FName> RemovedTags;
+
+	const FBox ExecutionSourceBounds = InExecutionSource->GetExecutionState().GetBounds();
+
+	auto IsKeyTrackingPCGData = [](const FPCGSelectionKey& InKey) { return InKey.OptionalExtraDependency && InKey.OptionalExtraDependency->IsChildOf(UPCGComponent::StaticClass()); };
+
+	AActor* Actor = PCGComponent->GetOwner();
+	const FPCGSelectionKey ActorKey = FPCGSelectionKey::CreateFromObjectPtr(Actor);
+
+	for (TPair<FPCGSelectionKey, TSet<IPCGGraphExecutionSource*>>& It : CulledTrackedKeysToExecutionSourcesMap)
+	{
+		if (!IsKeyTrackingPCGData(It.Key))
+		{
+			continue;
+		}
+
+		TSet<IPCGGraphExecutionSource*> TempTrackedExecutionSources;
+		It.Key.IsMatching(ActorKey, RemovedTags, It.Value, &TempTrackedExecutionSources);
+
+		// Removing all components that aren't intersecting with it, since it won't contribute to refresh
+		for (IPCGGraphExecutionSource* TrackedExecutionSource : TempTrackedExecutionSources)
+		{
+			if (UPCGComponent* TrackedComponent = Cast<UPCGComponent>(TrackedExecutionSource))
+			{
+				if (ensure(TrackedExecutionSource) && TrackedComponent->GetOwner() && PCGHelpers::GetActorBounds(TrackedComponent->GetOwner(), /*bIgnorePCGCreatedComponents=*/false).Intersect(ExecutionSourceBounds))
+				{
+					TrackedExecutionSources.Add(TrackedExecutionSource);
+				}
+			}
+		}
+	}
+
+	for (TPair<FPCGSelectionKey, TSet<IPCGGraphExecutionSource*>>& It : AlwaysTrackedKeysToExecutionSourcesMap)
+	{
+		if (!IsKeyTrackingPCGData(It.Key))
+		{
+			continue;
+		}
+
+		It.Key.IsMatching(ActorKey, RemovedTags, It.Value, &TrackedExecutionSources);
+	}
+
+	if (UObject* InExecutionSourceObject = Cast<UObject>(InExecutionSource); ensure(InExecutionSource))
+	{
+		for (IPCGGraphExecutionSource* TrackedExecutionSource : TrackedExecutionSources)
+		{
+			// Don't have a dependency on itself.
+			if (!TrackedExecutionSource || TrackedExecutionSource == InExecutionSource)
+			{
+				continue;
+			}
+
+			UPCGComponent* TrackedComponent = Cast<UPCGComponent>(TrackedExecutionSource);
+			// If the tracked component is currently ignoring the refresh from InComponent, don't add to the dependencies
+			if (TrackedComponent && TrackedComponent->IsIgnoringChangeOrigin(Actor))
+			{
+				continue;
+			}
+
+			if (UObject* ExecutionSourceObject = Cast<UObject>(TrackedExecutionSource); ensure(ExecutionSourceObject))
+			{
+				TArray<FObjectKey>& CurrentDependencies = ExecutionSourcesToDependencyMap.FindOrAdd(ExecutionSourceObject);
+				CurrentDependencies.AddUnique(InExecutionSourceObject);
+			}
+		}
+	}
+}
+
+void FPCGTrackingManager::OnPCGGraphCancelled(IPCGGraphExecutionSource* InExecutionSource)
+{
+	if (!InExecutionSource)
+	{
+		return;
+	}
+
+	if (UObject* InExecutionSourceObject = Cast<UObject>(InExecutionSource); ensure(InExecutionSourceObject))
+	{
+		TArray<FObjectKey> Keys;
+		ExecutionSourcesToDependencyMap.GetKeys(Keys);
+		for (const FObjectKey& Key : Keys)
+		{
+			TArray<FObjectKey>& Value = ExecutionSourcesToDependencyMap[Key];
+			Value.Remove(InExecutionSourceObject);
+			if (Value.IsEmpty())
+			{
+				ExecutionSourcesToDependencyMap.Remove(Key);
+			}
+		}
+	}
+}
+
+void FPCGTrackingManager::OnPCGGraphGeneratedOrCleaned(IPCGGraphExecutionSource* InExecutionSource)
+{
+	if (!PCGTrackingManager::IsValidExecutionSource(InExecutionSource))
+	{
+		return;
+	}
+
+	// @todo_pcg: support execution source
+	if (UPCGComponent* PCGComponent = Cast<UPCGComponent>(InExecutionSource))
+	{
+		OnObjectChanged(PCGComponent->GetOwner(), EChangedObjectReason::PCGComponentGenerated, /*InPreviousData=*/nullptr, PCGComponent);
+	}
+}
+
+bool FPCGTrackingManager::ClearCacheForKeys(const TArray<FPCGSelectionKey>& InKeys, const IPCGGraphExecutionSource* InExecutionSource, const bool bIntersect, const UObject* InOriginatingChange) const
+{
+	check(InExecutionSource);
+
+	// @todo_pcg: support execution source
+	const UPCGComponent* PCGComponent = Cast<UPCGComponent>(InExecutionSource);
+	if (!PCGComponent)
+	{
+		return false;
+	}
+
+	bool bShouldDirty = PCGComponent && PCGComponent->ShouldTrackLandscape() && InOriginatingChange && InOriginatingChange->IsA<ALandscapeProxy>();
+	
+	auto ClearCache = [this, InOriginatingChange, bIntersect, &bShouldDirty](const FPCGSelectionKey& Key, const FPCGSettingsAndCulling& SettingsAndCulling)
+	{
+		if (!SettingsAndCulling.Key.IsValid() || (SettingsAndCulling.Value && !bIntersect))
+		{
+			return;
+		}
+
+		// Extra care if the change originates from a PCGComponent. Only dirty if we are tracking a PCG component.
+		if (InOriginatingChange && InOriginatingChange->IsA<UPCGComponent>()
+			&& (!Key.OptionalExtraDependency || !Key.OptionalExtraDependency->IsChildOf(UPCGComponent::StaticClass())))
+		{
+			return;
+		}
+
+		bShouldDirty = true;
+		const UPCGSettings* Settings = SettingsAndCulling.Key.Get();
+		PCGSubsystem->CleanFromCache(Settings->GetElement().Get(), Settings);
+	};
+
+	for (const FPCGSelectionKey& Key : InKeys)
+	{
+		PCGComponent->ApplyToEachSettings(Key, ClearCache);
+	}
+
+	return bShouldDirty;
+}
+
+void FPCGTrackingManager::NotifyLandscapeEditModeExited()
+{
+	bIsCurrentlyExitingLandscapeEditMode = true;
+	// When the landscape edit mode is exited, force the refresh on all modified/dirtied landscapes.
+	for (TObjectKey<ALandscapeProxy> Landscape : DelayedModifiedLandscapes)
+	{
+		DirtiedLandscapes.AddUnique(Landscape);
+	}
+
+	DelayedModifiedLandscapes.Empty();
+
+	for (TObjectKey<ALandscapeProxy> Landscape : DirtiedLandscapes)
+	{
+		ApplyLandscapeChanges(Landscape.ResolveObjectPtr());
+	}
+
+	DirtiedLandscapes.Empty();
+
+	bIsCurrentlyExitingLandscapeEditMode = false;
+}
+#endif // WITH_EDITOR
diff --git a/Engine/Plugins/PCG/Source/PCG/Private/PCGTrackingManager.h b/Engine/Plugins/PCG/Source/PCG/Private/PCGTrackingManager.h
new file mode 100644
--- /dev/null
+++ b/Engine/Plugins/PCG/Source/PCG/Private/PCGTrackingManager.h
@@ -0,0 +1,303 @@
+// Copyright Epic Games, Inc. All Rights Reserved.
+
+#pragma once
+
+#include "PCGCommon.h"
+#include "Elements/PCGActorSelector.h"
+#include "Grid/PCGGridDescriptor.h"
+#include "Grid/PCGExecutionSourceOctree.h"
+#include "RuntimeGen/PCGRuntimeGenScheduler.h"
+
+#include "Containers/Array.h"
+#include "Containers/Map.h"
+#include "Containers/Set.h"
+#include "Templates/Function.h"
+#include "UObject/ObjectKey.h"
+#include "UObject/ObjectPtr.h"
+#include "UObject/ObjectSaveContext.h"
+#include "UObject/SoftObjectPtr.h"
+#include "UObject/WeakObjectPtr.h"
+
+class AActor;
+class ALandscapeProxy;
+class APCGPartitionActor;
+class FLandscapeProxyComponentDataChangedParams;
+class ILevelInstanceInterface;
+class UObject;
+class UPCGComponent;
+class UPCGGraph;
+class UPCGSubsystem;
+class UWorld;
+class UWorldPartition;
+class FWorldPartitionActorDescInstance;
+class FPackageReloadedEvent;
+
+enum class EPackageReloadPhase : uint8;
+
+/**
+* This class currently serves two purposes:
+* - Keep a mapping between original and local execution sources.
+* - Do change tracking to refresh execution sources.
+*
+* Its meant to be part of the PCG Subsystem and owned by it. We offload some logic to this class to avoid to clutter the subsystem.
+*/
+class FPCGTrackingManager
+{
+public:
+	friend UPCGSubsystem;
+	friend FPCGRuntimeGenScheduler;
+
+	~FPCGTrackingManager() = default;
+
+	/** Initializes callbacks, etc, tied to the PCG subsystem */
+	void Initialize();
+
+	/** Deinitializes callbacks, etc, tied to the PCG subsystem */
+	void Deinitialize();
+	
+	/** Should be called by the subsystem to handle delayed operations. */
+	void Tick();
+
+	/** Returns the PCGSubystem's World */
+	UWorld* GetWorld() const;
+
+#if WITH_EDITOR
+	/** Notify that we exited the Landscape edit mode. */
+	void NotifyLandscapeEditModeExited();
+
+	static void BuildPartitionActorRecords(APCGWorldActor* PCGWorldActor, UWorldPartition* WorldPartition, TMap<FPCGGridCellDescriptor, FGuid>& OutPartitionActorRecords, TSet<FGuid>& OutInvalidPartitionActors);
+#endif // WITH_EDITOR
+
+	/** Register a new execution source or update it. Returns true if it was added/updated. Thread safe */
+	bool RegisterOrUpdateExecutionSource(IPCGGraphExecutionSource* InExecutionSource, bool bDoPartitionMapping = true);
+
+	/** In case of blueprints, we need to remap the old execution source destroyed by the construction script to the new one. Returns true if a re-mapping was done. Thread safe */
+	bool RemapExecutionSource(const IPCGGraphExecutionSource* OldExecutionSource, IPCGGraphExecutionSource* NewExecutionSource, bool bDoPartitionMapping);
+
+	/** Unregister a execution source. Can force it, if we have a delayed unregister. Thread safe */
+	void UnregisterExecutionSource(IPCGGraphExecutionSource* InExecutionSource, bool bForce = false);
+
+	/** Register a new Partition actor, will be added to a map and will query all intersecting volume to bind to them if asked. Thread safe */
+	void RegisterPartitionActor(APCGPartitionActor* InActor);
+
+	/** Unregister a Partition actor, will be removed from the map and remove itself to all intersecting volumes. Thread safe */
+	void UnregisterPartitionActor(APCGPartitionActor* InActor);
+
+	/** Return a copy of all the registered partitioned execution sources. Thread safe */
+	TSet<IPCGGraphExecutionSource*> GetAllRegisteredPartitionedExecutionSources() const;
+
+	/** Return a copy of all the registered non-partitioned execution sources. Thread safe */
+	TSet<IPCGGraphExecutionSource*> GetAllRegisteredNonPartitionedExecutionSources() const;
+
+	/** Return a copy of all the registered execution sources. Thread safe */
+	TSet<IPCGGraphExecutionSource*> GetAllRegisteredExecutionSources() const;
+
+	/** Retrieves a local component using grid descriptor and grid coordinates, returns nullptr if no such component is found. */
+	UPCGComponent* GetLocalComponent(const FPCGGridDescriptor& GridDescriptor, const FIntVector& CellCoords, const UPCGComponent* InOriginalComponent) const;
+
+	APCGPartitionActor* GetPartitionActor(const FPCGGridDescriptor& GridDescriptor, const FIntVector& CellCoords) const;
+
+#if WITH_EDITOR
+	/** Returns true if there is record of a partition actor living in a certain grid cell, regardless of whether or not it is loaded. */
+	bool DoesPartitionActorRecordExist(const FPCGGridDescriptor& GridDescriptor, const FIntVector& GridCoords) const;
+#endif
+private:
+
+	/** Used to track where a change comes from. */
+	enum class EChangedObjectReason
+	{
+		ActorAdded,
+		ActorChanged,
+		ActorDeleted,
+		ObjectChanged,
+		ObjectSaved,
+		PCGComponentGenerated,
+		Undefined
+	};
+
+#if WITH_EDITOR
+	void RegisterDelegates();
+	void UnregisterDelegates();
+
+	/** If the partition grid size change, call this to empty the Partition actors map */
+	void ResetPartitionActorsMap();
+
+	void RegisterTracking(IPCGGraphExecutionSource* InExecutionSource);
+	void UpdateTracking(IPCGGraphExecutionSource* InExecutionSource, bool bInShouldDirtyActors, const TArray<FPCGSelectionKey>* OptionalChangedKeys = nullptr);
+
+	void ProcessDelayedEvents();
+	bool ShouldDelayActor(AActor* InActor) const;
+#endif
+	bool RegisterOrUpdateExecutionSource(IPCGGraphExecutionSource* InExecutionSource, bool bDoPartitionMapping, bool bForce);
+
+	// This class is only meant to be used as part of the PCG Subsytem and owned by it.
+	// So we put constructors private.
+	// We also need this class to be default constructible, since the PCGSubsytem needs to be default constructible.
+	FPCGTrackingManager() = default;
+	explicit FPCGTrackingManager(UPCGSubsystem* PCGSubsystem);
+
+	bool RegisterOrUpdatePartitionedExecutionSource(IPCGGraphExecutionSource* InExecutionSource, bool bDoPartitionMapping = true);
+	bool RegisterOrUpdateNonPartitionedExecutionSource(IPCGGraphExecutionSource* InExecutionSource);
+
+	void UnregisterPartitionedExecutionSource(IPCGGraphExecutionSource* InExecutionSource);
+	void UnregisterNonPartitionedExecutionSource(IPCGGraphExecutionSource* InExecutionSource);
+
+	/* Call the InFunc function to all local component registered to the original component. Return the list of all the tasks scheduled. Thread safe*/
+	TArray<FPCGTaskId> DispatchToRegisteredLocalComponents(UPCGComponent* OriginalComponent, const TFunctionRef<FPCGTaskId(UPCGComponent*, const TArray<FPCGTaskId>&)>& InFunc) const;
+
+	/* Call the InFunc function to all local component from the set of partition actors. Return the list of all the tasks scheduled. */
+	TArray<FPCGTaskId> DispatchToLocalComponents(UPCGComponent* OriginalComponent, const TSet<TObjectPtr<APCGPartitionActor>>& PartitionActors, const TFunctionRef<FPCGTaskId(UPCGComponent*, const TArray<FPCGTaskId>&)>& InFunc) const;
+
+	/** Call the InFunc function to all partitioned execution sources which bounds intersect 'InBounds'. */
+	void ForAllIntersectingPartitionedExecutionSources(const FBoxCenterAndExtent& InBounds, TFunctionRef<void(IPCGGraphExecutionSource*)> InFunc) const;
+
+	/** Call the InFunc function for all original execution sources (regardless of partitioned or not). */
+	void ForAllOriginalExecutionSources(TFunctionRef<void(IPCGGraphExecutionSource*)> InFunc);
+
+	/** Gather all the execution sources within some bounds. */
+	TArray<IPCGGraphExecutionSource*> GetAllIntersectingExecutionSources(const FBoxCenterAndExtent& InBounds) const;
+
+	/** Iterate other all the int coordinates given a box and call a callback. Thread safe */
+	void ForAllIntersectingPartitionActors(const FBox& InBounds, TFunctionRef<void(APCGPartitionActor*)> InFunc) const;
+
+	/** Update the current mapping between a PCG component and its PCG Partition actors */
+	void UpdateMappingPCGComponentPartitionActor(UPCGComponent* InComponent, bool bInChangeGraphInstances = true);
+
+	/** Returns the current mapping between a PCG component and its PCG Partition actors */
+	TSet<TObjectPtr<APCGPartitionActor>> GetPCGComponentPartitionActorMappings(UPCGComponent* InComponent) const;
+	
+	/** Delete the current mapping between a PCG component and its PCG Partition actors */
+	void DeleteMappingPCGComponentPartitionActor(UPCGComponent* InComponent);
+
+	/** Returns true if an executon source is registered */
+	bool IsExecutionSourceRegistered(const IPCGGraphExecutionSource* InExecutionSource) const;
+
+	/** Return true if there are any Original or Non-Partitioned execution sources set to GenerateAtRuntime. */
+	bool AnyRuntimeGenExecutionSourcesExist() const;
+
+	// Typedef to store the previous position and previous tags of changed actors.
+	using FActorPreviousData = TTuple<FBox, TSet<FName>, double>;
+
+#if WITH_EDITOR
+	/** Return true if the key is tracked.*/
+	bool IsKeyTracked(const FPCGSelectionKey& InKey) const;
+
+	/** Return true if the actor is tracked.*/
+	bool IsActorTracked(const AActor* InActor) const;
+
+	void OnPackageReloaded(const EPackageReloadPhase InPackageReloadPhase, FPackageReloadedEvent* InPackageReloadedEvent);
+	void OnActorDescInstanceAdded(FWorldPartitionActorDescInstance* InActorDescInstance);
+	void OnActorDescInstanceRemoved(FWorldPartitionActorDescInstance* InActorDescInstance);
+	void OnActorAdded(AActor* InActor);
+	void OnActorLoaded(AActor& InActor);
+	void OnActorAdded_Internal(AActor* InActor, bool bShouldDirty);
+	void OnActorChanged_Internal(AActor* InActor, UObject* InOriginatingChangeObject, const EChangedObjectReason InChangeReason);
+	void OnActorChanged_Recursive(AActor* InActor, UObject* InOriginatingChangeObject, const EChangedObjectReason InChangeReason);
+	void OnActorDeleted(AActor* InActor);
+	void OnActorUnloaded(AActor& InActor);
+	void OnActorDeleted_Internal(AActor* InActor, bool bShouldDirty);
+	void OnLandscapeChanged(ALandscapeProxy* InLandscape, const FLandscapeProxyComponentDataChangedParams& InChangeParams);
+	void OnLevelInstancesUpdated(const TArray<ILevelInstanceInterface*>& InLevelInstances);
+	void OnLevelInstanceEditCancelled(ILevelInstanceInterface* InLevelInstance, bool bInHasDiscardedChanges);
+	void ApplyLandscapeChanges(ALandscapeProxy* InLandscape);
+	void OnObjectModified(UObject* InObject);
+	void OnObjectPropertyChanged(UObject* InObject, FPropertyChangedEvent& InEvent);
+	void OnObjectSaved(UObject* InObject, FObjectPreSaveContext InObjectSaveContext);
+	void OnPCGGraphStartsGenerating(IPCGGraphExecutionSource* InExecutionSource);
+	void OnPCGGraphGeneratedOrCleaned(IPCGGraphExecutionSource* InExecutionSource);
+	void OnPCGGraphCancelled(IPCGGraphExecutionSource* InExecutionSource);
+
+	/** Remap the tracking in case of blueprints. */
+	void RemapTracking(const IPCGGraphExecutionSource* InOldExecutionSource, IPCGGraphExecutionSource* InNewExecutionSource);
+
+	/** Unregister tracking when a execution source is removed */
+	void UnregisterTracking(IPCGGraphExecutionSource* InExecutionSource);
+
+	/** Unregister tracking when a execution source is removed or keys are unregistered */
+	void UnregisterTracking(IPCGGraphExecutionSource* InExecutionSource, const TSet<FPCGSelectionKey>* OptionalKeysToUntrack);
+
+	/** Trigger an update when an object/actor changed.
+	* Can specify previous actor data if it has moved or tags changed to also update components that were at its previous position or tracked previous tags.
+	* Can also specify an optional object, originating the change, to avoid re-dirtying a component if it was the origin.
+	* Another option when an actor is deleted/unload, don't refresh their components.
+	*/
+	void OnObjectChanged(UObject* InObject, const EChangedObjectReason InChangeReason, const FActorPreviousData* InPreviousData = nullptr, const UObject* InOriginatingChangeObject = nullptr);
+
+	void OnSelectionKeyChanged(const FPCGSelectionKey& InSelectionKey, const UObject* InOriginatingChangeObject, TArrayView<const FBox> InBounds);
+	void OnSelectionKeyChanged(const FPCGSelectionKey& InSelectionKey, const UObject* InOriginatingChangeObject, const TSet<FName>& InRemovedTags, bool bInNoRefreshOwner, TArrayView<const FBox> InBounds);
+
+	/** Gather all settings from a given execution source that track the key, and clear the cache for them. Returns true if we should dirty afterwards (aka at least one settings was cleared and/or landscape changed). */
+	bool ClearCacheForKeys(const TArray<FPCGSelectionKey>& InKeys, const IPCGGraphExecutionSource* InExecutionSource, const bool bIntersect, const UObject* InOriginatingChange) const;
+
+	/** Build initial records for existing partition actors */
+	void BuildPartitionActorRecords();
+
+	/** Used by the PCG Builder to force dispatched local components to be dependency chained together so that they do not run at the same time for determinism */
+	void SetChainedDispatchToLocalComponents(bool bInChainedDispatch) { bChainedDispatchToLocalComponents = bInChainedDispatch; }
+#endif // WITH_EDITOR
+
+private:
+	// Cached subsystem
+	UPCGSubsystem* PCGSubsystem = nullptr;
+
+	// Mapping Component <-> PartitionActor
+	/** Octree tracking all partitioned pcg components */
+	FPCGExecutionSourceOctreeAndMap PartitionedOctree;
+
+	/** Mapping from grid size and grid coords to partition actor. We can only have 1 partition actor per grid cell. */
+	TMap<FPCGGridDescriptor, TMap<FIntVector, TObjectPtr<APCGPartitionActor>>> PartitionActorsMap;
+	mutable PCG::FSharedLock PartitionActorsMapLock;
+
+	/** Mapping between original components and its overlapping partition actors. */
+	TMap<const UPCGComponent*, TSet<TObjectPtr<APCGPartitionActor>>> ComponentToPartitionActorsMap;
+	mutable PCG::FSharedLock ComponentToPartitionActorsMapLock;
+
+	/** Will hold all the components that are not partitioned (and not local) and are tracking something. Will be use to dispatch actor tracking updates. */
+	FPCGExecutionSourceOctreeAndMap NonPartitionedOctree;
+
+#if WITH_EDITOR
+	/** Components to be registered/unregistered on the next frame. see RegisterOrUpdatePCGComponent/UnregisterPCGComponent for a better understanding on why it is needed. */
+	mutable PCG::FLock DelayedComponentsRegistrationLock;
+	TSet<UPCGComponent*> DelayedComponentsToUnregister;
+	TSet<UPCGComponent*> DelayedComponentsToRegister;
+	
+	// Tracking actors
+	/** Keep a mapping between tracked keys and the execution sources that track them, and the tracking needs to be culled.*/
+	TMap<FPCGSelectionKey, TSet<IPCGGraphExecutionSource*>> CulledTrackedKeysToExecutionSourcesMap;
+
+	/** Same mapping but for always tracked keys */
+	TMap<FPCGSelectionKey, TSet<IPCGGraphExecutionSource*>> AlwaysTrackedKeysToExecutionSourcesMap;
+
+	mutable PCG::FSharedLock TrackedComponentsLock;
+
+	// Keep track of actors that aren't yet ready (or if the subsystem is not yet ready), whether we should dirty them in next tick.
+	TMap<TObjectKey<AActor>, TTuple<TObjectKey<UObject>, EChangedObjectReason>> DelayedChangedActors;
+	
+	/** Transient map of actors and their previous data, it's set in the pre object change to be able to track changes (such as tags or positions) */
+	TMap<TObjectKey<AActor>, FActorPreviousData> ActorToPreviousDataMap;
+
+	/** Transient map that keep track of all execution source that depend on other execution sources component currently generating, to trigger the refresh only once, when all are done. */
+	TMap<FObjectKey, TArray<FObjectKey>> ExecutionSourcesToDependencyMap;
+
+	// Part for the delayed landscape change update
+	double LastLandscapeDirtyTime = -1.0;
+	TArray<TObjectKey<ALandscapeProxy>> DelayedModifiedLandscapes;
+
+	// Keep track of all dirtied landscapes and force a refresh on them when exiting edit mode.
+	TArray<TObjectKey<ALandscapeProxy>> DirtiedLandscapes;
+
+	// Temp boolean to indicate we are currently exiting the landscape edit mode
+	bool bIsCurrentlyExitingLandscapeEditMode = false;
+
+	// Time keeper for cleaning up cached previous actor data
+	double LastPreviousActorDataCleanup = -1.0;
+
+	// Set of existing PCG Partition Actors (for World Partition worlds)
+	TMap<FPCGGridCellDescriptor, FGuid> PartitionActorRecords;
+	// Previously generated PCG Partition Actors to ignore (will prevent them from getting registered and mark them for deletion)
+	TSet<FGuid> InvalidPartitionActors;
+
+	// If dispatch needs to run the local comopnents one after the other through dependency chain for determinism
+	bool bChainedDispatchToLocalComponents = false;
+#endif // WITH_EDITOR
+};
\ No newline at end of file
diff --git a/Engine/Plugins/PCG/Source/PCG/Private/PCGWorldActor.cpp b/Engine/Plugins/PCG/Source/PCG/Private/PCGWorldActor.cpp
--- a/Engine/Plugins/PCG/Source/PCG/Private/PCGWorldActor.cpp
+++ b/Engine/Plugins/PCG/Source/PCG/Private/PCGWorldActor.cpp
@@ -14,8 +14,9 @@
 #include "HAL/IConsoleManager.h"
 
 #if WITH_EDITOR
+#include "PCGTrackingManager.h"
+
 #include "UObject/UObjectHash.h"
-#include "PCGActorAndComponentMapping.h"
 #include "WorldPartition/WorldPartition.h"
 #include "WorldPartition/WorldPartitionHelpers.h"
 #include "WorldPartition/ActorPartition/PartitionActorDesc.h"
@@ -304,11 +305,14 @@ void APCGWorldActor::OnPartitionGridSizeChanged()
 
 	// And finally, regenerate all components that are partitioned (registered to the PCGSubsystem)
 	// to let them recreate the needed PCG Partition Actors.
-	for (UPCGComponent* PCGComponent : PCGSubsystem->GetAllRegisteredPartitionedComponents())
+	for (IPCGGraphExecutionSource* ExecutionSource : PCGSubsystem->GetAllRegisteredPartitionedExecutionSources())
 	{
-		check(PCGComponent);
-		PCGComponent->DirtyGenerated();
-		PCGComponent->Refresh(EPCGChangeType::GenerationGrid, /*bCancelExistingRefresh=*/true);
+		// @todo_pcg: support execution source
+		if (UPCGComponent* PCGComponent = Cast<UPCGComponent>(ExecutionSource))
+		{
+			PCGComponent->DirtyGenerated();
+			PCGComponent->Refresh(EPCGChangeType::GenerationGrid, /*bCancelExistingRefresh=*/true);
+		}
 	}
 }
 
diff --git a/Engine/Plugins/PCG/Source/PCG/Private/RuntimeGen/PCGRuntimeGenScheduler.cpp b/Engine/Plugins/PCG/Source/PCG/Private/RuntimeGen/PCGRuntimeGenScheduler.cpp
--- a/Engine/Plugins/PCG/Source/PCG/Private/RuntimeGen/PCGRuntimeGenScheduler.cpp
+++ b/Engine/Plugins/PCG/Source/PCG/Private/RuntimeGen/PCGRuntimeGenScheduler.cpp
@@ -3,12 +3,12 @@
 #include "RuntimeGen/PCGRuntimeGenScheduler.h"
 #include "RuntimeGen/SchedulingPolicies/PCGSchedulingPolicyBase.h"
 
-#include "PCGActorAndComponentMapping.h"
 #include "PCGCommon.h"
 #include "PCGComponent.h"
 #include "PCGGraph.h"
 #include "PCGModule.h"
 #include "PCGSubsystem.h"
+#include "PCGTrackingManager.h"
 #include "PCGWorldActor.h"
 #include "Grid/PCGGridDescriptor.h"
 #include "Grid/PCGPartitionActor.h"
@@ -189,15 +189,15 @@ FPCGGridDescriptor FPCGRuntimeGenScheduler::FGridGenerationKey::GetGridDescripto
 	return FPCGGridDescriptor().SetGridSize(GetGridSize()).SetIsRuntime(true).SetIs2DGrid(Use2DGrid());
 }
 
-FPCGRuntimeGenScheduler::FPCGRuntimeGenScheduler(UWorld* InWorld, FPCGActorAndComponentMapping* InActorAndComponentMapping)
+FPCGRuntimeGenScheduler::FPCGRuntimeGenScheduler(UWorld* InWorld, FPCGTrackingManager* InTrackingManager)
 {
-	check(InWorld && InActorAndComponentMapping);
+	check(InWorld && InTrackingManager);
 
 	World = InWorld;
 	Subsystem = UPCGSubsystem::GetInstance(World);
 	WorldPartitionSubsystem = World ? UWorld::GetSubsystem<UWorldPartitionSubsystem>(World) : nullptr;
 
-	ActorAndComponentMapping = InActorAndComponentMapping;
+	TrackingManager = InTrackingManager;
 	GenSourceManager = new FPCGGenSourceManager(InWorld);
 	bPoolingWasEnabledLastFrame = PCGRuntimeGenSchedulerHelpers::CVarRuntimeGenerationEnablePooling.GetValueOnAnyThread();
 	BasePoolSizeLastFrame = PCGRuntimeGenSchedulerHelpers::CVarRuntimeGenerationBasePoolSize.GetValueOnAnyThread();
@@ -268,7 +268,7 @@ void FPCGRuntimeGenScheduler::Tick(APCGWorldActor* InPCGWorldActor, double InEnd
 
 	GenSources.Empty(GenSources.Num());
 
-	if (bAnyRuntimeGenComponentsExist)
+	if (bAnyRuntimeGenExecutionSourcesExist)
 	{
 		GenSourceManager->Tick();
 		GenSources = GenSourceManager->GetAllGenSources(InPCGWorldActor);
@@ -315,8 +315,8 @@ void FPCGRuntimeGenScheduler::Tick(APCGWorldActor* InPCGWorldActor, double InEnd
 		FTickQueueComponentsForGenerationInputs Inputs;
 		Inputs.GenSources = &GenSources;
 		Inputs.PCGWorldActor = InPCGWorldActor;
-		Inputs.AllPartitionedComponents = ActorAndComponentMapping->GetAllRegisteredPartitionedComponents();
-		Inputs.AllNonPartitionedComponents = ActorAndComponentMapping->GetAllRegisteredNonPartitionedComponents();
+		Inputs.AllPartitionedExecutionSources = TrackingManager->GetAllRegisteredPartitionedExecutionSources();
+		Inputs.AllNonPartitionedExecutionSources = TrackingManager->GetAllRegisteredNonPartitionedExecutionSources();
 		Inputs.GeneratedComponents = &GeneratedComponents;
 
 		ChangeDetector.PreTick();
@@ -349,7 +349,7 @@ void FPCGRuntimeGenScheduler::Tick(APCGWorldActor* InPCGWorldActor, double InEnd
 
 bool FPCGRuntimeGenScheduler::ShouldTick()
 {
-	check(World && ActorAndComponentMapping);
+	check(World && TrackingManager);
 
 	if (!PCGRuntimeGenSchedulerHelpers::CVarRuntimeGenerationEnable.GetValueOnAnyThread())
 	{
@@ -377,30 +377,30 @@ bool FPCGRuntimeGenScheduler::ShouldTick()
 	}
 #endif
 
-	if (bAnyRuntimeGenComponentsExistDirty)
+	if (bAnyRuntimeGenExecutionSourcesExistDirty)
 	{
-		const bool bDidAnyRuntimeGenComponentsExist = bAnyRuntimeGenComponentsExist;
-		bAnyRuntimeGenComponentsExist = ActorAndComponentMapping->AnyRuntimeGenComponentsExist();
-		bAnyRuntimeGenComponentsExistDirty = false;
+		const bool bDidAnyRuntimeGenComponentsExist = bAnyRuntimeGenExecutionSourcesExist;
+		bAnyRuntimeGenExecutionSourcesExist = TrackingManager->AnyRuntimeGenExecutionSourcesExist();
+		bAnyRuntimeGenExecutionSourcesExistDirty = false;
 
 		if (PCGRuntimeGenSchedulerHelpers::CVarRuntimeGenerationEnableDebugging.GetValueOnAnyThread())
 		{
-			if (bDidAnyRuntimeGenComponentsExist != bAnyRuntimeGenComponentsExist)
+			if (bDidAnyRuntimeGenComponentsExist != bAnyRuntimeGenExecutionSourcesExist)
 			{
-				if (bAnyRuntimeGenComponentsExist)
+				if (bAnyRuntimeGenExecutionSourcesExist)
 				{
-					UE_LOG(LogPCG, Warning, TEXT("[RUNTIMEGEN] THERE ARE NOW RUNTIME COMPONENTS IN THE LEVEL. SCHEDULER WILL BEGIN TICKING."));
+					UE_LOG(LogPCG, Warning, TEXT("[RUNTIMEGEN] THERE ARE NOW RUNTIME EXECUTION SOURCES IN THE LEVEL. SCHEDULER WILL BEGIN TICKING."));
 				}
 				else
 				{
-					UE_LOG(LogPCG, Warning, TEXT("[RUNTIMEGEN] THERE ARE NO MORE RUNTIME COMPONENTS. SCHEDULER WILL ONLY TICK TO CLEANUP."));
+					UE_LOG(LogPCG, Warning, TEXT("[RUNTIMEGEN] THERE ARE NO MORE RUNTIME EXECUTION SOURCES. SCHEDULER WILL ONLY TICK TO CLEANUP."));
 				}
 			}
 		}
 	}
 
 	// We can stop ticking if there are no runtime gen components alive and there are no generated components that need cleaning up.
-	if (!bAnyRuntimeGenComponentsExist && GeneratedComponents.IsEmpty() && GeneratedComponentsToRemove.IsEmpty())
+	if (!bAnyRuntimeGenExecutionSourcesExist && GeneratedComponents.IsEmpty() && GeneratedComponentsToRemove.IsEmpty())
 	{
 		return false;
 	}
@@ -517,10 +517,13 @@ void FPCGRuntimeGenScheduler::TickQueueComponentsForGeneration(
 #endif
 
 	// Collect local components from all partitioned components.
-	for (UPCGComponent* OriginalComponent : Inputs.AllPartitionedComponents)
+	for (IPCGGraphExecutionSource* OriginalExecutionSource : Inputs.AllPartitionedExecutionSources)
 	{
 		TRACE_CPUPROFILER_EVENT_SCOPE(FPCGRuntimeGenScheduler::CollectLocalComponents);
 
+		// @todo_pcg: support execution source
+		UPCGComponent* OriginalComponent = Cast<UPCGComponent>(OriginalExecutionSource);
+
 		if (!ensure(OriginalComponent) || !OriginalComponent->GetGraph() || !OriginalComponent->bActivated)
 		{
 			continue;
@@ -788,10 +791,13 @@ void FPCGRuntimeGenScheduler::TickQueueComponentsForGeneration(
 	}
 
 	// Collect all non-partitioned components.
-	for (UPCGComponent* OriginalComponent : Inputs.AllNonPartitionedComponents)
+	for (IPCGGraphExecutionSource* OriginalExecutionSource : Inputs.AllNonPartitionedExecutionSources)
 	{
 		TRACE_CPUPROFILER_EVENT_SCOPE(FPCGRuntimeGenScheduler::CollectNonPartitionedComponents);
 
+		// @todo_pcg: support execution source
+		UPCGComponent* OriginalComponent = Cast<UPCGComponent>(OriginalExecutionSource);
+
 		if (!ensure(OriginalComponent) || !OriginalComponent->GetGraph() || !OriginalComponent->bActivated)
 		{
 			continue;
@@ -863,7 +869,7 @@ void FPCGRuntimeGenScheduler::TickCleanup(const TSet<IPCGGenSourceBase*>& InGenS
 {
 	TRACE_CPUPROFILER_EVENT_SCOPE(FPCGRuntimeGenScheduler::TickCleanup);
 
-	check(ActorAndComponentMapping && InPCGWorldActor);
+	check(TrackingManager && InPCGWorldActor);
 
 	auto CheckIfAllGenSourcesWantToCleanup = [&InGenSources](const UPCGSchedulingPolicyBase* Policy, const FPCGGridDescriptor& GridDescriptor, const FBox& GridBounds, const double& CleanupRadiusSquared)
 	{
@@ -923,7 +929,7 @@ void FPCGRuntimeGenScheduler::TickCleanup(const TSet<IPCGGenSourceBase*>& InGenS
 		TEXT("FPCGRuntimeGenScheduler::TickCleanup"),
 		NumGeneratedComponents,
 		/*MinBatchSize*/6, // Optimal value found in testing.
-		[&GeneratedComponentsArray, &ComponentsToClean, &InvalidKeys, &ComponentsToCleanCounter, &InvalidKeysCounter, ActorAndComponentMapping=ActorAndComponentMapping, &CheckIfAllGenSourcesWantToCleanup](int32 Index)
+		[&GeneratedComponentsArray, &ComponentsToClean, &InvalidKeys, &ComponentsToCleanCounter, &InvalidKeysCounter, TrackingManager=TrackingManager, &CheckIfAllGenSourcesWantToCleanup](int32 Index)
 		{
 			TRACE_CPUPROFILER_EVENT_SCOPE(SelectComponentForCleanup);
 			LLM_SCOPE_BYTAG(PCG);
@@ -975,7 +981,7 @@ void FPCGRuntimeGenScheduler::TickCleanup(const TSet<IPCGGenSourceBase*>& InGenS
 				IPCGGraphExecutionSource* LocalComponent = GenerationKey.GetCachedLocalComponent();
 				if (!LocalComponent)
 				{
-					LocalComponent = ActorAndComponentMapping->GetLocalComponent(GridDescriptor, GridCoords, OriginalComponent);
+					LocalComponent = TrackingManager->GetLocalComponent(GridDescriptor, GridCoords, OriginalComponent);
 				}
 
 				APCGPartitionActor* PartitionActor = LocalComponent ? Cast<APCGPartitionActor>(CastChecked<UPCGComponent>(LocalComponent)->GetOwner()) : nullptr;
@@ -1024,7 +1030,7 @@ void FPCGRuntimeGenScheduler::TickScheduleGeneration(TMap<FGridGenerationKey, do
 {
 	TRACE_CPUPROFILER_EVENT_SCOPE(FPCGRuntimeGenScheduler::TickScheduleGeneration);
 
-	check(Subsystem && ActorAndComponentMapping);
+	check(Subsystem && TrackingManager);
 
 	// Count number of currently generating components
 	int NumGenerating = 0;
@@ -1111,7 +1117,7 @@ void FPCGRuntimeGenScheduler::TickScheduleGeneration(TMap<FGridGenerationKey, do
 		else
 		{
 			// Grab local component and PA if they exist already.
-			IPCGGraphExecutionSource* LocalComponent = ActorAndComponentMapping->GetLocalComponent(GridDescriptor, GridCoords, OriginalComponent);
+			IPCGGraphExecutionSource* LocalComponent = TrackingManager->GetLocalComponent(GridDescriptor, GridCoords, OriginalComponent);
 			TObjectPtr<APCGPartitionActor> PartitionActor = LocalComponent ? Cast<APCGPartitionActor>(CastChecked<UPCGComponent>(LocalComponent)->GetOwner()) : nullptr;
 
 			if (!LocalComponent || !ensure(PartitionActor))
@@ -1155,12 +1161,12 @@ void FPCGRuntimeGenScheduler::TickScheduleGeneration(TMap<FGridGenerationKey, do
 
 				// Update component mapping for this PA (add local component).
 				{
-					PCG::TUniqueScopeLock WriteLock(ActorAndComponentMapping->ComponentToPartitionActorsMapLock);
-					TSet<TObjectPtr<APCGPartitionActor>>* PartitionActorsPtr = ActorAndComponentMapping->ComponentToPartitionActorsMap.Find(OriginalComponent);
+					PCG::TUniqueScopeLock WriteLock(TrackingManager->ComponentToPartitionActorsMapLock);
+					TSet<TObjectPtr<APCGPartitionActor>>* PartitionActorsPtr = TrackingManager->ComponentToPartitionActorsMap.Find(OriginalComponent);
 
 					if (!PartitionActorsPtr)
 					{
-						PartitionActorsPtr = &ActorAndComponentMapping->ComponentToPartitionActorsMap.Emplace(OriginalComponent);
+						PartitionActorsPtr = &TrackingManager->ComponentToPartitionActorsMap.Emplace(OriginalComponent);
 					}
 
 					// Log this original component before setting up the PA, so that we early out from RefreshComponent if it gets called
@@ -1229,7 +1235,7 @@ void FPCGRuntimeGenScheduler::TickRequestVirtualTexturePriming(const TSet<IPCGGe
 {
 	TRACE_CPUPROFILER_EVENT_SCOPE(FPCGRuntimeGenScheduler::TickRequestVirtualTexturePriming);
 
-	check(ActorAndComponentMapping);
+	check(TrackingManager);
 
 	TMap<TSoftObjectPtr<URuntimeVirtualTexture>, TArray<URuntimeVirtualTextureComponent*>> VirtualTextureToComponents;
 
@@ -1412,7 +1418,7 @@ void FPCGRuntimeGenScheduler::OnOriginalComponentRegistered(UPCGComponent* InOri
 	}
 
 	// When an original/non-partitioned component is registered, we need to dirty the state.
-	bAnyRuntimeGenComponentsExistDirty = true;
+	bAnyRuntimeGenExecutionSourcesExistDirty = true;
 
 #if WITH_EDITOR
 	RegisterGraphParameterChangedEvent(InOriginalComponent);
@@ -1428,10 +1434,10 @@ void FPCGRuntimeGenScheduler::OnOriginalComponentUnregistered(UPCGComponent* InO
 		return;
 	}
 
-	check(ActorAndComponentMapping);
+	check(TrackingManager);
 
 	// When an original/non-partitioned component is unregistered, we need to dirty the state.
-	bAnyRuntimeGenComponentsExistDirty = true;
+	bAnyRuntimeGenExecutionSourcesExistDirty = true;
 
 #if WITH_EDITOR
 	UnregisterGraphParameterChangedEvent(InOriginalComponent);
@@ -1463,7 +1469,7 @@ void FPCGRuntimeGenScheduler::OnOriginalComponentUnregistered(UPCGComponent* InO
 		const bool bIsOriginalComponent = GridDescriptor.GetGridSize() == PCGHiGenGrid::UnboundedGridSize();
 
 		// Get the generated component for this key (might be a local component).
-		IPCGGraphExecutionSource* ComponentToCleanup = bIsOriginalComponent ? InOriginalComponent : ActorAndComponentMapping->GetLocalComponent(GridDescriptor, GenerationKey.GetGridCoords(), InOriginalComponent);
+		IPCGGraphExecutionSource* ComponentToCleanup = bIsOriginalComponent ? InOriginalComponent : TrackingManager->GetLocalComponent(GridDescriptor, GenerationKey.GetGridCoords(), InOriginalComponent);
 
 		// It is possible for a PartitionActor's LocalComponent to have been cleaned up by the APCGPartitionActor::EndPlay call depending on the order in which actors get called
 		if (ComponentToCleanup)
@@ -1543,7 +1549,7 @@ void FPCGRuntimeGenScheduler::CleanupRemainingComponents(UPCGComponent* InOrigin
 	// If the PA is not far enough to be cleaned up but no longer is considered within the Gen Source Radius for Generation.
 	if (InOriginalComponent && InOriginalComponent->IsManagedByRuntimeGenSystem())
 	{
-		TSet<TObjectPtr<APCGPartitionActor>> PartitionActors = ActorAndComponentMapping->GetPCGComponentPartitionActorMappings(InOriginalComponent);
+		TSet<TObjectPtr<APCGPartitionActor>> PartitionActors = TrackingManager->GetPCGComponentPartitionActorMappings(InOriginalComponent);
 		for (TObjectPtr<APCGPartitionActor> PartitionActor : PartitionActors)
 		{
 			if (IPCGGraphExecutionSource* ComponentToCleanup = PartitionActor->GetLocalComponent(InOriginalComponent))
@@ -1558,7 +1564,7 @@ void FPCGRuntimeGenScheduler::CleanupLocalComponents(const APCGWorldActor* InPCG
 {
 	TRACE_CPUPROFILER_EVENT_SCOPE(FPCGRuntimeGenScheduler::CleanupLocalComponents);
 
-	check(ActorAndComponentMapping && InPCGWorldActor);
+	check(TrackingManager && InPCGWorldActor);
 
 	// Generated Entry Key, Local/Generated Component
 	using PCGComponentToClean = TTuple<FGridGenerationKey, IPCGGraphExecutionSource*>;
@@ -1588,7 +1594,7 @@ void FPCGRuntimeGenScheduler::CleanupLocalComponents(const APCGWorldActor* InPCG
 		// Only operate on LocalComponents.
 		if (Grid != EPCGHiGenGrid::Unbounded)
 		{
-			IPCGGraphExecutionSource* LocalComponent = ActorAndComponentMapping->GetLocalComponent(GridDescriptor, GridCoords, OriginalComponent);
+			IPCGGraphExecutionSource* LocalComponent = TrackingManager->GetLocalComponent(GridDescriptor, GridCoords, OriginalComponent);
 			ComponentsToClean.Add({ GenerationKey, LocalComponent });
 		}
 	}
@@ -1644,8 +1650,8 @@ void FPCGRuntimeGenScheduler::CleanupLocalComponent(APCGPartitionActor* Partitio
 
 		// Remove component mapping.
 		{
-			PCG::TUniqueScopeLock WriteLock(ActorAndComponentMapping->ComponentToPartitionActorsMapLock);
-			TSet<TObjectPtr<APCGPartitionActor>>* PartitionActorsPtr = ActorAndComponentMapping->ComponentToPartitionActorsMap.Find(OriginalComponent);
+			PCG::TUniqueScopeLock WriteLock(TrackingManager->ComponentToPartitionActorsMapLock);
+			TSet<TObjectPtr<APCGPartitionActor>>* PartitionActorsPtr = TrackingManager->ComponentToPartitionActorsMap.Find(OriginalComponent);
 
 			if (PartitionActorsPtr)
 			{
@@ -1653,7 +1659,7 @@ void FPCGRuntimeGenScheduler::CleanupLocalComponent(APCGPartitionActor* Partitio
 
 				if (PartitionActorsPtr->IsEmpty())
 				{
-					ActorAndComponentMapping->ComponentToPartitionActorsMap.Remove(OriginalComponent);
+					TrackingManager->ComponentToPartitionActorsMap.Remove(OriginalComponent);
 				}
 			}
 		}
@@ -1696,7 +1702,7 @@ void FPCGRuntimeGenScheduler::CleanupComponent(const FGridGenerationKey& Generat
 {
 	TRACE_CPUPROFILER_EVENT_SCOPE(FPCGRuntimeGenScheduler::CleanupComponent);
 
-	check(ActorAndComponentMapping);
+	check(TrackingManager);
 
 	const FPCGGridDescriptor GridDescriptor = GenerationKey.GetGridDescriptor();
 	const EPCGHiGenGrid Grid = GridDescriptor.GetHiGenGrid();
@@ -1709,8 +1715,8 @@ void FPCGRuntimeGenScheduler::CleanupComponent(const FGridGenerationKey& Generat
 	{
 		UE_LOG(LogPCG, Warning, TEXT("Runtime generated component could not be recovered on grid %d at (%d, %d, %d). It has been lost or destroyed."), GridDescriptor.GetGridSize(), GridCoords.X, GridCoords.Y, GridCoords.Z);
 
-		// If the GeneratedComponent has been lost for some reason, get the PA directly from the ActorAndComponentMapping.
-		PartitionActor = ActorAndComponentMapping->GetPartitionActor(GridDescriptor, GridCoords);
+		// If the GeneratedComponent has been lost for some reason, get the PA directly from the TrackingManager.
+		PartitionActor = TrackingManager->GetPartitionActor(GridDescriptor, GridCoords);
 	}
 	else // If the GeneratedComponent does exist, we can clean it up.
 	{
@@ -1746,7 +1752,7 @@ void FPCGRuntimeGenScheduler::CleanupDelayedRefreshComponents()
 {
 	TRACE_CPUPROFILER_EVENT_SCOPE(FPCGRuntimeGenScheduler::CleanupDelayedRefreshComponents);
 
-	check(ActorAndComponentMapping);
+	check(TrackingManager);
 
 	// Check that each refreshed local component is still intersecting its original component.
 	// If it is not, it would be leaked instead of refreshed, so we should force a full cleanup.
@@ -1773,7 +1779,7 @@ void FPCGRuntimeGenScheduler::CleanupDelayedRefreshComponents()
 			continue;
 		}
 
-		IPCGGraphExecutionSource* LocalComponent = OriginalComponent ? ActorAndComponentMapping->GetLocalComponent(GridDescriptor, GridCoords, OriginalComponent) : nullptr;
+		IPCGGraphExecutionSource* LocalComponent = OriginalComponent ? TrackingManager->GetLocalComponent(GridDescriptor, GridCoords, OriginalComponent) : nullptr;
 		APCGPartitionActor* PartitionActor = LocalComponent ? Cast<APCGPartitionActor>(CastChecked<UPCGComponent>(LocalComponent)->GetOwner()) : nullptr;
 
 		if (LocalComponent && PartitionActor)
@@ -1864,7 +1870,7 @@ void FPCGRuntimeGenScheduler::RefreshComponent(IPCGGraphExecutionSource* InCompo
 					const FPCGGridDescriptor GridDescriptor = Key.GetGridDescriptor();
 					
 					// TODO - clean up local immediate will have a flag in the future to clean up the local components on its own, so this call to CleanupLocalImmediate will not be required
-					IPCGGraphExecutionSource* LocalComponent = ActorAndComponentMapping->GetLocalComponent(
+					IPCGGraphExecutionSource* LocalComponent = TrackingManager->GetLocalComponent(
 						GridDescriptor,
 						Key.GetGridCoords(),
 						OriginalComponent);
@@ -1969,7 +1975,7 @@ void FPCGRuntimeGenScheduler::RefreshComponent(IPCGGraphExecutionSource* InCompo
 				{
 					const FIntVector GridCoords = GenerationKey.GetGridCoords();
 
-					if (IPCGGraphExecutionSource* LocalComponent = ActorAndComponentMapping->GetLocalComponent(GridDescriptor, GridCoords, OriginalComponent))
+					if (IPCGGraphExecutionSource* LocalComponent = TrackingManager->GetLocalComponent(GridDescriptor, GridCoords, OriginalComponent))
 					{
 						RefreshLocalComponent(LocalComponent);
 					}
@@ -1986,7 +1992,7 @@ void FPCGRuntimeGenScheduler::RefreshComponent(IPCGGraphExecutionSource* InCompo
 	if (!bIsLocalComponent)
 	{
 		// When an original/non-partitioned component is refreshed, we need to dirty the state.
-		bAnyRuntimeGenComponentsExistDirty = true;
+		bAnyRuntimeGenExecutionSourcesExistDirty = true;
 	}
 }
 
@@ -1994,7 +2000,7 @@ APCGPartitionActor* FPCGRuntimeGenScheduler::GetPartitionActorFromPool(const FPC
 {
 	TRACE_CPUPROFILER_EVENT_SCOPE(FPCGRuntimeGenScheduler::GetPartitionActorFromPool);
 
-	check(ActorAndComponentMapping);
+	check(TrackingManager);
 
 	if (!World)
 	{
@@ -2003,7 +2009,7 @@ APCGPartitionActor* FPCGRuntimeGenScheduler::GetPartitionActorFromPool(const FPC
 	}
 
 	// Attempt to find an existing RuntimeGen PA.
-	if (APCGPartitionActor* ExistingActor = ActorAndComponentMapping->GetPartitionActor(GridDescriptor, GridCoords))
+	if (APCGPartitionActor* ExistingActor = TrackingManager->GetPartitionActor(GridDescriptor, GridCoords))
 	{
 		return ExistingActor;
 	}
@@ -2175,29 +2181,37 @@ void FPCGRuntimeGenScheduler::OnGraphParametersChanged(UPCGGraphInterface* InGra
 {
 	TArray<UPCGComponent*> OriginalComponents;
 
-	TSet<UPCGComponent*> PartitionedComponents = ActorAndComponentMapping->GetAllRegisteredPartitionedComponents();
-	TSet<UPCGComponent*> NonPartitionedComponents = ActorAndComponentMapping->GetAllRegisteredNonPartitionedComponents();
+	TSet<IPCGGraphExecutionSource*> PartitionedExecutionSources = TrackingManager->GetAllRegisteredPartitionedExecutionSources();
+	TSet<IPCGGraphExecutionSource*> NonPartitionedExecutionSources = TrackingManager->GetAllRegisteredNonPartitionedExecutionSources();
 
-	OriginalComponents.Reserve(PartitionedComponents.Num() + NonPartitionedComponents.Num());
+	OriginalComponents.Reserve(PartitionedExecutionSources.Num() + NonPartitionedExecutionSources.Num());
 
 	auto ShouldCachePrimingInfosForComponent = [InGraph](const UPCGComponent* OriginalComponent)
 	{
-		return OriginalComponent && OriginalComponent->bActivated && OriginalComponent->IsManagedByRuntimeGenSystem() && OriginalComponent->GetGraph() == InGraph;
+		return OriginalComponent->bActivated && OriginalComponent->IsManagedByRuntimeGenSystem() && OriginalComponent->GetGraph() == InGraph;
 	};
 
-	for (UPCGComponent* OriginalComponent : PartitionedComponents)
+	// @todo_pcg: support execution source
+	for (IPCGGraphExecutionSource* OriginalSource : PartitionedExecutionSources)
 	{
-		if (ShouldCachePrimingInfosForComponent(OriginalComponent))
+		if (UPCGComponent* OriginalComponent = Cast<UPCGComponent>(OriginalSource))
 		{
-			OriginalComponents.Add(MoveTemp(OriginalComponent));
+			if (ShouldCachePrimingInfosForComponent(OriginalComponent))
+			{
+				OriginalComponents.Add(OriginalComponent);
+			}
 		}
 	}
 
-	for (UPCGComponent* OriginalComponent : NonPartitionedComponents)
+	// @todo_pcg: support execution source
+	for (IPCGGraphExecutionSource* OriginalSource : NonPartitionedExecutionSources)
 	{
-		if (ShouldCachePrimingInfosForComponent(OriginalComponent))
+		if (UPCGComponent* OriginalComponent = Cast<UPCGComponent>(OriginalSource))
 		{
-			OriginalComponents.Add(MoveTemp(OriginalComponent));
+			if (ShouldCachePrimingInfosForComponent(OriginalComponent))
+			{
+				OriginalComponents.Add(OriginalComponent);
+			}
 		}
 	}
 
diff --git a/Engine/Plugins/PCG/Source/PCG/Private/Subsystems/IPCGBaseSubsystem.cpp b/Engine/Plugins/PCG/Source/PCG/Private/Subsystems/IPCGBaseSubsystem.cpp
--- a/Engine/Plugins/PCG/Source/PCG/Private/Subsystems/IPCGBaseSubsystem.cpp
+++ b/Engine/Plugins/PCG/Source/PCG/Private/Subsystems/IPCGBaseSubsystem.cpp
@@ -2,10 +2,10 @@
 
 #include "Subsystems/IPCGBaseSubsystem.h"
 
-#include "PCGActorAndComponentMapping.h"
 #include "PCGComponent.h"
 #include "PCGGraph.h"
 #include "PCGModule.h"
+#include "PCGTrackingManager.h"
 #include "Graph/PCGGraphCache.h"
 #include "Graph/PCGGraphCompiler.h"
 #include "Graph/PCGGraphExecutor.h"
diff --git a/Engine/Plugins/PCG/Source/PCG/Private/Subsystems/PCGSubsystem.cpp b/Engine/Plugins/PCG/Source/PCG/Private/Subsystems/PCGSubsystem.cpp
--- a/Engine/Plugins/PCG/Source/PCG/Private/Subsystems/PCGSubsystem.cpp
+++ b/Engine/Plugins/PCG/Source/PCG/Private/Subsystems/PCGSubsystem.cpp
@@ -2,10 +2,10 @@
 
 #include "PCGSubsystem.h"
 
-#include "PCGActorAndComponentMapping.h"
 #include "PCGComponent.h"
 #include "PCGGraph.h"
 #include "PCGModule.h"
+#include "PCGTrackingManager.h"
 #include "PCGWorldActor.h"
 #include "Graph/PCGGraphCache.h"
 #include "Graph/PCGGraphCompiler.h"
@@ -153,7 +153,7 @@ public:
 UPCGSubsystem::UPCGSubsystem()
 	: UTickableWorldSubsystem()
 	, IPCGBaseSubsystem()
-	, ActorAndComponentMapping(new FPCGActorAndComponentMapping(this)) // Can't use MakeUnique because the ctor of FPCGActorAndComponentMapping is private
+	, TrackingManager(new FPCGTrackingManager(this)) // Can't use MakeUnique because the ctor of FPCGTrackingManager is private
 {
 }
 
@@ -207,7 +207,7 @@ void UPCGSubsystem::Deinitialize()
 	PCGWorldActor = nullptr;
 	bHasTickedOnce = false;
 
-	ActorAndComponentMapping->Deinitialize();
+	TrackingManager->Deinitialize();
 
 #if WITH_EDITOR
 	FCoreUObjectDelegates::OnObjectsReplaced.RemoveAll(this);
@@ -250,13 +250,13 @@ void UPCGSubsystem::PostInitialize()
 		}
 	}
 
-	ActorAndComponentMapping->Initialize();
+	TrackingManager->Initialize();
 
 	IPCGBaseSubsystem::InitializeBaseSubsystem();
 
 	// Initialize runtime generation scheduler
 	check(!RuntimeGenScheduler);
-	RuntimeGenScheduler = new FPCGRuntimeGenScheduler(GetWorld(), ActorAndComponentMapping.Get());
+	RuntimeGenScheduler = new FPCGRuntimeGenScheduler(GetWorld(), TrackingManager.Get());
 
 #if WITH_EDITOR
 	FCoreUObjectDelegates::OnObjectsReplaced.AddUObject(this, &UPCGSubsystem::OnObjectsReplaced);
@@ -388,7 +388,7 @@ void UPCGSubsystem::Tick(float DeltaSeconds)
 	// Base subsystem can modify end time
 	const double EndTime = IPCGBaseSubsystem::Tick();
 	
-	ActorAndComponentMapping->Tick();
+	TrackingManager->Tick();
 
 	if (PCGWorldActor)
 	{
@@ -453,17 +453,21 @@ int32 UPCGSubsystem::GetPCGQualityLevel()
 void UPCGSubsystem::OnPCGQualityLevelChanged()
 {
 	// Trigger deep refresh of all runtime generated components.
-	ActorAndComponentMapping->ForAllOriginalComponents([this](UPCGComponent* InComponent)
+	TrackingManager->ForAllOriginalExecutionSources([this](IPCGGraphExecutionSource* InExecutionSource)
 	{
-		if (InComponent && InComponent->IsManagedByRuntimeGenSystem())
+		if (InExecutionSource && InExecutionSource->GetExecutionState().IsManagedByRuntimeGenSystem())
 		{
 #if WITH_EDITOR
-			if (UPCGGraph* Graph = InComponent->GetGraph())
+			if (UPCGGraph* Graph = InExecutionSource->GetExecutionState().GetGraph())
 			{
 				Graph->OnPCGQualityLevelChanged();
 			}
 #else
-			RefreshRuntimeGenComponent(InComponent, EPCGChangeType::GenerationGrid);
+			// @todo_pcg: support execution source
+			if (UPCGComponent* PCGComponent = Cast<UPCGComponent>(InExecutionSource))
+			{
+				RefreshRuntimeGenComponent(PCGComponent, EPCGChangeType::GenerationGrid);
+			}
 #endif
 
 		}
@@ -614,25 +618,33 @@ void UPCGSubsystem::UnregisterPCGWorldActor(APCGWorldActor* InActor)
 	}
 }
 
-void UPCGSubsystem::OnOriginalComponentRegistered(UPCGComponent* InComponent)
+void UPCGSubsystem::OnOriginalExecutionSourceRegistered(IPCGGraphExecutionSource* InExecutionSource)
 {
+	// @todo_pcg: support execution source
 	if (RuntimeGenScheduler)
 	{
-		RuntimeGenScheduler->OnOriginalComponentRegistered(InComponent);
+		if (UPCGComponent* PCGComponent = Cast<UPCGComponent>(InExecutionSource))
+		{
+			RuntimeGenScheduler->OnOriginalComponentRegistered(PCGComponent);
+		}
 	}
 }
 
-void UPCGSubsystem::OnOriginalComponentUnregistered(UPCGComponent* InComponent)
+void UPCGSubsystem::OnOriginalExecutionSourceUnregistered(IPCGGraphExecutionSource* InExecutionSource)
 {
+	// @todo_pcg: support execution source
 	if (RuntimeGenScheduler)
 	{
-		RuntimeGenScheduler->OnOriginalComponentUnregistered(InComponent);
+		if (UPCGComponent* PCGComponent = Cast<UPCGComponent>(InExecutionSource))
+		{
+			RuntimeGenScheduler->OnOriginalComponentUnregistered(PCGComponent);
+		}
 	}
 
 #if WITH_EDITOR
 	if (IPCGEditorModule* PCGEditorModule = IPCGEditorModule::Get())
 	{
-		PCGEditorModule->ClearExecutionMetadata(InComponent);
+		PCGEditorModule->ClearExecutionMetadata(InExecutionSource);
 	}
 #endif
 }
@@ -667,7 +679,7 @@ FPCGTaskId UPCGSubsystem::ScheduleComponent(UPCGComponent* PCGComponent, EPCGHiG
 		}
 #endif // WITH_EDITOR
 
-		ActorAndComponentMapping->UpdateMappingPCGComponentPartitionActor(PCGComponent);
+		TrackingManager->UpdateMappingPCGComponentPartitionActor(PCGComponent);
 	}
 
 	// Execution dependencies require a task to finish executing before the dependent task.
@@ -766,7 +778,7 @@ FPCGTaskId UPCGSubsystem::ScheduleComponent(UPCGComponent* PCGComponent, EPCGHiG
 			return LocalComponent->GenerateInternal(bForce, LocalComponent->GetGenerationGrid(), EPCGComponentGenerationTrigger::GenerateOnDemand, *AllDependencies);
 		};
 
-		ExecutionDependencyTasks.Append(ActorAndComponentMapping->DispatchToRegisteredLocalComponents(PCGComponent, LocalGenerateTask));
+		ExecutionDependencyTasks.Append(TrackingManager->DispatchToRegisteredLocalComponents(PCGComponent, LocalGenerateTask));
 	}
 
 	if (!ExecutionDependencyTasks.IsEmpty() || !DataDependencyTasks.IsEmpty())
@@ -850,7 +862,7 @@ FPCGTaskId UPCGSubsystem::ScheduleCleanup(UPCGComponent* PCGComponent, bool bRem
 			return LocalComponent->CleanupLocal(bRemoveComponents, *AllDependencies);
 		};
 
-		AllTasks.Append(ActorAndComponentMapping->DispatchToRegisteredLocalComponents(PCGComponent, LocalCleanupTask));
+		AllTasks.Append(TrackingManager->DispatchToRegisteredLocalComponents(PCGComponent, LocalCleanupTask));
 	}
 
 	TWeakObjectPtr<UPCGComponent> ComponentPtr(PCGComponent);
@@ -870,7 +882,7 @@ FPCGTaskId UPCGSubsystem::ScheduleCleanup(UPCGComponent* PCGComponent, bool bRem
 			// the component mappings is counterproductive.
 			if (bRemoveComponents && Component->IsPartitioned())
 			{
-				ActorAndComponentMapping->DeleteMappingPCGComponentPartitionActor(Component);
+				TrackingManager->DeleteMappingPCGComponentPartitionActor(Component);
 			}
 		}
 
@@ -893,11 +905,11 @@ FPCGTaskId UPCGSubsystem::ScheduleCleanup(UPCGComponent* PCGComponent, bool bRem
 	return PostCleanupTaskId;
 }
 
-FPCGTaskId UPCGSubsystem::ScheduleGraph(UPCGComponent* SourceComponent, const TArray<FPCGTaskId>& Dependencies)
+FPCGTaskId UPCGSubsystem::ScheduleGraph(IPCGGraphExecutionSource* ExecutionSource, const TArray<FPCGTaskId>& Dependencies)
 {
-	if (SourceComponent)
+	if (ExecutionSource)
 	{
-		return GraphExecutor->Schedule(SourceComponent, Dependencies);
+		return GraphExecutor->Schedule(ExecutionSource, Dependencies);
 	}
 	else
 	{
@@ -907,14 +919,14 @@ FPCGTaskId UPCGSubsystem::ScheduleGraph(UPCGComponent* SourceComponent, const TA
 
 FPCGTaskId UPCGSubsystem::ScheduleGraph(
 	UPCGGraph* Graph,
-	UPCGComponent* SourceComponent,
+	IPCGGraphExecutionSource* ExecutionSource,
 	FPCGElementPtr PreGraphElement, 
 	FPCGElementPtr InputElement, 
 	const TArray<FPCGTaskId>& Dependencies, 
 	const FPCGStack* InFromStack, 
 	bool bAllowHierarchicalGeneration)
 {
-	return ScheduleGraph(FPCGScheduleGraphParams(Graph, SourceComponent, PreGraphElement, InputElement, Dependencies, InFromStack, bAllowHierarchicalGeneration));
+	return ScheduleGraph(FPCGScheduleGraphParams(Graph, ExecutionSource, PreGraphElement, InputElement, Dependencies, InFromStack, bAllowHierarchicalGeneration));
 }
 
 ETickableTickType UPCGSubsystem::GetTickableTickType() const
@@ -927,28 +939,28 @@ TStatId UPCGSubsystem::GetStatId() const
 	RETURN_QUICK_DECLARE_CYCLE_STAT(UPCGSubsystem, STATGROUP_Tickables);
 }
 
-FPCGTaskId UPCGSubsystem::ScheduleGeneric(TFunction<bool()> InOperation, UPCGComponent* SourceComponent, const TArray<FPCGTaskId>& TaskExecutionDependencies)
+FPCGTaskId UPCGSubsystem::ScheduleGeneric(TFunction<bool()> InOperation, IPCGGraphExecutionSource* ExecutionSource, const TArray<FPCGTaskId>& TaskExecutionDependencies)
 {
 	check(GraphExecutor);
-	return GraphExecutor->ScheduleGeneric(InOperation, SourceComponent, TaskExecutionDependencies);
+	return GraphExecutor->ScheduleGeneric(InOperation, ExecutionSource, TaskExecutionDependencies);
 }
 
-FPCGTaskId UPCGSubsystem::ScheduleGeneric(TFunction<bool()> InOperation, TFunction<void()> InAbortOperation, UPCGComponent* SourceComponent, const TArray<FPCGTaskId>& TaskExecutionDependencies)
+FPCGTaskId UPCGSubsystem::ScheduleGeneric(TFunction<bool()> InOperation, TFunction<void()> InAbortOperation, IPCGGraphExecutionSource* ExecutionSource, const TArray<FPCGTaskId>& TaskExecutionDependencies)
 {
 	check(GraphExecutor);
-	return GraphExecutor->ScheduleGeneric(InOperation, InAbortOperation, SourceComponent, TaskExecutionDependencies);
+	return GraphExecutor->ScheduleGeneric(InOperation, InAbortOperation, ExecutionSource, TaskExecutionDependencies);
 }
 
-FPCGTaskId UPCGSubsystem::ScheduleGenericWithContext(TFunction<bool(FPCGContext*)> InOperation, UPCGComponent* SourceComponent, const TArray<FPCGTaskId>& TaskExecutionDependencies, const TArray<FPCGTaskId>& TaskDataDependencies)
+FPCGTaskId UPCGSubsystem::ScheduleGenericWithContext(TFunction<bool(FPCGContext*)> InOperation, IPCGGraphExecutionSource* ExecutionSource, const TArray<FPCGTaskId>& TaskExecutionDependencies, const TArray<FPCGTaskId>& TaskDataDependencies)
 {
 	check(GraphExecutor);
-	return GraphExecutor->ScheduleGenericWithContext(InOperation, SourceComponent, TaskExecutionDependencies, TaskDataDependencies);
+	return GraphExecutor->ScheduleGenericWithContext(InOperation, ExecutionSource, TaskExecutionDependencies, TaskDataDependencies);
 }
 
-FPCGTaskId UPCGSubsystem::ScheduleGenericWithContext(TFunction<bool(FPCGContext*)> InOperation, TFunction<void(FPCGContext*)> InAbortOperation, UPCGComponent* SourceComponent, const TArray<FPCGTaskId>& TaskExecutionDependencies, const TArray<FPCGTaskId>& TaskDataDependencies)
+FPCGTaskId UPCGSubsystem::ScheduleGenericWithContext(TFunction<bool(FPCGContext*)> InOperation, TFunction<void(FPCGContext*)> InAbortOperation, IPCGGraphExecutionSource* ExecutionSource, const TArray<FPCGTaskId>& TaskExecutionDependencies, const TArray<FPCGTaskId>& TaskDataDependencies)
 {
 	check(GraphExecutor);
-	return GraphExecutor->ScheduleGenericWithContext(InOperation, InAbortOperation, SourceComponent, TaskExecutionDependencies, TaskDataDependencies);
+	return GraphExecutor->ScheduleGenericWithContext(InOperation, InAbortOperation, ExecutionSource, TaskExecutionDependencies, TaskDataDependencies);
 }
 
 
@@ -975,7 +987,7 @@ void UPCGSubsystem::CancelGenerationInternal(IPCGGraphExecutionSource* Source, b
 			return InvalidPCGTaskId;
 		};
 
-		ActorAndComponentMapping->DispatchToRegisteredLocalComponents(Component, LocalCancel);
+		TrackingManager->DispatchToRegisteredLocalComponents(Component, LocalCancel);
 	}
 }
 
@@ -991,11 +1003,15 @@ void UPCGSubsystem::RefreshRuntimeGenComponent(UPCGComponent* RuntimeComponent,
 
 void UPCGSubsystem::RefreshAllRuntimeGenComponents(EPCGChangeType ChangeType)
 {
-	for (UPCGComponent* Component : GetAllRegisteredComponents())
+	// @todo_pcg: support execution source
+	for (IPCGGraphExecutionSource* ExecutionSource : GetAllRegisteredExecutionSources())
 	{
-		if (Component->GenerationTrigger == EPCGComponentGenerationTrigger::GenerateAtRuntime)
+		if (UPCGComponent* Component = Cast<UPCGComponent>(ExecutionSource))
 		{
-			RefreshRuntimeGenComponent(Component, ChangeType);
+			if (Component->GenerationTrigger == EPCGComponentGenerationTrigger::GenerateAtRuntime)
+			{
+				RefreshRuntimeGenComponent(Component, ChangeType);
+			}
 		}
 	}
 }
@@ -1003,11 +1019,15 @@ void UPCGSubsystem::RefreshAllRuntimeGenComponents(EPCGChangeType ChangeType)
 #if WITH_EDITOR
 void UPCGSubsystem::RefreshAllComponentsFiltered(const TFunction<bool(UPCGComponent*)>& ComponentFilter, EPCGChangeType ChangeType)
 {
-	for (UPCGComponent* Component : GetAllRegisteredComponents())
+	// @todo_pcg: support execution source
+	for (IPCGGraphExecutionSource* ExecutionSource : GetAllRegisteredExecutionSources())
 	{
-		if (ComponentFilter(Component))
+		if (UPCGComponent* Component = Cast<UPCGComponent>(ExecutionSource))
 		{
-			Component->Refresh(ChangeType);
+			if (ComponentFilter(Component))
+			{
+				Component->Refresh(ChangeType);
+			}
 		}
 	}
 }
@@ -1022,7 +1042,7 @@ void UPCGSubsystem::ForAllRegisteredLocalComponents(UPCGComponent* InOriginalCom
 		return InvalidPCGTaskId;
 	};
 
-	ActorAndComponentMapping->DispatchToRegisteredLocalComponents(InOriginalComponent, WrapperFunc);
+	TrackingManager->DispatchToRegisteredLocalComponents(InOriginalComponent, WrapperFunc);
 }
 
 void UPCGSubsystem::ForAllRegisteredIntersectingLocalComponents(UPCGComponent* InOriginalComponent, const FBoxCenterAndExtent& InBounds, const TFunctionRef<void(UPCGComponent*)>& InFunc) const
@@ -1036,7 +1056,7 @@ void UPCGSubsystem::ForAllRegisteredIntersectingLocalComponents(UPCGComponent* I
 		return;
 	}
 
-	ActorAndComponentMapping->ForAllIntersectingPartitionActors(Overlap, [InOriginalComponent, &InFunc](APCGPartitionActor* Actor)
+	TrackingManager->ForAllIntersectingPartitionActors(Overlap, [InOriginalComponent, &InFunc](APCGPartitionActor* Actor)
 	{
 		if (UPCGComponent* LocalComponent = Actor->GetLocalComponent(InOriginalComponent))
 		{
@@ -1107,7 +1127,7 @@ FPCGTaskId UPCGSubsystem::ForAllOverlappingCells(UPCGComponent* InComponent, con
 			{
 				APCGPartitionActor* Actor = FindOrCreatePCGPartitionActor(Descriptor, CellCoord, bCanCreateActor);
 #if WITH_EDITOR
-				if (!Actor && !GetWorld()->IsGameWorld() && ActorAndComponentMapping->DoesPartitionActorRecordExist(Descriptor, CellCoord))
+				if (!Actor && !GetWorld()->IsGameWorld() && TrackingManager->DoesPartitionActorRecordExist(Descriptor, CellCoord))
 				{
 					FPCGTaskId ExecuteTaskId = InUnloadedFunc(Descriptor, CellCoord, IntersectedBounds);
 					
@@ -1183,58 +1203,25 @@ void UPCGSubsystem::CleanupLocalComponentsImmediate(UPCGComponent* InOriginalCom
 		return InvalidPCGTaskId;
 	};
 
-	ActorAndComponentMapping->DispatchToRegisteredLocalComponents(InOriginalComponent, LocalCleanupTask);
+	TrackingManager->DispatchToRegisteredLocalComponents(InOriginalComponent, LocalCleanupTask);
 
 	// Remove the local component mappings if requested and the component is partitioned. If 'bRemoveComponents' is false, that indicates we are doing a refresh, so destroying
 	// the component mappings is counterproductive.
 	if (bRemoveComponents && InOriginalComponent->IsPartitioned())
 	{
-		ActorAndComponentMapping->DeleteMappingPCGComponentPartitionActor(InOriginalComponent);
+		TrackingManager->DeleteMappingPCGComponentPartitionActor(InOriginalComponent);
 	}
 }
 
-// deprecated
-UPCGComponent* UPCGSubsystem::GetLocalComponent(uint32 GridSize, const FIntVector& CellCoords, const UPCGComponent* InOriginalComponent, bool bTransient) const
-{
-	check(InOriginalComponent);
-
-	FPCGGridDescriptor GridDescriptor = FPCGGridDescriptor()
-		.SetGridSize(GridSize)
-		.SetIsRuntime(bTransient)
-		.SetIs2DGrid(InOriginalComponent->Use2DGrid());
-
-	return GetLocalComponent(GridDescriptor, CellCoords, InOriginalComponent);
-}
-
-// deprecated
-APCGPartitionActor* UPCGSubsystem::GetRegisteredPCGPartitionActor(uint32 GridSize, const FIntVector& GridCoords, bool bRuntimeGenerated) const
-{
-	FPCGGridDescriptor GridDescriptor = FPCGGridDescriptor()
-		.SetGridSize(GridSize)
-		.SetIsRuntime(bRuntimeGenerated);
-
-	return GetRegisteredPCGPartitionActor(GridDescriptor, GridCoords);
-}
-
-// deprecated
-APCGPartitionActor* UPCGSubsystem::FindOrCreatePCGPartitionActor(const FGuid& Guid, uint32 GridSize, const FIntVector& GridCoords, bool bRuntimeGenerated, bool bCanCreateActor) const
-{
-	FPCGGridDescriptor GridDescriptor = FPCGGridDescriptor()
-		.SetGridSize(GridSize)
-		.SetIsRuntime(bRuntimeGenerated);
-
-	return FindOrCreatePCGPartitionActor(GridDescriptor, GridCoords, bCanCreateActor);
-}
-
 UPCGComponent* UPCGSubsystem::GetLocalComponent(const FPCGGridDescriptor& GridDescriptor, const FIntVector& CellCoords, const UPCGComponent* InOriginalComponent) const
 {
 	check(InOriginalComponent);
-	return ActorAndComponentMapping->GetLocalComponent(GridDescriptor, CellCoords, InOriginalComponent);
+	return TrackingManager->GetLocalComponent(GridDescriptor, CellCoords, InOriginalComponent);
 }
 
 APCGPartitionActor* UPCGSubsystem::GetRegisteredPCGPartitionActor(const FPCGGridDescriptor& GridDescriptor, const FIntVector& GridCoords) const
 {
-	return ActorAndComponentMapping->GetPartitionActor(GridDescriptor, GridCoords);
+	return TrackingManager->GetPartitionActor(GridDescriptor, GridCoords);
 }
 
 APCGPartitionActor* UPCGSubsystem::FindOrCreatePCGPartitionActor(const FPCGGridDescriptor& GridDescriptor, const FIntVector& GridCoords, bool bCanCreateActor, bool bHideFromOutliner) const
@@ -1271,7 +1258,7 @@ APCGPartitionActor* UPCGSubsystem::FindOrCreatePCGPartitionActor(const FPCGGridD
 
 #if WITH_EDITOR
 		// Check if there is already an unloaded actor for this cell. RuntimeGenerated PAs are never unloaded, so we ignore them.
-		if (ActorAndComponentMapping->DoesPartitionActorRecordExist(GridDescriptor, GridCoords))
+		if (TrackingManager->DoesPartitionActorRecordExist(GridDescriptor, GridCoords))
 		{
 			return nullptr;
 		}
@@ -1383,7 +1370,7 @@ APCGPartitionActor* UPCGSubsystem::FindOrCreatePCGPartitionActor(const FPCGGridD
 
 TSet<TObjectPtr<APCGPartitionActor>> UPCGSubsystem::GetPCGComponentPartitionActorMappings(UPCGComponent* InComponent) const
 {
-	return ActorAndComponentMapping->GetPCGComponentPartitionActorMappings(InComponent);
+	return TrackingManager->GetPCGComponentPartitionActorMappings(InComponent);
 }
 
 FPCGGenSourceManager* UPCGSubsystem::GetGenSourceManager() const
@@ -1422,7 +1409,7 @@ void UPCGSubsystem::DirtyGraph(UPCGComponent* Component, const FBox& InBounds, E
 		return InvalidPCGTaskId;
 	};
 
-	ActorAndComponentMapping->DispatchToRegisteredLocalComponents(Component, DirtyTask);
+	TrackingManager->DispatchToRegisteredLocalComponents(Component, DirtyTask);
 }
 
 void UPCGSubsystem::ClearPCGLink(UPCGComponent* InComponent, const FBox& InBounds, AActor* InNewActor)
@@ -1719,13 +1706,13 @@ void UPCGSubsystem::PropagateEditingModeToLocalComponents(UPCGComponent* InCompo
 {
 	if (ensure(InComponent && InComponent->IsPartitioned()))
 	{
-		FBox Bounds = ActorAndComponentMapping->PartitionedOctree.GetBounds(InComponent);
+		FBox Bounds = TrackingManager->PartitionedOctree.GetBounds(InComponent);
 		if (!Bounds.IsValid)
 		{
 			return;
 		}
 
-		ActorAndComponentMapping->ForAllIntersectingPartitionActors(Bounds, [InComponent, EditingMode](APCGPartitionActor* Actor)
+		TrackingManager->ForAllIntersectingPartitionActors(Bounds, [InComponent, EditingMode](APCGPartitionActor* Actor)
 		{
 			Actor->ChangeTransientState(InComponent, EditingMode);
 		});
@@ -1758,158 +1745,50 @@ void UPCGSubsystem::ClearLandscapeCache()
 
 void UPCGSubsystem::GenerateAllPCGComponents(bool bForce) const
 {
-	for (UPCGComponent* Component : ActorAndComponentMapping->GetAllRegisteredComponents())
+	for (IPCGGraphExecutionSource* ExecutionSource : TrackingManager->GetAllRegisteredExecutionSources())
 	{
-		Component->Generate(bForce);
+		IPCGGraphExecutionState::FGenerateParams GenerateParams;
+		GenerateParams.bEvenIfAlreadyGenerated = bForce;
+		ExecutionSource->GetExecutionState().Generate(GenerateParams);
 	}
 }
 
 void UPCGSubsystem::CleanupAllPCGComponents(bool bPurge) const
 {
-	for (UPCGComponent* Component : ActorAndComponentMapping->GetAllRegisteredComponents())
+	for (IPCGGraphExecutionSource* ExecutionSource : TrackingManager->GetAllRegisteredExecutionSources())
 	{
-		if (bPurge)
+		UPCGComponent* Component = Cast<UPCGComponent>(ExecutionSource);
+		if (Component && bPurge)
 		{
 			Component->CleanupLocalDeleteAllGeneratedObjects({});
 		}
 		else
 		{
-			Component->Cleanup(/*bRemoveComponents=*/true);
+			IPCGGraphExecutionState::FCleanupParams CleanupParams;
+			CleanupParams.bReleaseManagedResources = true;
+			ExecutionSource->GetExecutionState().Cleanup(CleanupParams);
 		}
 	}
 }
 
 void UPCGSubsystem::NotifySelectionKeyChanged(const FPCGSelectionKey& InSelectionKey, const UObject* InOriginatingChangeObject, TArrayView<const FBox> InChangedBounds)
 {
-	ActorAndComponentMapping->OnSelectionKeyChanged(InSelectionKey, InOriginatingChangeObject, InChangedBounds);
-}
-
-const FPCGNodeVisualLogs& UPCGSubsystem::GetNodeVisualLogs() const
-{
-	if (IPCGEditorModule* PCGEditorModule = IPCGEditorModule::Get())
-	{
-		return PCGEditorModule->GetNodeVisualLogs();
-	}
-	else
-	{
-		// Not meant to be used but to respect the existing API return a placeholder log
-		static FPCGNodeVisualLogs Placeholder;
-		return Placeholder;
-	}
-}
-
-FPCGNodeVisualLogs& UPCGSubsystem::GetNodeVisualLogsMutable()
-{
-	if (IPCGEditorModule* PCGEditorModule = IPCGEditorModule::Get())
-	{
-		return PCGEditorModule->GetNodeVisualLogsMutable();
-	}
-	else
-	{
-		// Not meant to be used but to respect the existing API return a placeholder log
-		static FPCGNodeVisualLogs Placeholder;
-		return Placeholder;
-	}
+	TrackingManager->OnSelectionKeyChanged(InSelectionKey, InOriginatingChangeObject, InChangedBounds);
 }
 
 void UPCGSubsystem::NotifyLandscapeEditModeExited()
 {
-	ActorAndComponentMapping->NotifyLandscapeEditModeExited();
-}
-
-/* UE_DEPRECATED(5.6) */
-void UPCGSubsystem::ClearExecutionMetadata(const FPCGStack& BaseStack)
-{
-	if (IPCGEditorModule* PCGEditorModule = IPCGEditorModule::Get())
-	{
-		if (UPCGComponent* Component = const_cast<UPCGComponent*>(Cast<UPCGComponent>(BaseStack.GetRootSource())))
-		{
-			PCGEditorModule->ClearExecutionMetadata(Component);
-		}
-	}
-}
-
-void UPCGSubsystem::ClearExecutionMetadata(UPCGComponent* InComponent)
-{
-	if (IPCGEditorModule* PCGEditorModule = IPCGEditorModule::Get())
-	{
-		PCGEditorModule->ClearExecutionMetadata(InComponent);
-	}
-}
-
-/* UE_DEPRECATED(5.6) */
-TArray<FPCGStack> UPCGSubsystem::GetExecutedStacks(const UPCGComponent* InComponent, const UPCGGraph* InSubgraph, bool bOnlyWithSubgraphAsCurrentFrame)
-{
-	if (IPCGEditorModule* PCGEditorModule = IPCGEditorModule::Get())
-	{
-		TArray<FPCGStackSharedPtr> StackPtrs = PCGEditorModule->GetExecutedStacksPtrs(InComponent, InSubgraph, bOnlyWithSubgraphAsCurrentFrame);
-		TArray<FPCGStack> Stacks;
-		Stacks.Reserve(StackPtrs.Num());
-		Algo::Transform(StackPtrs, Stacks, [](const FPCGStackSharedPtr& StackPtr) { return *StackPtr; });
-
-		return Stacks;
-	}
-	else
-	{
-		return {};
-	}
-}
-
-/* UE_DEPRECATED(5.6) */
-TArray<FPCGStack> UPCGSubsystem::GetExecutedStacks(const FPCGStack& BeginningWithStack)
-{
-	IPCGEditorModule* PCGEditorModule = IPCGEditorModule::Get();
-	if (!PCGEditorModule)
-	{
-		return {};
-	}
-
-	TArray<FPCGStackSharedPtr> StackPtrs = PCGEditorModule->GetExecutedStacksPtrs(BeginningWithStack);
-	TArray<FPCGStack> Stacks;
-
-	Stacks.Reserve(StackPtrs.Num());
-
-	for (FPCGStackSharedPtr StackPtr : StackPtrs)
-	{
-		Stacks.Add(*StackPtr);
-	}
-
-	return Stacks;
-}
-
-/* UE_DEPRECATED(5.6) */
-void UPCGSubsystem::ClearExecutedStacks(const FPCGStack& BeginningWithStack)
-{
-PRAGMA_DISABLE_DEPRECATION_WARNINGS
-	return ClearExecutedStacks(BeginningWithStack.GetRootComponent());
-PRAGMA_ENABLE_DEPRECATION_WARNINGS
-}
-
-/* UE_DEPRECATED(5.6) */
-void UPCGSubsystem::ClearExecutedStacks(const UPCGComponent* InRootComponent)
-{
-	if (IPCGEditorModule* PCGEditorModule = IPCGEditorModule::Get())
-	{
-		PCGEditorModule->ClearExecutedStacks(InRootComponent);
-	}
-}
-
-/* UE_DEPRECATED(5.6) */
-void UPCGSubsystem::ClearExecutedStacks(const UPCGGraph* InContainingGraph)
-{
-	if (IPCGEditorModule* PCGEditorModule = IPCGEditorModule::Get())
-	{
-		PCGEditorModule->ClearExecutedStacks(InContainingGraph);
-	}
+	TrackingManager->NotifyLandscapeEditModeExited();
 }
 
 void UPCGSubsystem::CreateMissingPartitionActors()
 {
 	if (!PCGHelpers::IsRuntimeOrPIE())
 	{
-		ActorAndComponentMapping->ForAllOriginalComponents([this](UPCGComponent* PCGComponent)
+		TrackingManager->ForAllOriginalExecutionSources([this](IPCGGraphExecutionSource* InExecutionSource)
 		{
-			if (PCGComponent->IsPartitioned() && !PCGComponent->IsManagedByRuntimeGenSystem())
+			UPCGComponent* PCGComponent = Cast<UPCGComponent>(InExecutionSource);
+			if (PCGComponent && PCGComponent->IsPartitioned() && !PCGComponent->IsManagedByRuntimeGenSystem())
 			{
 				bool bHasUnbounded = false;
 				PCGHiGenGrid::FSizeArray GridSizes;
@@ -1918,7 +1797,7 @@ void UPCGSubsystem::CreateMissingPartitionActors()
 				{
 					CreatePartitionActorsWithinBounds(PCGComponent, PCGComponent->GetGridBounds(), GridSizes);
 				}
-				ActorAndComponentMapping->UpdateMappingPCGComponentPartitionActor(PCGComponent);
+				TrackingManager->UpdateMappingPCGComponentPartitionActor(PCGComponent);
 			}
 		});
 	}
@@ -1954,12 +1833,12 @@ void UPCGSubsystem::CreatePartitionActorsWithinBounds(UPCGComponent* InComponent
 
 void UPCGSubsystem::UpdateMappingPCGComponentPartitionActor(UPCGComponent* InComponent) 
 { 
-	ActorAndComponentMapping->UpdateMappingPCGComponentPartitionActor(InComponent); 
+	TrackingManager->UpdateMappingPCGComponentPartitionActor(InComponent); 
 }
 
 void UPCGSubsystem::SetChainedDispatchToLocalComponents(bool bInChainedDispatch)
 {
-	ActorAndComponentMapping->SetChainedDispatchToLocalComponents(bInChainedDispatch);
+	TrackingManager->SetChainedDispatchToLocalComponents(bInChainedDispatch);
 }
 
 void UPCGSubsystem::OnObjectsReplaced(const TMap<UObject*, UObject*>& InOldToNewInstances)
@@ -2025,24 +1904,34 @@ bool UPCGSubsystem::IsComponentBeingUnloaded(UActorComponent* InComponent) const
 	return false;
 }
 
-void UPCGSubsystem::OnPCGGraphCancelled(UPCGComponent* InComponent)
+void UPCGSubsystem::OnPCGGraphCancelled(IPCGGraphExecutionSource* InExecutionSource)
 {
-	ActorAndComponentMapping->OnPCGGraphCancelled(InComponent);
+	TrackingManager->OnPCGGraphCancelled(InExecutionSource);
 }
 
-void UPCGSubsystem::OnPCGGraphStartGenerating(UPCGComponent* InComponent)
+void UPCGSubsystem::OnPCGGraphStartGenerating(IPCGGraphExecutionSource* InExecutionSource)
 {
-	ActorAndComponentMapping->OnPCGGraphStartsGenerating(InComponent);
+	TrackingManager->OnPCGGraphStartsGenerating(InExecutionSource);
 }
 
-void UPCGSubsystem::OnPCGGraphGenerated(UPCGComponent* InComponent)
+void UPCGSubsystem::OnPCGGraphGenerated(IPCGGraphExecutionSource* InExecutionSource)
 {
-	ActorAndComponentMapping->OnPCGGraphGeneratedOrCleaned(InComponent);
+	TrackingManager->OnPCGGraphGeneratedOrCleaned(InExecutionSource);
 }
 
-void UPCGSubsystem::OnPCGGraphCleaned(UPCGComponent* InComponent)
+void UPCGSubsystem::OnPCGGraphCleaned(IPCGGraphExecutionSource* InExecutionSource)
 {
-	ActorAndComponentMapping->OnPCGGraphGeneratedOrCleaned(InComponent);
+	TrackingManager->OnPCGGraphGeneratedOrCleaned(InExecutionSource);
+}
+
+void UPCGSubsystem::UpdateComponentTracking(IPCGGraphExecutionSource* InExecutionSource, bool bShouldDirtyActors, const TArray<FPCGSelectionKey>* OptionalChangedKeys)
+{
+	TrackingManager->UpdateTracking(InExecutionSource, bShouldDirtyActors, OptionalChangedKeys);
+}
+
+void UPCGSubsystem::ResetPartitionActorsMap()
+{
+	TrackingManager->ResetPartitionActorsMap();
 }
 
 #endif // WITH_EDITOR
@@ -2194,59 +2083,133 @@ void UPCGSubsystem::ExecuteBeginTickActions()
 	}
 }
 
-bool UPCGSubsystem::RegisterOrUpdatePCGComponent(UPCGComponent* InComponent, bool bDoActorMapping)
+void UPCGSubsystem::RegisterPartitionActor(APCGPartitionActor* InActor)
 {
-	return ActorAndComponentMapping->RegisterOrUpdatePCGComponent(InComponent, bDoActorMapping);
+	TrackingManager->RegisterPartitionActor(InActor);
 }
 
 void UPCGSubsystem::UnregisterPartitionActor(APCGPartitionActor* InActor)
 {
-	ActorAndComponentMapping->UnregisterPartitionActor(InActor);
+	TrackingManager->UnregisterPartitionActor(InActor);
 }
 
-TSet<UPCGComponent*> UPCGSubsystem::GetAllRegisteredPartitionedComponents() const
+bool UPCGSubsystem::RegisterOrUpdateExecutionSource(IPCGGraphExecutionSource* InExecutionSource, bool bDoPartitionMapping)
 {
-	return ActorAndComponentMapping->GetAllRegisteredPartitionedComponents();
+	return TrackingManager->RegisterOrUpdateExecutionSource(InExecutionSource, bDoPartitionMapping);
 }
 
-TSet<UPCGComponent*> UPCGSubsystem::GetAllRegisteredComponents() const
+bool UPCGSubsystem::RemapExecutionSource(const IPCGGraphExecutionSource* InOldExecutionSource, IPCGGraphExecutionSource* InNewExecutionSource, bool bDoPartitionMapping)
 {
-	return ActorAndComponentMapping->GetAllRegisteredComponents();
+	return TrackingManager->RemapExecutionSource(InOldExecutionSource, InNewExecutionSource, bDoPartitionMapping);
 }
 
-void UPCGSubsystem::ForAllIntersectingPartitionedComponents(const FBoxCenterAndExtent& InBounds, TFunctionRef<void(UPCGComponent*)> InFunc) const
+void UPCGSubsystem::UnregisterExecutionSource(IPCGGraphExecutionSource* InExecutionSource, bool bForce)
 {
-	return ActorAndComponentMapping->ForAllIntersectingPartitionedComponents(InBounds, std::move(InFunc));
+	return TrackingManager->UnregisterExecutionSource(InExecutionSource, bForce);
 }
 
-TArray<UPCGComponent*> UPCGSubsystem::GetAllIntersectingComponents(const FBoxCenterAndExtent& InBounds) const
+TSet<IPCGGraphExecutionSource*> UPCGSubsystem::GetAllRegisteredPartitionedExecutionSources() const
+{
+	return TrackingManager->GetAllRegisteredPartitionedExecutionSources();
+}
+
+TSet<IPCGGraphExecutionSource*> UPCGSubsystem::GetAllRegisteredExecutionSources() const
+{
+	return TrackingManager->GetAllRegisteredExecutionSources();
+}
+
+void UPCGSubsystem::ForAllIntersectingPartitionedExecutionSources(const FBoxCenterAndExtent& InBounds, TFunctionRef<void(IPCGGraphExecutionSource*)> InFunc) const
+{
+	return TrackingManager->ForAllIntersectingPartitionedExecutionSources(InBounds, std::move(InFunc));
+}
+
+TArray<IPCGGraphExecutionSource*> UPCGSubsystem::GetAllIntersectingExecutionSources(const FBoxCenterAndExtent& InBounds) const
 {
-	return ActorAndComponentMapping->GetAllIntersectingComponents(InBounds);
+	return TrackingManager->GetAllIntersectingExecutionSources(InBounds);
 }
 
+// deprecated
+bool UPCGSubsystem::RegisterOrUpdatePCGComponent(UPCGComponent* InComponent, bool bDoActorMapping)
+{
+	return RegisterOrUpdateExecutionSource(InComponent, bDoActorMapping);
+}
+
+// deprecated
 bool UPCGSubsystem::RemapPCGComponent(const UPCGComponent* OldComponent, UPCGComponent* NewComponent, bool bDoActorMapping)
 {
-	return ActorAndComponentMapping->RemapPCGComponent(OldComponent, NewComponent, bDoActorMapping);
+	return RemapExecutionSource(OldComponent, NewComponent, bDoActorMapping);
 }
 
+// deprecated
 void UPCGSubsystem::UnregisterPCGComponent(UPCGComponent* InComponent, bool bForce)
 {
-	ActorAndComponentMapping->UnregisterPCGComponent(InComponent, bForce);
+	UnregisterExecutionSource(InComponent, bForce);
 }
 
-void UPCGSubsystem::RegisterPartitionActor(APCGPartitionActor* InActor)
+// deprecated
+TSet<UPCGComponent*> UPCGSubsystem::GetAllRegisteredPartitionedComponents() const
 {
-	ActorAndComponentMapping->RegisterPartitionActor(InActor);
+	TSet<IPCGGraphExecutionSource*> ExecutionSources = TrackingManager->GetAllRegisteredPartitionedExecutionSources();
+
+	TSet<UPCGComponent*> PCGComponents;
+	PCGComponents.Reserve(ExecutionSources.Num());
+	Algo::TransformIf(ExecutionSources, PCGComponents,
+		[](IPCGGraphExecutionSource* InExecutionSource) { return Cast<UPCGComponent>(InExecutionSource) != nullptr; },
+		[](IPCGGraphExecutionSource* InExecutionSource) { return Cast<UPCGComponent>(InExecutionSource); });
+
+	return PCGComponents;
 }
 
-#if WITH_EDITOR
-void UPCGSubsystem::UpdateComponentTracking(UPCGComponent* InComponent, bool bShouldDirtyActors, const TArray<FPCGSelectionKey>* OptionalChangedKeys)
+// deprecated
+TSet<UPCGComponent*> UPCGSubsystem::GetAllRegisteredComponents() const
 {
-	ActorAndComponentMapping->UpdateTracking(InComponent, bShouldDirtyActors, OptionalChangedKeys);
+	TSet<IPCGGraphExecutionSource*> ExecutionSources = TrackingManager->GetAllRegisteredExecutionSources();
+
+	TSet<UPCGComponent*> PCGComponents;
+	PCGComponents.Reserve(ExecutionSources.Num());
+	Algo::TransformIf(ExecutionSources, PCGComponents,
+		[](IPCGGraphExecutionSource* InExecutionSource) { return Cast<UPCGComponent>(InExecutionSource) != nullptr; },
+		[](IPCGGraphExecutionSource* InExecutionSource) { return Cast<UPCGComponent>(InExecutionSource); });
+
+	return PCGComponents;
 }
 
-void UPCGSubsystem::ResetPartitionActorsMap()
+// deprecated
+void UPCGSubsystem::ForAllIntersectingPartitionedComponents(const FBoxCenterAndExtent& InBounds, TFunctionRef<void(UPCGComponent*)> InFunc) const
+{
+	auto WrapperFunc = [&InFunc](IPCGGraphExecutionSource* InExecutionSource)
+	{
+		if(UPCGComponent* PCGComponent = Cast<UPCGComponent>(InExecutionSource))
+		{
+			InFunc(PCGComponent);
+		}
+	};
+
+	return TrackingManager->ForAllIntersectingPartitionedExecutionSources(InBounds, WrapperFunc);
+}
+
+// deprecated
+TArray<UPCGComponent*> UPCGSubsystem::GetAllIntersectingComponents(const FBoxCenterAndExtent& InBounds) const
 {
-	ActorAndComponentMapping->ResetPartitionActorsMap();
+	TArray<IPCGGraphExecutionSource*> IntersectingExecutionSources = TrackingManager->GetAllIntersectingExecutionSources(InBounds);
+	
+	TArray<UPCGComponent*> IntersectingPCGComponents;
+	IntersectingPCGComponents.Reserve(IntersectingExecutionSources.Num());
+	Algo::TransformIf(IntersectingExecutionSources, IntersectingPCGComponents,
+		[](IPCGGraphExecutionSource* InExecutionSource) { return Cast<UPCGComponent>(InExecutionSource) != nullptr; },
+		[](IPCGGraphExecutionSource* InExecutionSource) { return Cast<UPCGComponent>(InExecutionSource); });
+
+	return IntersectingPCGComponents;
+}
+
+// deprecated
+void UPCGSubsystem::OnOriginalComponentRegistered(UPCGComponent* InComponent)
+{
+	OnOriginalExecutionSourceRegistered(InComponent);
+}
+
+// deprecated
+void UPCGSubsystem::OnOriginalComponentUnregistered(UPCGComponent* InComponent)
+{
+	OnOriginalExecutionSourceUnregistered(InComponent);
 }
-#endif // WITH_EDITOR
diff --git a/Engine/Plugins/PCG/Source/PCG/Public/Elements/PCGActorSelector.h b/Engine/Plugins/PCG/Source/PCG/Public/Elements/PCGActorSelector.h
--- a/Engine/Plugins/PCG/Source/PCG/Public/Elements/PCGActorSelector.h
+++ b/Engine/Plugins/PCG/Source/PCG/Public/Elements/PCGActorSelector.h
@@ -14,6 +14,7 @@
 #define UE_API PCG_API
 
 class AActor;
+class IPCGGraphExecutionSource;
 class UActorComponent;
 class UPCGComponent;
 class UWorld;
@@ -126,7 +127,7 @@ struct TStructOpsTypeTraits<FPCGCustomSelectionKey> : public TStructOpsTypeTrait
 USTRUCT()
 struct FPCGSelectionKey
 {
-	friend class FPCGActorAndComponentMapping;
+	friend class FPCGTrackingManager;
 
 	GENERATED_BODY()
 
@@ -149,8 +150,8 @@ struct FPCGSelectionKey
 	// Friend functions need to be explicitly exported.
 	PCG_API friend uint32 GetTypeHash(const FPCGSelectionKey& In);
 
-	UE_API bool IsMatching(const FPCGSelectionKey& InSelectionKey, const UPCGComponent* InComponent) const;
-	UE_API bool IsMatching(const FPCGSelectionKey& InSelectionKey, const TSet<FName>& InRemovedTags, const TSet<UPCGComponent*>& InComponents, TSet<UPCGComponent*>* OptionalMatchedComponents = nullptr) const;
+	UE_API bool IsMatching(const FPCGSelectionKey& InSelectionKey, const IPCGGraphExecutionSource* InExecutionSource) const;
+	UE_API bool IsMatching(const FPCGSelectionKey& InSelectionKey, const TSet<FName>& InRemovedTags, const TSet<IPCGGraphExecutionSource*>& InExecutionSources, TSet<IPCGGraphExecutionSource*>* OptionalExecutionSources = nullptr) const;
 
 	UE_API void SetExtraDependency(const UClass* InExtraDependency);
 	UE_API void UpdateAfterTagChange();
@@ -181,16 +182,14 @@ struct FPCGSelectionKey
 	TInstancedStruct<FPCGCustomSelectionKey> CustomKey;
 
 	// Deprecated methods
+	UE_DEPRECATED(5.8, "Replaced by the version using IPCGGraphExecutionSource.")
+	UE_API bool IsMatching(const FPCGSelectionKey& InSelectionKey, const TSet<FName>& InRemovedTags, const TSet<UPCGComponent*>& InComponents, TSet<UPCGComponent*>* OptionalMatchedComponents = nullptr) const;
+	
 	UE_DEPRECATED(5.7, "Replaced by the version with a FPCGSelectionKey.")
 	UE_API bool IsMatching(const UObject* InObject, const FSoftObjectPath& InObjectPath, const UPCGComponent* InComponent) const;
 	UE_DEPRECATED(5.7, "Replaced by the version with a FPCGSelectionKey.")
 	UE_API bool IsMatching(const UObject* InObject, const FSoftObjectPath& InObjectPath, const TSet<FName>& InRemovedTags, const TSet<UPCGComponent*>& InComponents, TSet<UPCGComponent*>* OptionalMatchedComponents = nullptr) const;
 
-	UE_DEPRECATED(5.6, "Replaced by the version with a soft object path.")
-	UE_API bool IsMatching(const UObject* InObject, const UPCGComponent* InComponent) const;
-	UE_DEPRECATED(5.6, "Replaced by the version with a soft object path.")
-	UE_API bool IsMatching(const UObject* InObject, const TSet<FName>& InRemovedTags, const TSet<UPCGComponent*>& InComponents, TSet<UPCGComponent*>* OptionalMatchedComponents = nullptr) const;
-
 protected:
 	FString CachedTagString;
 	bool bTagContainsWildcard = false;
diff --git a/Engine/Plugins/PCG/Source/PCG/Public/Grid/PCGComponentOctree.h b/Engine/Plugins/PCG/Source/PCG/Public/Grid/PCGComponentOctree.h
deleted file mode 100644
diff --git a/Engine/Plugins/PCG/Source/PCG/Public/Grid/PCGPartitionActor.h b/Engine/Plugins/PCG/Source/PCG/Public/Grid/PCGPartitionActor.h
--- a/Engine/Plugins/PCG/Source/PCG/Public/Grid/PCGPartitionActor.h
+++ b/Engine/Plugins/PCG/Source/PCG/Public/Grid/PCGPartitionActor.h
@@ -144,7 +144,7 @@ private:
 
 	UE_API UPCGSubsystem* GetSubsystem() const;
 
-	friend class FPCGActorAndComponentMapping;
+	friend class FPCGTrackingManager;
 #if WITH_EDITOR
 	UE_API void UpdateBoundsComponentExtents();
 	UE_API void SetInvalidForPCG();
@@ -188,7 +188,7 @@ private:
 	bool bRequiresUse2DGridFixup = false;
 #endif // WITH_EDITORONLY_DATA
 
-	/** Tracks the registration status of this PA with the ActorAndComponentMapping system. Helps us avoid invalid (un)registers. */
+	/** Tracks the registration status of this PA with the TrackingManager system. Helps us avoid invalid (un)registers. */
 	bool bIsRegistered = false;
 
 	/** Tracks if this actor was created by the Runtime Generation system. */
diff --git a/Engine/Plugins/PCG/Source/PCG/Public/Grid/PCGPartitionActorDesc.h b/Engine/Plugins/PCG/Source/PCG/Public/Grid/PCGPartitionActorDesc.h
--- a/Engine/Plugins/PCG/Source/PCG/Public/Grid/PCGPartitionActorDesc.h
+++ b/Engine/Plugins/PCG/Source/PCG/Public/Grid/PCGPartitionActorDesc.h
@@ -24,7 +24,7 @@ protected:
 	virtual void Serialize(FArchive& Ar) override;
 
 private:
-	friend class FPCGActorAndComponentMapping;
+	friend class FPCGTrackingManager;
 
 	bool bUse2DGrid = true;
 	bool bRequiresUse2DGridFixup = false;
diff --git a/Engine/Plugins/PCG/Source/PCG/Public/PCGComponent.h b/Engine/Plugins/PCG/Source/PCG/Public/PCGComponent.h
--- a/Engine/Plugins/PCG/Source/PCG/Public/PCGComponent.h
+++ b/Engine/Plugins/PCG/Source/PCG/Public/PCGComponent.h
@@ -115,7 +115,7 @@ class UPCGComponent : public UActorComponent, public IPCGGraphExecutionSource
 	friend class UPCGManagedActors;
 	friend class UPCGSubsystem;
 	friend class FPCGGraphExecutor;
-	friend class FPCGActorAndComponentMapping;
+	friend class FPCGTrackingManager;
 	friend class UE::FPCGInternalHelpers;
 	friend struct FPCGWorldPartitionBuilder;
 
diff --git a/Engine/Plugins/PCG/Source/PCG/Public/PCGComponentExecutionState.h b/Engine/Plugins/PCG/Source/PCG/Public/PCGComponentExecutionState.h
--- a/Engine/Plugins/PCG/Source/PCG/Public/PCGComponentExecutionState.h
+++ b/Engine/Plugins/PCG/Source/PCG/Public/PCGComponentExecutionState.h
@@ -29,6 +29,8 @@ public:
 	virtual void OnGraphExecutionAborted(bool bQuiet, bool bCleanupUnusedResources) override;
 	virtual void Cancel() override;
 	virtual bool IsGenerating() const override;
+	virtual bool IsGenerated() const override;
+	virtual FBox GetGeneratedBounds() const override;
 	virtual void ExecutePreGraph(FPCGContext* InContext) override;
 	virtual bool IsManagedByRuntimeGenSystem() const override;
 	virtual IPCGGraphExecutionSource* GetOriginalSource() const override;
diff --git a/Engine/Plugins/PCG/Source/PCG/Public/PCGGraphExecutionStateInterface.h b/Engine/Plugins/PCG/Source/PCG/Public/PCGGraphExecutionStateInterface.h
--- a/Engine/Plugins/PCG/Source/PCG/Public/PCGGraphExecutionStateInterface.h
+++ b/Engine/Plugins/PCG/Source/PCG/Public/PCGGraphExecutionStateInterface.h
@@ -120,6 +120,12 @@ public:
 	/** Indicates if the current source is generating. */
 	virtual bool IsGenerating() const = 0;
 
+	/** Indicates if the current source is generated */
+	virtual bool IsGenerated() const { return false; }
+
+	/** If this source is generated, returns the generated bounds */
+	virtual FBox GetGeneratedBounds() const { return FBox(EForceInit::ForceInit); }
+
 	/** Gets the task ID for the current generation task if it exists. */
 	virtual FPCGTaskId GetGenerationTaskId() const { return InvalidPCGTaskId; }
 
diff --git a/Engine/Plugins/PCG/Source/PCG/Public/PCGWorldActor.h b/Engine/Plugins/PCG/Source/PCG/Public/PCGWorldActor.h
--- a/Engine/Plugins/PCG/Source/PCG/Public/PCGWorldActor.h
+++ b/Engine/Plugins/PCG/Source/PCG/Public/PCGWorldActor.h
@@ -87,7 +87,7 @@ public:
 	bool bEnableWorldPartitionGenerationSources = false;
 
 private:
-	friend class FPCGActorAndComponentMapping;
+	friend class FPCGTrackingManager;
 
 	void RegisterToSubsystem();
 	void UnregisterFromSubsystem();
diff --git a/Engine/Plugins/PCG/Source/PCG/Public/RuntimeGen/PCGRuntimeGenScheduler.h b/Engine/Plugins/PCG/Source/PCG/Public/RuntimeGen/PCGRuntimeGenScheduler.h
--- a/Engine/Plugins/PCG/Source/PCG/Public/RuntimeGen/PCGRuntimeGenScheduler.h
+++ b/Engine/Plugins/PCG/Source/PCG/Public/RuntimeGen/PCGRuntimeGenScheduler.h
@@ -12,7 +12,7 @@
 
 class APCGPartitionActor;
 class APCGWorldActor;
-class FPCGActorAndComponentMapping;
+class FPCGTrackingManager;
 class FPCGGenSourceManager;
 class IPCGGenSourceBase;
 class ULevel;
@@ -67,7 +67,7 @@ class FPCGRuntimeGenScheduler
 	friend class UPCGSubsystem;
 
 public:
-	FPCGRuntimeGenScheduler(UWorld* InWorld, FPCGActorAndComponentMapping* InActorAndComponentMapping);
+	FPCGRuntimeGenScheduler(UWorld* InWorld, FPCGTrackingManager* InTrackingManager);
 	~FPCGRuntimeGenScheduler();
 
 	FPCGRuntimeGenScheduler(const FPCGRuntimeGenScheduler&) = delete;
@@ -155,8 +155,8 @@ protected:
 	{
 		const TSet<IPCGGenSourceBase*>* GenSources = nullptr;
 		const APCGWorldActor* PCGWorldActor = nullptr;
-		TSet<UPCGComponent*> AllPartitionedComponents;
-		TSet<UPCGComponent*> AllNonPartitionedComponents;
+		TSet<IPCGGraphExecutionSource*> AllPartitionedExecutionSources;
+		TSet<IPCGGraphExecutionSource*> AllNonPartitionedExecutionSources;
 		TSet<FGridGenerationKey>* GeneratedComponents = nullptr;
 	};
 
@@ -243,7 +243,7 @@ private:
 	/** PartitionActorPoolSize represents the current maximum capacity of the PartitionActorPool. */
 	int32 PartitionActorPoolSize = 0;
 
-	FPCGActorAndComponentMapping* ActorAndComponentMapping = nullptr;
+	FPCGTrackingManager* TrackingManager = nullptr;
 	FPCGGenSourceManager* GenSourceManager = nullptr;
 	UPCGSubsystem* Subsystem = nullptr;
 	UWorld* World = nullptr;
@@ -259,9 +259,9 @@ private:
 	/** Requests to flush all actors are deferred so they can be handled at a known time during tick. */
 	bool bActorFlushRequested = false;
 
-	/** Track the existence of runtime gen components to avoid unnecessary computation when there is no work to do. */
-	bool bAnyRuntimeGenComponentsExist = false;
-	bool bAnyRuntimeGenComponentsExistDirty = false;
+	/** Track the existence of runtime gen execution sources to avoid unnecessary computation when there is no work to do. */
+	bool bAnyRuntimeGenExecutionSourcesExist = false;
+	bool bAnyRuntimeGenExecutionSourcesExistDirty = false;
 
 	/** Setting up a PA calls APCGPartitionActor::AddGraphInstance which later calls RefreshComponent, which can create
 	* an infinite refresh loop. To break this loop we write the OC pointer to this variable, and if Refresh gets called for
diff --git a/Engine/Plugins/PCG/Source/PCG/Public/Subsystems/IPCGBaseSubsystem.h b/Engine/Plugins/PCG/Source/PCG/Public/Subsystems/IPCGBaseSubsystem.h
--- a/Engine/Plugins/PCG/Source/PCG/Public/Subsystems/IPCGBaseSubsystem.h
+++ b/Engine/Plugins/PCG/Source/PCG/Public/Subsystems/IPCGBaseSubsystem.h
@@ -6,7 +6,6 @@
 
 #include "PCGCommon.h"
 #include "Elements/PCGActorSelector.h"
-#include "Grid/PCGComponentOctree.h"
 #include "Utils/PCGNodeVisualLogs.h"
 #include "UObject/Interface.h"
 
diff --git a/Engine/Plugins/PCG/Source/PCG/Public/Subsystems/PCGSubsystem.h b/Engine/Plugins/PCG/Source/PCG/Public/Subsystems/PCGSubsystem.h
--- a/Engine/Plugins/PCG/Source/PCG/Public/Subsystems/PCGSubsystem.h
+++ b/Engine/Plugins/PCG/Source/PCG/Public/Subsystems/PCGSubsystem.h
@@ -8,20 +8,18 @@
 #include "PCGCommon.h"
 #include "Elements/PCGActorSelector.h"
 #include "Grid/PCGGridDescriptor.h"
-#include "Grid/PCGComponentOctree.h"
 #include "Utils/PCGNodeVisualLogs.h"
 
 #include "UObject/ObjectKey.h"
 
 #include "PCGSubsystem.generated.h"
 
-#define UE_API PCG_API
-
 class APCGPartitionActor;
 class APCGWorldActor;
-class FPCGActorAndComponentMapping;
+class FPCGTrackingManager;
 class FPCGGenSourceManager;
 class FPCGRuntimeGenScheduler;
+class UPCGComponent;
 class UPCGComputeGraph;
 class UPCGData;
 class UPCGGraph;
@@ -32,6 +30,7 @@ enum class ETickableTickType : uint8;
 namespace ERHIFeatureLevel { enum Type : int; }
 
 class IPCGGraphCache;
+class IPCGGraphExecutionSource;
 class FPCGGraphCompiler;
 class FPCGGraphExecutor;
 struct FPCGContext;
@@ -43,12 +42,12 @@ class UPCGSettings;
 class IPCGElement;
 typedef TSharedPtr<IPCGElement, ESPMode::ThreadSafe> FPCGElementPtr;
 
+class FBoxCenterAndExtent;
+
 class UWorld;
 
 #if WITH_EDITOR
-/** Deprecated - use FPCGOnPCGComponentUnregistered */
-DECLARE_MULTICAST_DELEGATE(FPCGOnComponentUnregistered);
-/** Deprecated - use FPCGOnPCGSourceGenerationDone */
+
 DECLARE_MULTICAST_DELEGATE_OneParam(FPCGOnComponentGenerationCompleteOrCancelled, UPCGSubsystem*);
 
 DECLARE_MULTICAST_DELEGATE_OneParam(FPCGOnPCGComponentUnregistered, UPCGComponent*);
@@ -67,86 +66,86 @@ class UPCGSubsystem : public UTickableWorldSubsystem, public IPCGBaseSubsystem
 
 public:
 	friend class UPCGComponent;
-	friend FPCGActorAndComponentMapping;
+	friend FPCGTrackingManager;
 	friend struct FPCGWorldPartitionBuilder;
 
-	UE_API UPCGSubsystem();
+	PCG_API UPCGSubsystem();
 
 	/** Add UObject references for GC */
 	static PCG_API void AddReferencedObjects(UObject* InThis, FReferenceCollector& Collector);
 
 	/** To be used when a PCG component can not have a world anymore, to unregister itself. */
-	static UE_API UPCGSubsystem* GetSubsystemForCurrentWorld();
+	static PCG_API UPCGSubsystem* GetSubsystemForCurrentWorld();
 
 	virtual UWorld* GetSubsystemWorld() const override { return GetWorld(); }
 
 	//~ Begin USubsystem Interface.
-	UE_API virtual void Deinitialize() override;
+	PCG_API virtual void Deinitialize() override;
 	//~ End USubsystem Interface
 
 	//~ Begin UWorldSubsystem Interface.
-	UE_API virtual void PostInitialize() override;
+	PCG_API virtual void PostInitialize() override;
 	// need UpdateStreamingState? 
 	//~ End UWorldSubsystem Interface.
 
 	//~ Begin FTickableGameObject
-	UE_API virtual void Tick(float DeltaSeconds) override;
+	PCG_API virtual void Tick(float DeltaSeconds) override;
 	virtual bool IsTickableInEditor() const override { return true; }
-	UE_API virtual ETickableTickType GetTickableTickType() const override;
-	UE_API virtual TStatId GetStatId() const override;
+	PCG_API virtual ETickableTickType GetTickableTickType() const override;
+	PCG_API virtual TStatId GetStatId() const override;
 	//~ End FTickableGameObject
 
 	/** Will return the subsystem from the World if it exists and if it is initialized */
-	static UE_API UPCGSubsystem* GetInstance(UWorld* World);
+	static PCG_API UPCGSubsystem* GetInstance(UWorld* World);
 
 	/** Adds an action that will be executed once at the beginning of this subsystem's next Tick(). */
 	using FTickAction = TFunction<void()>;
 
 	UE_DEPRECATED(5.6, "Use FPCGModule::ExecuteNextTick instead")
-	UE_API void RegisterBeginTickAction(FTickAction&& Action);
+	PCG_API void RegisterBeginTickAction(FTickAction&& Action);
 
 #if WITH_EDITOR
 	/** Returns PIE world if it is active, otherwise returns editor world. */
-	static UE_API UPCGSubsystem* GetActiveEditorInstance();
+	static PCG_API UPCGSubsystem* GetActiveEditorInstance();
 
-	UE_API void SetConstructionScriptSourceComponent(UPCGComponent* InComponent);
-	UE_API bool RemoveAndCopyConstructionScriptSourceComponent(AActor* InComponentOwner, FName InComponentName, UPCGComponent*& OutSourceComponent);
+	PCG_API void SetConstructionScriptSourceComponent(UPCGComponent* InComponent);
+	PCG_API bool RemoveAndCopyConstructionScriptSourceComponent(AActor* InComponentOwner, FName InComponentName, UPCGComponent*& OutSourceComponent);
 
 	/** Returns true if the Component is being unloaded through World Partition editor unloading or if the level the component is part of is being removed from world */
-	UE_API bool IsComponentBeingUnloaded(UActorComponent* InActorComponent) const;
+	PCG_API bool IsComponentBeingUnloaded(UActorComponent* InActorComponent) const;
 #endif
 
-	UE_API APCGWorldActor* GetPCGWorldActor();
-	UE_API APCGWorldActor* FindPCGWorldActor();
+	PCG_API APCGWorldActor* GetPCGWorldActor();
+	PCG_API APCGWorldActor* FindPCGWorldActor();
 
 	/** Returns current quality level between Low (0) and Cinematic (4). */
-	static UE_API int32 GetPCGQualityLevel();
-	UE_API void OnPCGQualityLevelChanged();
+	static PCG_API int32 GetPCGQualityLevel();
+	PCG_API void OnPCGQualityLevelChanged();
 
 #if WITH_EDITOR
-	UE_API void DestroyAllPCGWorldActors();
-	UE_API void DestroyCurrentPCGWorldActor();
-	UE_API void LogAbnormalComponentStates(bool bGroupByState) const;
+	PCG_API void DestroyAllPCGWorldActors();
+	PCG_API void DestroyCurrentPCGWorldActor();
+	PCG_API void LogAbnormalComponentStates(bool bGroupByState) const;
 #endif
-	UE_API void RegisterPCGWorldActor(APCGWorldActor* InActor);
-	UE_API void UnregisterPCGWorldActor(APCGWorldActor* InActor);
+	PCG_API void RegisterPCGWorldActor(APCGWorldActor* InActor);
+	PCG_API void UnregisterPCGWorldActor(APCGWorldActor* InActor);
 
-	UE_API void OnOriginalComponentRegistered(UPCGComponent* InComponent);
-	UE_API void OnOriginalComponentUnregistered(UPCGComponent* InComponent);
+	PCG_API void OnOriginalExecutionSourceRegistered(IPCGGraphExecutionSource* InExecutionSource);
+	PCG_API void OnOriginalExecutionSourceUnregistered(IPCGGraphExecutionSource* InExecutionSource);
 
-	UE_API UPCGLandscapeCache* GetLandscapeCache();
+	PCG_API UPCGLandscapeCache* GetLandscapeCache();
 
 	/** Schedule graph(owner->graph) */
-	UE_API FPCGTaskId ScheduleComponent(UPCGComponent* PCGComponent, EPCGHiGenGrid Grid, bool bForce, const TArray<FPCGTaskId>& InDependencies);
+	PCG_API FPCGTaskId ScheduleComponent(UPCGComponent* PCGComponent, EPCGHiGenGrid Grid, bool bForce, const TArray<FPCGTaskId>& InDependencies);
 
 	/** Schedule cleanup(owner->graph). */
-	UE_API FPCGTaskId ScheduleCleanup(UPCGComponent* PCGComponent, bool bRemoveComponents, const TArray<FPCGTaskId>& Dependencies);
+	PCG_API FPCGTaskId ScheduleCleanup(UPCGComponent* PCGComponent, bool bRemoveComponents, const TArray<FPCGTaskId>& Dependencies);
 
 	using IPCGBaseSubsystem::ScheduleGraph;
 	
-	UE_API FPCGTaskId ScheduleGraph(
+	PCG_API FPCGTaskId ScheduleGraph(
 		UPCGGraph* Graph,
-		UPCGComponent* SourceComponent,
+		IPCGGraphExecutionSource* ExecutionSource,
 		FPCGElementPtr PreGraphElement,
 		FPCGElementPtr InputElement,
 		const TArray<FPCGTaskId>& Dependencies,
@@ -154,96 +153,120 @@ public:
 		bool bAllowHierarchicalGeneration);
 
 	// Schedule graph (used internally for dynamic subgraph execution)
-	UE_API FPCGTaskId ScheduleGraph(UPCGComponent* SourceComponent, const TArray<FPCGTaskId>& Dependencies);
+	PCG_API FPCGTaskId ScheduleGraph(IPCGGraphExecutionSource* ExecutionSource, const TArray<FPCGTaskId>& Dependencies);
 
 	using IPCGBaseSubsystem::ScheduleGeneric;
 
 	/** General job scheduling
 	*  @param InOperation:               Callback that returns true if the task is done, false otherwise.
-	*  @param SourceComponent:           PCG component associated with this task. Can be null.
+	*  @param ExecutionSource:           Execution source associated with this task. Can be null.
 	*  @param TaskExecutionDependencies: Task will wait on these tasks to execute and won't take their output data as input.
 	*/
-	UE_API FPCGTaskId ScheduleGeneric(TFunction<bool()> InOperation, UPCGComponent* SourceComponent, const TArray<FPCGTaskId>& TaskExecutionDependencies);
+	PCG_API FPCGTaskId ScheduleGeneric(TFunction<bool()> InOperation, IPCGGraphExecutionSource* ExecutionSource, const TArray<FPCGTaskId>& TaskExecutionDependencies);
 
 	/** General job scheduling
 	*  @param InOperation:               Callback that returns true if the task is done, false otherwise.
 	*  @param InAbortOperation:          Callback that will be called if the generic task is cancelled for any reason.
-	*  @param SourceComponent:           PCG component associated with this task. Can be null.
+	*  @param ExecutionSource:           Execution source associated with this task. Can be null.
 	*  @param TaskExecutionDependencies: Task will wait on these tasks to execute and won't take their output data as input.
 	*/
-	UE_API FPCGTaskId ScheduleGeneric(TFunction<bool()> InOperation, TFunction<void()> InAbortOperation, UPCGComponent* SourceComponent, const TArray<FPCGTaskId>& TaskExecutionDependencies);
+	PCG_API FPCGTaskId ScheduleGeneric(TFunction<bool()> InOperation, TFunction<void()> InAbortOperation, IPCGGraphExecutionSource* ExecutionSource, const TArray<FPCGTaskId>& TaskExecutionDependencies);
 
 	/** General job scheduling with context
 	*  @param InOperation:               Callback that takes a Context as argument and returns true if the task is done, false otherwise.
-	*  @param SourceComponent:           PCG component associated with this task. Can be null.
+	*  @param ExecutionSource:           Execution source associated with this task. Can be null.
 	*  @param TaskExecutionDependencies: Task will wait on these tasks to execute and won't take their output data as input.
 	*  @param TaskDataDependencies:      Task will wait on these tasks to execute and will take their output data as input.
 	*/
-	UE_API FPCGTaskId ScheduleGenericWithContext(TFunction<bool(FPCGContext*)> InOperation, UPCGComponent* SourceComponent, const TArray<FPCGTaskId>& TaskExecutionDependencies, const TArray<FPCGTaskId>& TaskDataDependencies);
+	PCG_API FPCGTaskId ScheduleGenericWithContext(TFunction<bool(FPCGContext*)> InOperation, IPCGGraphExecutionSource* ExecutionSource, const TArray<FPCGTaskId>& TaskExecutionDependencies, const TArray<FPCGTaskId>& TaskDataDependencies);
 
 	/** General job scheduling with context
 	*  @param InOperation:               Callback that takes a Context as argument and returns true if the task is done, false otherwise.
 	*  @param InAbortOperation:          Callback that will be called if the generic task is cancelled for any reason.
-	*  @param SourceComponent:           PCG component associated with this task. Can be null.
+	*  @param ExecutionSource:           Execution source associated with this task. Can be null.
 	*  @param TaskExecutionDependencies: Task will wait on these tasks to execute and won't take their output data as input.
 	*  @param TaskDataDependencies:      Task will wait on these tasks to execute and will take their output data as input.
 	*/
-	UE_API FPCGTaskId ScheduleGenericWithContext(TFunction<bool(FPCGContext*)> InOperation, TFunction<void(FPCGContext*)> InAbortOperation, UPCGComponent* SourceComponent, const TArray<FPCGTaskId>& TaskExecutionDependencies, const TArray<FPCGTaskId>& TaskDataDependencies);
+	PCG_API FPCGTaskId ScheduleGenericWithContext(TFunction<bool(FPCGContext*)> InOperation, TFunction<void(FPCGContext*)> InAbortOperation, IPCGGraphExecutionSource* ExecutionSource, const TArray<FPCGTaskId>& TaskExecutionDependencies, const TArray<FPCGTaskId>& TaskDataDependencies);
 
 	/** Asks the runtime generation scheduler to refresh a given GenerateAtRuntime component. ChangeType should be 'GenerationGrid' to perform a full cleanup of PAs and local components. */
-	UE_API void RefreshRuntimeGenComponent(UPCGComponent* RuntimeComponent, EPCGChangeType ChangeType = EPCGChangeType::None);
+	PCG_API void RefreshRuntimeGenComponent(UPCGComponent* RuntimeComponent, EPCGChangeType ChangeType = EPCGChangeType::None);
 
 	/** Asks the runtime generation scheduler to refresh all GenerateAtRuntime components. ChangeType should be 'GenerationGrid' to perform a full cleanup of PAs and local components. */
-	UE_API void RefreshAllRuntimeGenComponents(EPCGChangeType ChangeType = EPCGChangeType::None);
+	PCG_API void RefreshAllRuntimeGenComponents(EPCGChangeType ChangeType = EPCGChangeType::None);
 
 #if WITH_EDITOR
 	/** Refresh all components selected by the filter (runtime generated or otherwise). */
-	UE_API void RefreshAllComponentsFiltered(const TFunction<bool(UPCGComponent*)>& ComponentFilter, EPCGChangeType ChangeType = EPCGChangeType::None);
+	PCG_API void RefreshAllComponentsFiltered(const TFunction<bool(UPCGComponent*)>& ComponentFilter, EPCGChangeType ChangeType = EPCGChangeType::None);
 #endif
 
 	FPCGRuntimeGenScheduler* GetRuntimeGenScheduler() const { return RuntimeGenScheduler; }
 
-	/** Register a new PCG Component or update it, will be added to the octree if it doesn't exists yet. Returns true if it was added/updated. Thread safe */
-	UE_API bool RegisterOrUpdatePCGComponent(UPCGComponent* InComponent, bool bDoActorMapping = true);
+	/** Register a new execution source or update it, will be added to the octree if it doesn't exists yet. Returns true if it was added/updated. Thread safe */
+	PCG_API bool RegisterOrUpdateExecutionSource(IPCGGraphExecutionSource* InExecutionSource, bool bDoPartitionMapping = true);
 
-	/** In case of BP Actors, we need to remap the old component destroyed by the construction script to the new one. Returns true if re-mapping succeeded. */
-	UE_API bool RemapPCGComponent(const UPCGComponent* OldComponent, UPCGComponent* NewComponent, bool bDoActorMapping);
+	/** In case of blueprints, we need to remap the old execution source destroyed by the construction script to the new one. Returns true if re-mapping succeeded. */
+	PCG_API bool RemapExecutionSource(const IPCGGraphExecutionSource* InOldExecutionSource, IPCGGraphExecutionSource* InNewExecutionSource, bool bDoPartitionMapping);
 
-	/** Unregister a PCG Component, will be removed from the octree. Can force it, if we have a delayed unregister. Thread safe */
-	UE_API void UnregisterPCGComponent(UPCGComponent* InComponent, bool bForce = false);
+	/** Unregister a execution source, will be removed from the octree. Can force it, if we have a delayed unregister. Thread safe */
+	PCG_API void UnregisterExecutionSource(IPCGGraphExecutionSource* InExecutionSource, bool bForce = false);
 
 	/** Register a new Partition actor, will be added to a map and will query all intersecting volume to bind to them if asked. Thread safe */
-	UE_API void RegisterPartitionActor(APCGPartitionActor* InActor);
+	PCG_API void RegisterPartitionActor(APCGPartitionActor* InActor);
 
 	/** Unregister a Partition actor, will be removed from the map and remove itself to all intersecting volumes. Thread safe */
-	UE_API void UnregisterPartitionActor(APCGPartitionActor* InActor);
+	PCG_API void UnregisterPartitionActor(APCGPartitionActor* InActor);
 
-	UE_API TSet<UPCGComponent*> GetAllRegisteredPartitionedComponents() const;
-	UE_API TSet<UPCGComponent*> GetAllRegisteredComponents() const;
+	PCG_API TSet<IPCGGraphExecutionSource*> GetAllRegisteredPartitionedExecutionSources() const;
+	PCG_API TSet<IPCGGraphExecutionSource*> GetAllRegisteredExecutionSources() const;
 
-	/** Call the InFunc function to all local component registered to the original component. Thread safe*/
-	UE_API void ForAllRegisteredLocalComponents(UPCGComponent* InOriginalComponent, const TFunctionRef<void(UPCGComponent*)>& InFunc) const;
+	/** Call the InFunc function to all local components registered to the original component. Thread safe*/
+	PCG_API void ForAllRegisteredLocalComponents(UPCGComponent* InOriginalComponent, const TFunctionRef<void(UPCGComponent*)>& InFunc) const;
 
 	/** Call the InFunc function to all local component registered to the original component within some bounds. Thread safe*/
-	UE_API void ForAllRegisteredIntersectingLocalComponents(UPCGComponent* InOriginalComponent, const FBoxCenterAndExtent& InBounds, const TFunctionRef<void(UPCGComponent*)>& InFunc) const;
+	PCG_API void ForAllRegisteredIntersectingLocalComponents(UPCGComponent* InOriginalComponent, const FBoxCenterAndExtent& InBounds, const TFunctionRef<void(UPCGComponent*)>& InFunc) const;
 
-	/** Get all components in specified bounds. */
-	UE_API void ForAllIntersectingPartitionedComponents(const FBoxCenterAndExtent& InBounds, TFunctionRef<void(UPCGComponent*)> InFunc) const;
-
-	/** Gather all the PCG components within some bounds. */
-	UE_API TArray<UPCGComponent*> GetAllIntersectingComponents(const FBoxCenterAndExtent& InBounds) const;
+	/** Gather all the Execution sources within some bounds. */
+	PCG_API TArray<IPCGGraphExecutionSource*> GetAllIntersectingExecutionSources(const FBoxCenterAndExtent& InBounds) const;
 
 	/** Traverses the hierarchy associated with the given component and calls InFunc for each overlapping component. */
-	UE_API void ForAllOverlappingComponentsInHierarchy(UPCGComponent* InComponent, const TFunctionRef<void(UPCGComponent*)>& InFunc) const;
+	PCG_API void ForAllOverlappingComponentsInHierarchy(UPCGComponent* InComponent, const TFunctionRef<void(UPCGComponent*)>& InFunc) const;
+
+	/** Call the InFunc on all partitioned execution source in specified bounds. */
+	PCG_API void ForAllIntersectingPartitionedExecutionSources(const FBoxCenterAndExtent& InBounds, TFunctionRef<void(IPCGGraphExecutionSource*)> InFunc) const;
+
+	UE_DEPRECATED(5.8, "Use OnOriginalExecutionSourceRegistered instead")
+	PCG_API void OnOriginalComponentRegistered(UPCGComponent* InComponent);
+	
+	UE_DEPRECATED(5.8, "Use OnOriginalExecutionSourceUnregistered instead")
+	PCG_API void OnOriginalComponentUnregistered(UPCGComponent* InComponent);
+
+	UE_DEPRECATED(5.8, "Use RegisterOrUpdateExecutionSource instead")
+	PCG_API bool RegisterOrUpdatePCGComponent(UPCGComponent* InComponent, bool bDoActorMapping = true);
+
+	UE_DEPRECATED(5.8, "Use RemapExecutionSource instead")
+	PCG_API bool RemapPCGComponent(const UPCGComponent* OldComponent, UPCGComponent* NewComponent, bool bDoActorMapping);
+
+	UE_DEPRECATED(5.8, "Use UnregisterExecutionSource instead")
+	PCG_API void UnregisterPCGComponent(UPCGComponent* InComponent, bool bForce = false);
+
+	UE_DEPRECATED(5.8, "Use GetAllRegisteredPartitionedExecutionSources instead")
+	PCG_API TSet<UPCGComponent*> GetAllRegisteredPartitionedComponents() const;
+
+	UE_DEPRECATED(5.8, "Use GetAllRegisteredExecutionSources instead")
+	PCG_API TSet<UPCGComponent*> GetAllRegisteredComponents() const;
+
+	UE_DEPRECATED(5.8, "Use GetAllIntersectingExecutionSources instead")
+	PCG_API TArray<UPCGComponent*> GetAllIntersectingComponents(const FBoxCenterAndExtent& InBounds) const;
+
+	UE_DEPRECATED(5.8, "Use ForAllIntersectingPartitionedExecutionSources instead")
+	PCG_API void ForAllIntersectingPartitionedComponents(const FBoxCenterAndExtent& InBounds, TFunctionRef<void(UPCGComponent*)> InFunc) const;
 
 	/**
 	 * Call InFunc to all partition grid cells matching 'InGridSizes' and overlapping with 'InBounds'. 'InFunc' can schedule work or execute immediately.
 	 * 'InGridSizes' should be sorted in descending order. If 'bCanCreateActor' is true, it will create the partition actor at that cell if necessary.
 	 */
-	UE_DEPRECATED(5.5, "Use version with UPCGComponent")
-	FPCGTaskId ForAllOverlappingCells(const FBox& InBounds, const PCGHiGenGrid::FSizeArray& InGridSizes, bool bCanCreateActor, const TArray<FPCGTaskId>& Dependencies, TFunctionRef<FPCGTaskId(APCGPartitionActor*, const FBox&)> InFunc) const { return InvalidPCGTaskId;  }
-
-	UE_API FPCGTaskId ForAllOverlappingCells(
+	PCG_API FPCGTaskId ForAllOverlappingCells(
 		UPCGComponent* InPCGComponent,
 		const FBox& InBounds,
 		const PCGHiGenGrid::FSizeArray& InGridSizes,
@@ -253,126 +276,83 @@ public:
 		TFunctionRef<FPCGTaskId(const FPCGGridDescriptor&, const FIntVector&, const FBox&)> InUnloadedFunc = [](const FPCGGridDescriptor&, const FIntVector&, const FBox&) { return InvalidPCGTaskId; }) const;
 		
 	/** Immediately cleanup the local components associated with an original component. */
-	UE_API void CleanupLocalComponentsImmediate(UPCGComponent* InOriginalComponent, bool bRemoveComponents);
-
-	UE_DEPRECATED(5.5, "Use FPCGGridDescriptor version")
-	UE_API UPCGComponent* GetLocalComponent(uint32 GridSize, const FIntVector& CellCoords, const UPCGComponent* InOriginalComponent, bool bTransient = false) const;
-
-	UE_DEPRECATED(5.5, "Use FPCGGridDescriptor version")
-	UE_API APCGPartitionActor* GetRegisteredPCGPartitionActor(uint32 GridSize, const FIntVector& GridCoords, bool bRuntimeGenerated = false) const;
-
-	UE_DEPRECATED(5.5, "Use FPCGGridDescriptor verison")
-	UE_API APCGPartitionActor* FindOrCreatePCGPartitionActor(const FGuid& Guid, uint32 GridSize, const FIntVector& GridCoords, bool bRuntimeGenerated, bool bCanCreateActor = true) const;
+	PCG_API void CleanupLocalComponentsImmediate(UPCGComponent* InOriginalComponent, bool bRemoveComponents);
 
 	/** Retrieves a local component using grid descriptor and grid coordinates, returns nullptr if no such component is found. */
-	UE_API UPCGComponent* GetLocalComponent(const FPCGGridDescriptor& GridDescriptor, const FIntVector& CellCoords, const UPCGComponent* InOriginalComponent) const;
+	PCG_API UPCGComponent* GetLocalComponent(const FPCGGridDescriptor& GridDescriptor, const FIntVector& CellCoords, const UPCGComponent* InOriginalComponent) const;
 
 	/** Retrieves a registered partition actor using grid size and grid coordinates, returns nullptr if no such partition actor is found. */
-	UE_API APCGPartitionActor* GetRegisteredPCGPartitionActor(const FPCGGridDescriptor& GridDescriptor, const FIntVector& GridCoords) const;
+	PCG_API APCGPartitionActor* GetRegisteredPCGPartitionActor(const FPCGGridDescriptor& GridDescriptor, const FIntVector& GridCoords) const;
 
 	/** Creates a new partition actor if one does not already exist with the same grid size, coords, and generation mode. */
-	UE_API APCGPartitionActor* FindOrCreatePCGPartitionActor(const FPCGGridDescriptor& GridDescriptor, const FIntVector& GridCoords, bool bCanCreateActor = true, bool bHideFromOutliner = false) const;
+	PCG_API APCGPartitionActor* FindOrCreatePCGPartitionActor(const FPCGGridDescriptor& GridDescriptor, const FIntVector& GridCoords, bool bCanCreateActor = true, bool bHideFromOutliner = false) const;
 	
 	/** Retrieves partition actors for this original component */
-	UE_API TSet<TObjectPtr<APCGPartitionActor>> GetPCGComponentPartitionActorMappings(UPCGComponent* InComponent) const;
+	PCG_API TSet<TObjectPtr<APCGPartitionActor>> GetPCGComponentPartitionActorMappings(UPCGComponent* InComponent) const;
 
-	UE_API FPCGGenSourceManager* GetGenSourceManager() const;
+	PCG_API FPCGGenSourceManager* GetGenSourceManager() const;
 
 #if WITH_EDITOR
 public:
 	/** Schedule refresh on the current or next frame */
-	UE_API FPCGTaskId ScheduleRefresh(UPCGComponent* SourceComponent, bool bForceRefresh);
+	PCG_API FPCGTaskId ScheduleRefresh(UPCGComponent* SourceComponent, bool bForceRefresh);
 
 	/** Immediately dirties the partition actors in the given bounds */
-	UE_API void DirtyGraph(UPCGComponent* Component, const FBox& InBounds, EPCGComponentDirtyFlag DirtyFlag);
+	PCG_API void DirtyGraph(UPCGComponent* Component, const FBox& InBounds, EPCGComponentDirtyFlag DirtyFlag);
 
 	/** Delete serialized partition actors in the level. If 'bOnlyDeleteUnused' is true, only PAs with no graph instances will be deleted. */
-	UE_API void DeleteSerializedPartitionActors(bool bOnlyDeleteUnused, bool bOnlyChildren = false);
+	PCG_API void DeleteSerializedPartitionActors(bool bOnlyDeleteUnused, bool bOnlyChildren = false);
 
 	/** Update the tracking on a given component. */
-	UE_API void UpdateComponentTracking(UPCGComponent* InComponent, bool bShouldDirtyActors, const TArray<FPCGSelectionKey>* OptionalChangedKeys = nullptr);
+	PCG_API void UpdateComponentTracking(IPCGGraphExecutionSource* InExecutionSource, bool bShouldDirtyActors, const TArray<FPCGSelectionKey>* OptionalChangedKeys = nullptr);
 
 	/** Propagates transient state change from an original component to the relevant partition actors */
-	UE_API void PropagateEditingModeToLocalComponents(UPCGComponent* InOriginalComponent, EPCGEditorDirtyMode EditingMode);
+	PCG_API void PropagateEditingModeToLocalComponents(UPCGComponent* InOriginalComponent, EPCGEditorDirtyMode EditingMode);
 
 	/** Move all resources from sub actors to a new actor */
-	UE_API void ClearPCGLink(UPCGComponent* InComponent, const FBox& InBounds, AActor* InNewActor);
+	PCG_API void ClearPCGLink(UPCGComponent* InComponent, const FBox& InBounds, AActor* InNewActor);
 
 	/** If the partition grid size change, call this to empty the Partition actors map */
-	UE_API void ResetPartitionActorsMap();
+	PCG_API void ResetPartitionActorsMap();
 
 	/** Builds the landscape data cache. If force build is true, then it will build the cache even if it is never serialized. */
-	UE_API void BuildLandscapeCache(bool bQuiet = false, bool bForceBuild = true);
+	PCG_API void BuildLandscapeCache(bool bQuiet = false, bool bForceBuild = true);
 
 	/** Clears the landscape data cache */
-	UE_API void ClearLandscapeCache();
+	PCG_API void ClearLandscapeCache();
 
 	/** Will gather all the components registered, and ask for generate. */
-	UE_API void GenerateAllPCGComponents(bool bForce) const;
+	PCG_API void GenerateAllPCGComponents(bool bForce) const;
 
 	/** Will gather all the components registered, and ask for cleanup. */
-	UE_API void CleanupAllPCGComponents(bool bPurge) const;
+	PCG_API void CleanupAllPCGComponents(bool bPurge) const;
 
 	/** Notification that a Selection key needs refreshing */
-	UE_API void NotifySelectionKeyChanged(const FPCGSelectionKey& InSelectionKey, const UObject* InOriginatingChangeObject, TArrayView<const FBox> InBounds);
-
-	/** Get graph warnings and errors for all nodes. */
-	UE_DEPRECATED(5.6, "Use IPCGEditorModule::GetNodeVisualLogs instead")
-	UE_API const FPCGNodeVisualLogs& GetNodeVisualLogs() const;
-	
-	UE_DEPRECATED(5.6, "Use IPCGEditorModule::GetNodeVisualLogsMutable instead")
-	UE_API FPCGNodeVisualLogs& GetNodeVisualLogsMutable();
+	PCG_API void NotifySelectionKeyChanged(const FPCGSelectionKey& InSelectionKey, const UObject* InOriginatingChangeObject, TArrayView<const FBox> InBounds);
 
 	/** Notify that we exited the Landscape edit mode. */
-	UE_API void NotifyLandscapeEditModeExited();
-
-	UE_DEPRECATED(5.6, "No longer supported")
-	UE_API void ClearExecutionMetadata(UPCGComponent* InComponent);
-
-	UE_DEPRECATED(5.6, "No longer supported")
-	UE_API void ClearExecutionMetadata(const FPCGStack& BaseStack);
-
-	UE_DEPRECATED(5.6, "No longer supported")
-	UE_API void ClearExecutedStacks(const UPCGComponent* InRootComponent);
-
-	UE_DEPRECATED(5.6, "No longer supported")
-	UE_API void ClearExecutedStacks(const UPCGGraph* InContainingGraph);
-
-	UE_DEPRECATED(5.6, "No longer supported")
-	UE_API TArray<FPCGStack> GetExecutedStacks(const UPCGComponent* InComponent, const UPCGGraph* InSubgraph, bool bOnlyWithSubgraphAsCurrentFrame = true);
-
-	UE_DEPRECATED(5.6, "No longer supported")
-	UE_API TArray<FPCGStack> GetExecutedStacks(const FPCGStack& BeginningWithStack);
-
-	UE_DEPRECATED(5.6, "No longer supported")
-	UE_API void ClearExecutedStacks(const FPCGStack& BeginningWithStack);
-
-	UE_DEPRECATED(5.5, "Deprecated in favor of OnPCGComponentUnregistered, will not be notified anymore")
-	FPCGOnComponentUnregistered OnComponentUnregistered;
-
-	UE_DEPRECATED(5.5, "Deprecated in favor of OnPCGComponentGenerationDone, will not be notified anymore")
-	FPCGOnComponentGenerationCompleteOrCancelled OnComponentGenerationCompleteOrCancelled;
+	PCG_API void NotifyLandscapeEditModeExited();
 
 	FPCGOnPCGComponentUnregistered OnPCGComponentUnregistered;
 	
 	UE_DEPRECATED(5.7, "Deprecated in favor of OnPCGSourceGenerationDone, will not be notified anymore")
 	FPCGOnPCGComponentGenerationDone OnPCGComponentGenerationDone;
 
-	UE_API void CreateMissingPartitionActors();
+	PCG_API void CreateMissingPartitionActors();
 
 	void OnPreviewPlatformChanged();
 	void OnPreviewFeatureLevelChanged(ERHIFeatureLevel::Type);
 
 private:
-	UE_API void OnPCGGraphCancelled(UPCGComponent* InComponent);
-	UE_API void OnPCGGraphStartGenerating(UPCGComponent* InComponent);
-	UE_API void OnPCGGraphGenerated(UPCGComponent* InComponent);
-	UE_API void OnPCGGraphCleaned(UPCGComponent* InComponent);
+	void OnPCGGraphCancelled(IPCGGraphExecutionSource* InExecutionSource);
+	void OnPCGGraphStartGenerating(IPCGGraphExecutionSource* InExecutionSource);
+	void OnPCGGraphGenerated(IPCGGraphExecutionSource* InExecutionSource);
+	void OnPCGGraphCleaned(IPCGGraphExecutionSource* InExecutionSource);
 
-	UE_API void CreatePartitionActorsWithinBounds(UPCGComponent* InComponent, const FBox& InBounds, const PCGHiGenGrid::FSizeArray& InGridSizes);
-	UE_API void UpdateMappingPCGComponentPartitionActor(UPCGComponent* InComponent);
+	PCG_API void CreatePartitionActorsWithinBounds(UPCGComponent* InComponent, const FBox& InBounds, const PCGHiGenGrid::FSizeArray& InGridSizes);
+	PCG_API void UpdateMappingPCGComponentPartitionActor(UPCGComponent* InComponent);
 
-	UE_API void SetChainedDispatchToLocalComponents(bool bInChainedDispatch);
+	PCG_API void SetChainedDispatchToLocalComponents(bool bInChainedDispatch);
 
 	void OnObjectsReplaced(const TMap<UObject*, UObject*>& InOldToNewInstances);
 	void OnLevelAddedToWorld(ULevel* InLevel, UWorld* InWorld);
@@ -383,7 +363,7 @@ private:
 	
 private:
 	/** When registering PAs, we might not have a PCG World Actor already registered (at runtime). In that case we'll look for a World Actor in the same level.*/
-	UE_API APCGWorldActor* GetPCGWorldActorForPartitionActor(APCGPartitionActor* InActor);
+	APCGWorldActor* GetPCGWorldActorForPartitionActor(APCGPartitionActor* InActor);
 
 	UPCGData* GetPCGData(FPCGTaskId InGraphExecutionTaskId, bool& bOutFound);
 	UPCGData* GetInputPCGData(FPCGTaskId InGraphExecutionTaskId, bool& bOutFound);
@@ -407,12 +387,12 @@ private:
 	void SetLocalSpaceBounds(FPCGTaskId InGraphExecutionTaskId, const FBox& InBounds);
 	void SetOriginalLocalSpaceBounds(FPCGTaskId InGraphExecutionTaskId, const FBox& InBounds);
 
-	UE_API void ExecuteBeginTickActions();
+	void ExecuteBeginTickActions();
 
 	APCGWorldActor* PCGWorldActor = nullptr;
 	FPCGRuntimeGenScheduler* RuntimeGenScheduler = nullptr;
 	bool bHasTickedOnce = false;
-	TUniquePtr<FPCGActorAndComponentMapping> ActorAndComponentMapping;
+	TUniquePtr<FPCGTrackingManager> TrackingManager;
 
 	/** Functions will be executed at the beginning of the tick and then removed from this array. */
 	TArray<FTickAction> BeginTickActions;
@@ -425,7 +405,7 @@ private:
 	using FConstructionScriptSourceComponents = TMap<FName, TObjectKey<UPCGComponent>>;
 	TMap<TObjectKey<AActor>, FConstructionScriptSourceComponents> PerActorConstructionScriptSourceComponents;
 
-	static UE_API TSet<UWorld*> DisablePartitionActorCreationForWorld;
+	PCG_API static TSet<UWorld*> DisablePartitionActorCreationForWorld;
 
 	// Used by UPCGWorldPartitonBuilder to disable PA creation while outside of a certain scope
 	static void SetDisablePartitionActorCreationForWorld(UWorld* InWorld, bool bDisable) 
@@ -449,5 +429,3 @@ private:
 protected:
 	PCG_API virtual void CancelGenerationInternal(IPCGGraphExecutionSource* Source, bool bCleanupUnusedResources) override;
 };
-
-#undef UE_API
