[PCG] Support for array overrides for graphs parameters in subgraphs

* Was not submitted with the original CL for array support

#rnx

diff --git a/Engine/Plugins/PCG/Source/PCG/Private/PCGSubgraph.cpp b/Engine/Plugins/PCG/Source/PCG/Private/PCGSubgraph.cpp
--- a/Engine/Plugins/PCG/Source/PCG/Private/PCGSubgraph.cpp
+++ b/Engine/Plugins/PCG/Source/PCG/Private/PCGSubgraph.cpp
@@ -280,6 +280,7 @@ TArray<FPCGSettingsOverridableParam> UPCGBaseSubgraphSettings::GatherOverridable
 				PCGSettingsHelpers::FPCGGetAllOverridableParamsConfig Config;
 				Config.bExcludeSuperProperties = true;
 				Config.ExcludePropertyFlags = CPF_DisableEditOnInstance;
+				Config.bExtractArrays = true;
 				OverridableParams.Append(PCGSettingsHelpers::GetAllOverridableParams(ScriptStruct, Config));
 			}
 		}
