[PCG] Remove temporary generate landscape textures window mode cvar.

Results looked identical before and after on multiple platforms.

#rnx

diff --git a/Engine/Plugins/PCG/Source/PCG/Private/Elements/PCGGenerateLandscapeTextures.cpp b/Engine/Plugins/PCG/Source/PCG/Private/Elements/PCGGenerateLandscapeTextures.cpp
--- a/Engine/Plugins/PCG/Source/PCG/Private/Elements/PCGGenerateLandscapeTextures.cpp
+++ b/Engine/Plugins/PCG/Source/PCG/Private/Elements/PCGGenerateLandscapeTextures.cpp
@@ -51,9 +51,6 @@ namespace PCGGenerateLandscapeTextures
 		TEXT("Draws debug boxes around landscapes for which grass maps are generated, colored by the current task ID."));
 #endif
 
-	// todo_pcg: Remove after running tests.
-	TAutoConsoleVariable<bool> CVarWindowedGeneration(TEXT("pcg.Grass.WindowedGeneration"), false, TEXT(""));
-
 	static bool IsTextureFullyStreamedIn(UTexture* InTexture)
 	{
 		const bool bCheckForLODTransition = true;
@@ -565,28 +562,19 @@ bool FPCGGenerateLandscapeTexturesElement::ExecuteInternal(FPCGContext* InContex
 
 		FIntPoint WindowStart, WindowEnd;
 
-		if (PCGGenerateLandscapeTextures::CVarWindowedGeneration.GetValueOnGameThread())
-		{
-			const FBox CellBounds = Context->ExecutionSource->GetExecutionState().GetBounds();
-			const FBox& TotalBounds = Context->TotalLandscapeComponentBounds;
-
-			UE_LOG(LogPCG, Verbose, TEXT("PCGUnpackGrassMap: %s, Cell bounds (%.2f, %.2f -> %.2f, %.2f), Total bounds (%.2f, %.2f -> %.2f, %.2f), efficiency %.2f"),
-				*Context->GetExecutionSourceName(),
-				CellBounds.Min.X, CellBounds.Min.Y, CellBounds.Max.X, CellBounds.Max.Y,
-				TotalBounds.Min.X, TotalBounds.Min.Y, TotalBounds.Max.X, TotalBounds.Max.Y,
-				(CellBounds.Max.X - CellBounds.Min.X)* (CellBounds.Max.Y - CellBounds.Min.Y) / ((TotalBounds.Max.X - TotalBounds.Min.X) * (TotalBounds.Max.Y - TotalBounds.Min.Y)));
-
-			WindowStart.X = FMath::Clamp(FMath::Floor((CellBounds.Min.X - TotalBounds.Min.X) / QuadSizeWorld), 0, FullExportResolution.X - 1);
-			WindowStart.Y = FMath::Clamp(FMath::Floor((CellBounds.Min.Y - TotalBounds.Min.Y) / QuadSizeWorld), 0, FullExportResolution.Y - 1);
-			WindowEnd.X = FMath::Clamp(FMath::CeilToInt((CellBounds.Max.X - TotalBounds.Min.X) / QuadSizeWorld), WindowStart.X, FullExportResolution.X - 1);
-			WindowEnd.Y = FMath::Clamp(FMath::CeilToInt((CellBounds.Max.Y - TotalBounds.Min.Y) / QuadSizeWorld), WindowStart.Y, FullExportResolution.Y - 1);
-		}
-		else
-		{
-			WindowStart = FIntPoint::ZeroValue;
-			WindowEnd.X = FullExportResolution.X - 1;
-			WindowEnd.Y = FullExportResolution.Y - 1;
-		}
+		const FBox CellBounds = Context->ExecutionSource->GetExecutionState().GetBounds();
+		const FBox& TotalBounds = Context->TotalLandscapeComponentBounds;
+
+		UE_LOG(LogPCG, Verbose, TEXT("PCGUnpackGrassMap: %s, Cell bounds (%.2f, %.2f -> %.2f, %.2f), Total bounds (%.2f, %.2f -> %.2f, %.2f), efficiency %.2f"),
+			*Context->GetExecutionSourceName(),
+			CellBounds.Min.X, CellBounds.Min.Y, CellBounds.Max.X, CellBounds.Max.Y,
+			TotalBounds.Min.X, TotalBounds.Min.Y, TotalBounds.Max.X, TotalBounds.Max.Y,
+			(CellBounds.Max.X - CellBounds.Min.X)* (CellBounds.Max.Y - CellBounds.Min.Y) / ((TotalBounds.Max.X - TotalBounds.Min.X) * (TotalBounds.Max.Y - TotalBounds.Min.Y)));
+
+		WindowStart.X = FMath::Clamp(FMath::Floor((CellBounds.Min.X - TotalBounds.Min.X) / QuadSizeWorld), 0, FullExportResolution.X - 1);
+		WindowStart.Y = FMath::Clamp(FMath::Floor((CellBounds.Min.Y - TotalBounds.Min.Y) / QuadSizeWorld), 0, FullExportResolution.Y - 1);
+		WindowEnd.X = FMath::Clamp(FMath::CeilToInt((CellBounds.Max.X - TotalBounds.Min.X) / QuadSizeWorld), WindowStart.X, FullExportResolution.X - 1);
+		WindowEnd.Y = FMath::Clamp(FMath::CeilToInt((CellBounds.Max.Y - TotalBounds.Min.Y) / QuadSizeWorld), WindowStart.Y, FullExportResolution.Y - 1);
 
 		const FIntPoint OutputResolution(WindowEnd.X - WindowStart.X + 1, WindowEnd.Y - WindowStart.Y + 1);
 
