[PCG] Support transform overrides in Scene Capture node

diff --git a/Engine/Plugins/PCG/Source/PCG/Private/Elements/PCGSceneCapture.cpp b/Engine/Plugins/PCG/Source/PCG/Private/Elements/PCGSceneCapture.cpp
--- a/Engine/Plugins/PCG/Source/PCG/Private/Elements/PCGSceneCapture.cpp
+++ b/Engine/Plugins/PCG/Source/PCG/Private/Elements/PCGSceneCapture.cpp
@@ -20,17 +20,41 @@ namespace PCGSceneCaptureConstants
 	const FName BoundingShapeLabel = TEXT("BoundingShape");
 }
 
+#if WITH_EDITOR
+void UPCGSceneCaptureSettings::PostLoad()
+{
+	Super::PostLoad();
+
+	const UPCGNode* Node = Cast<UPCGNode>(GetOuter());
+	const UPCGPin* BoundingShapePin = Node ? Node->GetInputPin(PCGSceneCaptureConstants::BoundingShapeLabel) : nullptr;
+
+	if (BoundingShapePin && BoundingShapePin->IsConnected())
+	{
+		OrientationMode = EPCGSceneCaptureOrientationMode::FromBoundingShape;
+	}
+}
+#endif // WITH_EDITOR
+
+bool UPCGSceneCaptureSettings::IsPinUsedByNodeExecution(const UPCGPin* InPin) const
+{
+	return OrientationMode == EPCGSceneCaptureOrientationMode::Explicit || (InPin
+		&& (InPin->Properties.Label != GET_MEMBER_NAME_CHECKED(UPCGSceneCaptureSettings, CaptureLocation)
+		&& InPin->Properties.Label != GET_MEMBER_NAME_CHECKED(UPCGSceneCaptureSettings, CaptureRotation)
+		&& InPin->Properties.Label != GET_MEMBER_NAME_CHECKED(UPCGSceneCaptureSettings, CaptureHalfExtents)));
+}
+
 TArray<FPCGPinProperties> UPCGSceneCaptureSettings::InputPinProperties() const
 {
 	TArray<FPCGPinProperties> Properties;
 
-	FPCGPinProperties& BoundsPin = Properties.Emplace_GetRef(
-		PCGSceneCaptureConstants::BoundingShapeLabel,
-		EPCGDataType::Spatial,
-		/*bInAllowMultipleConnections=*/true,
-		/*bAllowMultipleData=*/false,
-		LOCTEXT("BoundingShapePinTooltip", "Optional bounds to use instead of the PCG generation bounds. Note, bound overrides will always be axis aligned and top down."));
-	BoundsPin.SetAdvancedPin();
+	if (OrientationMode == EPCGSceneCaptureOrientationMode::FromBoundingShape)
+	{
+		FPCGPinProperties& BoundsPin = Properties.Emplace_GetRef(PCGSceneCaptureConstants::BoundingShapeLabel, EPCGDataType::Spatial);
+		BoundsPin.SetRequiredPin();
+#if WITH_EDITOR
+		BoundsPin.Tooltip = LOCTEXT("BoundingShapePinTooltip", "Optional bounding shape to use instead of the PCG actor or explicit transform. Note, bounding shape overrides will always be an axis aligned bounding box with top down projection (orientation is ignored).");
+#endif
+	}
 
 	return Properties;
 }
@@ -48,6 +72,20 @@ FPCGElementPtr UPCGSceneCaptureSettings::CreateElement() const
 	return MakeShared<FPCGSceneCaptureElement>();
 }
 
+#if WITH_EDITOR
+EPCGChangeType UPCGSceneCaptureSettings::GetChangeTypeForProperty(const FName& InPropertyName) const
+{
+	EPCGChangeType ChangeType = Super::GetChangeTypeForProperty(InPropertyName);
+
+	if (InPropertyName == GET_MEMBER_NAME_CHECKED(UPCGSceneCaptureSettings, OrientationMode))
+	{
+		ChangeType |= EPCGChangeType::Structural;
+	}
+
+	return ChangeType;
+}
+#endif // WITH_EDITOR
+
 void FPCGSceneCaptureContext::AddExtraStructReferencedObjects(FReferenceCollector& Collector)
 {
 	Collector.AddReferencedObject(SceneCaptureComponent);
@@ -83,44 +121,58 @@ bool FPCGSceneCaptureElement::ExecuteInternal(FPCGContext* InContext) const
 			return true;
 		}
 
-		FTransform CaptureTransform = FTransform::Identity;
-		FBox CaptureBounds(-FVector(1.0), FVector(1.0));
-		bool bUnionWasCreated = false;
+		// Default values in case Settings->OrientationMode == EPCGSceneCaptureOrientationMode::Explicit
+		FVector CaptureLocation = Settings->CaptureLocation;
+		FQuat CaptureRotation = Settings->CaptureRotation;
+		FVector CaptureHalfExtents = Settings->CaptureHalfExtents;
 
-		if (const UPCGSpatialData* BoundingShape = Context->InputData.GetSpatialUnionOfInputsByPin(Context, PCGSceneCaptureConstants::BoundingShapeLabel, bUnionWasCreated))
+		if (Settings->OrientationMode == EPCGSceneCaptureOrientationMode::FromBoundingShape)
 		{
-			// Spatial data does not have a transform, so scene captures through an overridden bounds must be top down and axis aligned.
-			CaptureBounds = BoundingShape->GetBounds();
+			bool bUnionWasCreated = false;
+
+			if (const UPCGSpatialData* BoundingShape = Context->InputData.GetSpatialUnionOfInputsByPin(Context, PCGSceneCaptureConstants::BoundingShapeLabel, bUnionWasCreated))
+			{
+				BoundingShape->GetBounds().GetCenterAndExtents(CaptureLocation, CaptureHalfExtents);
+
+				// Spatial data does not have a transform, so scene captures through an overridden bounds must be top down and axis aligned.
+				CaptureRotation = FQuat::Identity;
+			}
+			else
+			{
+				PCGLog::LogErrorOnGraph(FText::Format(LOCTEXT("MissingBoundingShape", "No spatial data found on pin '{0}'."), FText::FromName(PCGSceneCaptureConstants::BoundingShapeLabel)), InContext);
+				return true;
+			}
 		}
-		else
+		else if (Settings->OrientationMode == EPCGSceneCaptureOrientationMode::FromExecutionSource)
 		{
 			const FBox OriginalActorLocalBounds = ExecutionState.GetOriginalLocalSpaceBounds();
-			CaptureTransform = ExecutionState.GetTransform();
-			CaptureTransform.SetScale3D(CaptureTransform.GetScale3D() * 0.5 * (OriginalActorLocalBounds.Max - OriginalActorLocalBounds.Min));
-			CaptureBounds = CaptureBounds.TransformBy(CaptureTransform);
+			const FTransform ActorTransform = ExecutionState.GetTransform();
+
+			CaptureLocation = ActorTransform.GetLocation();
+			CaptureRotation = ActorTransform.GetRotation();
+			CaptureHalfExtents = ActorTransform.GetScale3D() * 0.5 * (OriginalActorLocalBounds.Max - OriginalActorLocalBounds.Min);
 		}
 
-		if (!CaptureBounds.IsValid || CaptureBounds.GetVolume() <= UE_DOUBLE_SMALL_NUMBER)
+		if (CaptureHalfExtents.Size() <= UE_DOUBLE_SMALL_NUMBER)
 		{
-			PCGLog::LogErrorOnGraph(LOCTEXT("InvalidCaptureBounds", "Invalid capture bounds. Volume must be greater than zero."), InContext);
+			PCGLog::LogErrorOnGraph(LOCTEXT("InvalidCaptureExtents", "Invalid capture extents. Volume must be greater than zero."), InContext);
 			return true;
 		}
 
-		const FVector CaptureBoundsLocation = CaptureBounds.GetCenter();
-		const FQuat CaptureBoundsRotation = CaptureTransform.GetRotation().GetNormalized(); // Zero quat will return identity when normalized.
-		const FVector CaptureBoundsSize = CaptureBounds.GetSize();
-		const FVector SceneCaptureLocation = CaptureBoundsLocation + CaptureBoundsRotation.GetUpVector() * (CaptureBoundsSize.Z / 2.0);
-		const FQuat SceneCaptureRotation = FQuat(CaptureBoundsRotation.GetUpVector(), FMath::DegreesToRadians(-90.0f)) * FQuat(CaptureBoundsRotation.GetRightVector(), FMath::DegreesToRadians(90.0f)) * CaptureBoundsRotation;
-		const float OrthoWidth = FMath::Max(CaptureBoundsSize.X, CaptureBoundsSize.Y);
+		CaptureRotation = CaptureRotation.GetNormalized(); // Zero quat will return identity when normalized.
+		const FQuat CaptureComponentRotation = FQuat(CaptureRotation.GetUpVector(), FMath::DegreesToRadians(-90.0f)) * FQuat(CaptureRotation.GetRightVector(), FMath::DegreesToRadians(90.0f)) * CaptureRotation;
 
-		Context->RenderTargetTransform = FTransform(
-			CaptureBoundsRotation,
-			CaptureBoundsLocation,
-			FVector(OrthoWidth / 2.0f));
+		if (Settings->bCaptureFromTopOfExtents)
+		{
+			CaptureLocation += CaptureRotation.GetUpVector() * CaptureHalfExtents.Z;
+		}
 
+		const float OrthoWidth = FMath::Max(CaptureHalfExtents.X, CaptureHalfExtents.Y) * 2.0f;
 		const int32 MaxTextureDimension = GetMax2DTextureDimension();
 		const int32 RenderTargetDimensions = FMath::RoundToInt(OrthoWidth / Settings->TexelSize);
 
+		Context->RenderTargetTransform = FTransform(CaptureRotation, CaptureLocation, FVector(OrthoWidth / 2.0, OrthoWidth / 2.0, CaptureHalfExtents.Z));
+
 		if (RenderTargetDimensions <= 0 || RenderTargetDimensions > MaxTextureDimension)
 		{
 			PCGLog::LogErrorOnGraph(
@@ -154,8 +206,8 @@ bool FPCGSceneCaptureElement::ExecuteInternal(FPCGContext* InContext) const
 		Context->SceneCaptureComponent->RegisterComponentWithWorld(World);
 		Context->SceneCaptureComponent->ProjectionType = ECameraProjectionMode::Orthographic;
 		Context->SceneCaptureComponent->OrthoWidth = OrthoWidth;
-		Context->SceneCaptureComponent->SetWorldLocation(SceneCaptureLocation);
-		Context->SceneCaptureComponent->SetWorldRotation(SceneCaptureRotation);
+		Context->SceneCaptureComponent->SetWorldLocation(CaptureLocation);
+		Context->SceneCaptureComponent->SetWorldRotation(CaptureComponentRotation);
 
 		// Gather actors and primitives for inclusion/exclusion filtering.
 		{
diff --git a/Engine/Plugins/PCG/Source/PCG/Public/Elements/PCGSceneCapture.h b/Engine/Plugins/PCG/Source/PCG/Public/Elements/PCGSceneCapture.h
--- a/Engine/Plugins/PCG/Source/PCG/Public/Elements/PCGSceneCapture.h
+++ b/Engine/Plugins/PCG/Source/PCG/Public/Elements/PCGSceneCapture.h
@@ -11,6 +11,14 @@
 
 #include "PCGSceneCapture.generated.h"
 
+UENUM()
+enum class EPCGSceneCaptureOrientationMode : uint8
+{
+	FromExecutionSource UMETA(ToolTip="Derive scene capture orientation and extents from the PCG execution source."),
+	FromBoundingShape UMETA(ToolTip="Override scene capture orientation and extents using a spatial bounding shape. Note, bounding shape overrides will always become an axis aligned bounding box with top down projection (rotation is ignored)."),
+	Explicit UMETA(ToolTip="Control scene capture orientation and extents using explicit values."),
+};
+
 /** Perform a 2D orthographic scene capture and write the result to a render target data. Can be costly, use with caution with runtime generation. */
 UCLASS(MinimalAPI, BlueprintType, ClassGroup = (Procedural))
 class UPCGSceneCaptureSettings : public UPCGSettings
@@ -18,6 +26,12 @@ class UPCGSceneCaptureSettings : public UPCGSettings
 	GENERATED_BODY()
 
 public:
+#if WITH_EDITOR
+	//~Begin UObject interface
+	virtual void PostLoad() override;
+	//~End UObject interface
+#endif
+
 	//~Begin UPCGSettings interface
 #if WITH_EDITOR
 	virtual FName GetDefaultNodeName() const override { return FName(TEXT("SceneCapture")); }
@@ -25,11 +39,15 @@ public:
 	virtual FText GetNodeTooltipText() const override { return NSLOCTEXT("PCGSceneCaptureElement", "NodeTooltip", "Perform a 2D orthographic scene capture and write the result to a render target data. Can be costly, use with caution with runtime generation."); }
 	virtual EPCGSettingsType GetType() const override { return EPCGSettingsType::Sampler; }
 #endif
+	virtual bool IsPinUsedByNodeExecution(const UPCGPin* InPin) const override;
 
 protected:
 	virtual TArray<FPCGPinProperties> InputPinProperties() const override;
 	virtual TArray<FPCGPinProperties> OutputPinProperties() const override;
 	virtual FPCGElementPtr CreateElement() const override;
+#if WITH_EDITOR
+	virtual EPCGChangeType GetChangeTypeForProperty(const FName& InPropertyName) const override;
+#endif
 	//~End UPCGSettings interface
 
 public:
@@ -42,21 +60,41 @@ public:
 	TEnumAsByte<enum ESceneCaptureSource> CaptureSource = ESceneCaptureSource::SCS_SceneColorHDR;
 
 	/** Only capture actors and primitive components that have tags listed in Included Tags. */
-	UPROPERTY(BlueprintReadWrite, EditAnywhere, Category = "Settings", meta = (PCG_Overridable))
+	UPROPERTY(BlueprintReadWrite, EditAnywhere, Category = "Settings|Capture Filtering", meta = (PCG_Overridable))
 	bool bOnlyCaptureContentMatchingIncludedTags = false;
 
 	/** Tags to match against actors/components when selecting for inclusion in the scene capture. */
-	UPROPERTY(BlueprintReadWrite, EditAnywhere, Category = "Settings", meta = (EditCondition = "bOnlyCaptureContentMatchingIncludedTags", EditConditionHides, PCG_Overridable))
+	UPROPERTY(BlueprintReadWrite, EditAnywhere, Category = "Settings|Capture Filtering", meta = (EditCondition = "bOnlyCaptureContentMatchingIncludedTags", EditConditionHides, PCG_Overridable))
 	TArray<FName> IncludedTags;
 
 	/** Tags to match against actors/components when selecting for exclusion from the scene capture. */
-	UPROPERTY(BlueprintReadWrite, EditAnywhere, Category = "Settings", meta = (PCG_Overridable))
+	UPROPERTY(BlueprintReadWrite, EditAnywhere, Category = "Settings|Capture Filtering", meta = (PCG_Overridable))
 	TArray<FName> ExcludedTags;
 
 	/** Exclude content created from PCG. */
-	UPROPERTY(BlueprintReadWrite, EditAnywhere, Category = "Settings", DisplayName = "Exclude PCG Content", meta = (PCG_Overridable))
+	UPROPERTY(BlueprintReadWrite, EditAnywhere, Category = "Settings|Capture Filtering", DisplayName = "Exclude PCG Content", meta = (PCG_Overridable))
 	bool bExcludePCGContent = true;
 
+	/** Defines how the scene capture bounds and transform are computed. */
+	UPROPERTY(BlueprintReadWrite, EditAnywhere, Category = Settings, meta=(PCG_Overridable))
+	EPCGSceneCaptureOrientationMode OrientationMode = EPCGSceneCaptureOrientationMode::FromExecutionSource;
+
+	/** Positions the scene capture. */
+	UPROPERTY(BlueprintReadWrite, EditAnywhere, Category = Settings, meta = (EditCondition = "OrientationMode==EPCGSceneCaptureOrientationMode::Explicit", EditConditionHides, PCG_Overridable))
+	FVector CaptureLocation = FVector::ZeroVector;
+
+	/** Rotates the scene capture. */
+	UPROPERTY(BlueprintReadWrite, EditAnywhere, Category = Settings, meta = (EditCondition = "OrientationMode==EPCGSceneCaptureOrientationMode::Explicit", EditConditionHides, PCG_Overridable))
+	FQuat CaptureRotation = FQuat::Identity;
+
+	/** Used to determine the orthographic width of the scene capture. Extents on Z are used for the resulting texture data's Z extents. */
+	UPROPERTY(BlueprintReadWrite, EditAnywhere, Category = Settings, meta = (EditCondition = "OrientationMode==EPCGSceneCaptureOrientationMode::Explicit", EditConditionHides, PCG_Overridable))
+	FVector CaptureHalfExtents = FVector(1000.0);
+
+	/** Scene capture position will be translated such that the capture is taken from the top of the capture extents rather than from the center. */
+	UPROPERTY(BlueprintReadWrite, EditAnywhere, Category = Settings, meta = (PCG_Overridable))
+	bool bCaptureFromTopOfExtents = true;
+
 	/** Size of a texel in the render target in world units (cm). */
 	UPROPERTY(BlueprintReadWrite, EditAnywhere, Category = Settings, meta=(UIMin = "1.0", ClampMin = "1.0", PCG_Overridable))
 	float TexelSize = 50.0f;
