[PCG] Add profiling collumn to get the number of tasks for a given subgraph

#rnx

diff --git a/Engine/Plugins/PCG/Source/PCG/Private/Utils/PCGExtraCapture.cpp b/Engine/Plugins/PCG/Source/PCG/Private/Utils/PCGExtraCapture.cpp
--- a/Engine/Plugins/PCG/Source/PCG/Private/Utils/PCGExtraCapture.cpp
+++ b/Engine/Plugins/PCG/Source/PCG/Private/Utils/PCGExtraCapture.cpp
@@ -122,6 +122,7 @@ namespace PCGUtils
 		if (!Info.Children.IsEmpty())
 		{
 			Info.CallTime.OutputDataMemorySize = 0;
+			Info.NumberOfTasks = 0;
 		}
 
 		for (FCallTreeInfo& Child : Info.Children)
@@ -141,6 +142,8 @@ namespace PCGUtils
 			Info.CallTime.MaxExecutionFrameTime = FMath::Max(Info.CallTime.MaxExecutionFrameTime, Child.CallTime.MaxExecutionFrameTime);
 
 			Info.CallTime.OutputDataMemorySize += Child.CallTime.OutputDataMemorySize;
+
+			Info.NumberOfTasks += Child.NumberOfTasks;
 		}
 	}
 }
diff --git a/Engine/Plugins/PCG/Source/PCG/Public/Utils/PCGExtraCapture.h b/Engine/Plugins/PCG/Source/PCG/Public/Utils/PCGExtraCapture.h
--- a/Engine/Plugins/PCG/Source/PCG/Public/Utils/PCGExtraCapture.h
+++ b/Engine/Plugins/PCG/Source/PCG/Public/Utils/PCGExtraCapture.h
@@ -55,6 +55,9 @@ namespace PCGUtils
 		FString Name; // overriden name for the task, will take precedence over the node name if not empty
 		FCallTime CallTime;
 
+		// Store the number of tasks executed within this call.
+		int32 NumberOfTasks = 1;
+
 		TArray<FCallTreeInfo> Children;
 	};
 
diff --git a/Engine/Plugins/PCG/Source/PCGEditor/Private/Widgets/SPCGEditorGraphProfilingView.cpp b/Engine/Plugins/PCG/Source/PCGEditor/Private/Widgets/SPCGEditorGraphProfilingView.cpp
--- a/Engine/Plugins/PCG/Source/PCGEditor/Private/Widgets/SPCGEditorGraphProfilingView.cpp
+++ b/Engine/Plugins/PCG/Source/PCGEditor/Private/Widgets/SPCGEditorGraphProfilingView.cpp
@@ -13,6 +13,7 @@
 #include "HAL/PlatformApplicationMisc.h"
 #include "Widgets/Input/SButton.h"
 #include "Widgets/Input/SCheckBox.h"
+#include "Widgets/Input/SComboButton.h"
 #include "Widgets/Input/SSearchBox.h"
 #include "Widgets/Input/SSpinBox.h"
 #include "Widgets/Layout/SScrollBox.h"
@@ -37,6 +38,7 @@ namespace PCGEditorGraphProfilingView
 	const FName NAME_TotalTime = FName(TEXT("TotalTime"));
 	const FName NAME_TotalWallTime = FName(TEXT("Total_WallTime"));
 	const FName NAME_MemoryInMB = FName(TEXT("Data size (MB)"));
+	const FName NAME_NumberOfTasks = FName(TEXT("NumberOfTasks"));
 
 	/** Labels of the columns */
 	const FText TEXT_NodeLabel = LOCTEXT("NodeLabel", "Node");
@@ -51,6 +53,7 @@ namespace PCGEditorGraphProfilingView
 	const FText TEXT_TotalTimeLabel = LOCTEXT("TotalTimeLabel", "Total");
 	const FText TEXT_TotalWallTimeLabel = LOCTEXT("TotalWallTimeLabel", "Total WallTime");
 	const FText TEXT_MemoryInMBLabel = LOCTEXT("MemoryInMBLabel", "Data size (MB)");
+	const FText TEXT_NumberOfTasks = LOCTEXT("NumberOfTasksLabel", "Num Subtasks");
 
 	/** Tooltips */
 	const FText TEXT_PrepareDataTimeTooltip = LOCTEXT("PrepareDataTimeTooltip", "Cost of the PrepareData execution phase which some nodes use to process the incoming data, in ms.");
@@ -64,6 +67,77 @@ namespace PCGEditorGraphProfilingView
 	const FText TEXT_TotalTimeTooltip = LOCTEXT("TotalTimeTooltip", "The total time spent in this node, summed over all execution and prepare data frames, in ms.");
 	const FText TEXT_TotalWallTimeTooltip = LOCTEXT("TotalWallTimeTooltip", "Total real time elapsed between the first call until completion, including any wait/sleep time, in ms.");
 	const FText TEXT_MemoryInMBTooltip = LOCTEXT("MemoryInMBTooltip", "Size in MB of newly created data.");
+	const FText TEXT_NumberOfTasksTooltip = LOCTEXT("NumberOfTasksTooltip", "Number of subtasks executed for this subgraph. It's an approximation to give an order of magnitude. Only show a value for subgraph nodes.");
+}
+
+FText FPCGProfilingTotalData::GetTotalTimeLabel() const
+{
+	return FText::Format(LOCTEXT("GraphTotalTimeLabel", "Total Time: {0} s"), TotalTime);
+}
+
+FText FPCGProfilingTotalData::GetTotalWallTimeLabel() const
+{
+	return FText::Format(LOCTEXT("GraphTotalWallTimeLabel", "Total Wall Time: {0} s"), TotalWallTime);
+}
+
+FText FPCGProfilingTotalData::GetTotalDataSizeInMBLabel() const
+{
+	return FText::Format(LOCTEXT("GraphTotalDataSizeInMBLabel", "Total data size (MB): {0} MB"), TotalDataSizeInMB);
+}
+
+FText FPCGProfilingTotalData::GetTotalNumTasksLabel() const
+{
+	return FText::Format(LOCTEXT("GraphTotalDataSizeInMBLabel", "Total Num Tasks: {0}"), TotalNumTasks);
+}
+
+TSharedRef<SWidget> FPCGProfilingTotalData::MakeOptionsWidget()
+{
+	TSharedPtr<SComboButton> OptionsWidget = SNew(SComboButton)
+	.ForegroundColor(FSlateColor::UseStyle())
+	.HasDownArrow(false)
+	.OnGetMenuContent(FOnGetContent::CreateLambda([this](){ return OnOptionsMenu();}))
+	.ContentPadding(1)
+	.ButtonContent()
+	[
+		SNew(SImage)
+		.Image(FAppStyle::GetBrush("EditorViewportToolBar.OptionsDropdown"))
+		.ColorAndOpacity(FSlateColor::UseForeground())
+	];
+
+	return OptionsWidget.ToSharedRef();
+}
+
+TSharedRef<SWidget> FPCGProfilingTotalData::OnOptionsMenu()
+{
+	FMenuBuilder MenuBuilder(false, nullptr);
+
+	auto AddEntry = [&MenuBuilder](FText Text, bool* bValue)
+	{
+		MenuBuilder.AddMenuEntry(
+			MoveTemp(Text),
+			FText{},
+			FSlateIcon(),
+			FUIAction(FExecuteAction::CreateLambda([bValue]() { *bValue = !*bValue;}),
+				FCanExecuteAction(),
+				FIsActionChecked::CreateLambda([bValue](){ return *bValue;})),
+			NAME_None,
+			EUserInterfaceActionType::ToggleButton);
+	};
+
+	AddEntry(LOCTEXT("TotalTimeOption", "Show Total Time"), &bShowTotalTime);
+	AddEntry(LOCTEXT("TotalTimeOption", "Show Total Wall Time"), &bShowTotalWallTime);
+	AddEntry(LOCTEXT("TotalTimeOption", "Show Total Data Size"), &bShowTotalDataSize);
+	AddEntry(LOCTEXT("TotalTimeOption", "Show Total Num Tasks"), &bShowTotalTasks);
+
+	return MenuBuilder.MakeWidget();
+}
+
+void FPCGProfilingTotalData::Reset()
+{
+	TotalNumTasks = 0;
+	TotalDataSizeInMB = 0.0;
+	TotalTime = 0;
+	TotalWallTime = 0;
 }
 
 void SPCGProfilingListViewItemRow::Construct(const FArguments& InArgs, const TSharedRef<STableViewBase>& InOwnerTableView, const PCGProfilingListViewItemPtr& Item)
@@ -162,6 +236,11 @@ FText FPCGProfilingListViewItem::GetTextForColumn(FName ColumnId, bool bNoGroupi
 		// In MB
 		return FText::AsNumber(CallTime.OutputDataMemorySize / (1024.0 * 1024.0), NumberFormattingOptions);
 	}
+	else if (ColumnId == PCGEditorGraphProfilingView::NAME_NumberOfTasks)
+	{
+		// Do not show the number of tasks if it is not higher than 1
+		return NumberOfTasks > 1 ? FText::AsNumber(NumberOfTasks) : FText{};
+	}
 	else
 	{
 		return LOCTEXT("ItemColumnError", "Unrecognized Column");
@@ -299,16 +378,7 @@ void SPCGEditorGraphProfilingView::Construct(const FArguments& InArgs, TSharedPt
 			.VAlign(EVerticalAlignment::VAlign_Center)
 			.AutoWidth()
 			[
-				SNew(STextBlock)
-				.Text(LOCTEXT("TotalExecutionTime", "Total Time:"))
-			]
-			+SHorizontalBox::Slot()
-			.VAlign(EVerticalAlignment::VAlign_Center)
-			.AutoWidth()
-			[
-				SNew(STextBlock)
-				.Text(this, &SPCGEditorGraphProfilingView::GetTotalTimeLabel)
-				.MinDesiredWidth(50.0f)
+				ProfilingTotalData.MakeOptionsWidget()
 			]
 			+SHorizontalBox::Slot()
 			.Padding(FMargin(30.0f, 0.0f, 2.0f, 0.0f))
@@ -316,15 +386,17 @@ void SPCGEditorGraphProfilingView::Construct(const FArguments& InArgs, TSharedPt
 			.AutoWidth()
 			[
 				SNew(STextBlock)
-				.Text(LOCTEXT("TotalExecutionWallTime", "Total Wall Time:"))
+				.Text_Lambda([this]() { return ProfilingTotalData.GetTotalTimeLabel();})
+				.Visibility_Lambda([this]() { return ProfilingTotalData.bShowTotalTime ? EVisibility::Visible : EVisibility::Collapsed;})
 			]
 			+SHorizontalBox::Slot()
+			.Padding(FMargin(30.0f, 0.0f, 2.0f, 0.0f))
 			.VAlign(EVerticalAlignment::VAlign_Center)
 			.AutoWidth()
 			[
 				SNew(STextBlock)
-				.Text(this, &SPCGEditorGraphProfilingView::GetTotalWallTimeLabel)
-				.MinDesiredWidth(50.0f)
+				.Text_Lambda([this]() { return ProfilingTotalData.GetTotalWallTimeLabel();})
+				.Visibility_Lambda([this]() { return ProfilingTotalData.bShowTotalWallTime ? EVisibility::Visible : EVisibility::Collapsed;})
 			]
 			+SHorizontalBox::Slot()
 			.Padding(FMargin(30.0f, 0.0f, 2.0f, 0.0f))
@@ -332,15 +404,17 @@ void SPCGEditorGraphProfilingView::Construct(const FArguments& InArgs, TSharedPt
 			.AutoWidth()
 			[
 				SNew(STextBlock)
-				.Text(LOCTEXT("TotalMemoryInM", "Total data size (MB):"))
+				.Text_Lambda([this]() { return ProfilingTotalData.GetTotalDataSizeInMBLabel();})
+				.Visibility_Lambda([this]() { return ProfilingTotalData.bShowTotalDataSize ? EVisibility::Visible : EVisibility::Collapsed;})
 			]
 			+SHorizontalBox::Slot()
+			.Padding(FMargin(30.0f, 0.0f, 2.0f, 0.0f))
 			.VAlign(EVerticalAlignment::VAlign_Center)
 			.AutoWidth()
 			[
 				SNew(STextBlock)
-				.Text(this, &SPCGEditorGraphProfilingView::GetTotalDataSizeInMB)
-				.MinDesiredWidth(50.0f)
+				.Text_Lambda([this]() { return ProfilingTotalData.GetTotalNumTasksLabel();})
+				.Visibility_Lambda([this]() { return ProfilingTotalData.bShowTotalTasks ? EVisibility::Visible : EVisibility::Collapsed;})
 			]
 		]
 		+SVerticalBox::Slot()
@@ -499,22 +573,16 @@ TSharedRef<SHeaderRow> SPCGEditorGraphProfilingView::CreateHeaderRowWidget()
 		.SortMode(this, &SPCGEditorGraphProfilingView::GetColumnSortMode, PCGEditorGraphProfilingView::NAME_MemoryInMB)
 		.OnSort(this, &SPCGEditorGraphProfilingView::OnSortColumnHeader)
 		.InitialSortMode(EColumnSortMode::Descending)
-		.DefaultTooltip(PCGEditorGraphProfilingView::TEXT_MemoryInMBTooltip);
-}
-
-FText SPCGEditorGraphProfilingView::GetTotalTimeLabel() const
-{
-	return FText::Format(LOCTEXT("GraphTotalTimeLabel", "{0} s"), TotalTime);
-}
-
-FText SPCGEditorGraphProfilingView::GetTotalWallTimeLabel() const
-{
-	return FText::Format(LOCTEXT("GraphTotalWallTimeLabel", "{0} s"), TotalWallTime);
-}
-
-FText SPCGEditorGraphProfilingView::GetTotalDataSizeInMB() const
-{
-	return FText::Format(LOCTEXT("GraphTotalDataSizeInMBLabel", "{0} MB"), TotalDataSizeInMB);
+		.DefaultTooltip(PCGEditorGraphProfilingView::TEXT_MemoryInMBTooltip)
+		+ SHeaderRow::Column(PCGEditorGraphProfilingView::NAME_NumberOfTasks)
+		.ManualWidth(110.0f)
+		.DefaultLabel(PCGEditorGraphProfilingView::TEXT_NumberOfTasks)
+		.HAlignHeader(HAlign_Center)
+		.HAlignCell(HAlign_Right)
+		.SortMode(this, &SPCGEditorGraphProfilingView::GetColumnSortMode, PCGEditorGraphProfilingView::NAME_NumberOfTasks)
+		.OnSort(this, &SPCGEditorGraphProfilingView::OnSortColumnHeader)
+		.InitialSortMode(EColumnSortMode::Descending)
+		.DefaultTooltip(PCGEditorGraphProfilingView::TEXT_NumberOfTasksTooltip);
 }
 
 void SPCGEditorGraphProfilingView::OnSortColumnHeader(const EColumnSortPriority::Type SortPriority, const FName& ColumnId, const EColumnSortMode::Type NewSortMode)
@@ -665,6 +733,7 @@ namespace PCGEditorGraphProfilingView
 				ListViewItem->bHasData = (Info.CallTime.MaxExecutionFrameTime > 0);
 
 				ListViewItem->CallTime = Info.CallTime;
+				ListViewItem->NumberOfTasks = Info.NumberOfTasks;
 
 				OutListViewItems.Add(ListViewItem);
 			}
@@ -684,9 +753,7 @@ void SPCGEditorGraphProfilingView::RequestRefresh()
 
 FReply SPCGEditorGraphProfilingView::Refresh()
 {
-	TotalTime = 0;
-	TotalWallTime = 0;
-	TotalDataSizeInMB = 0;
+	ProfilingTotalData.Reset();
 
 	ListViewItems.Empty();
 	ListView->RequestListRefresh();
@@ -723,9 +790,10 @@ FReply SPCGEditorGraphProfilingView::Refresh()
 
 	if (TreeInfo.Children.Num() > 0)
 	{
-		TotalTime = TreeInfo.CallTime.TotalTime();
-		TotalWallTime = TreeInfo.CallTime.TotalWallTime();
-		TotalDataSizeInMB = TreeInfo.CallTime.OutputDataMemorySize / (1024.0 * 1024.0);
+		ProfilingTotalData.TotalTime = TreeInfo.CallTime.TotalTime();
+		ProfilingTotalData.TotalWallTime = TreeInfo.CallTime.TotalWallTime();
+		ProfilingTotalData.TotalDataSizeInMB = TreeInfo.CallTime.OutputDataMemorySize / (1024.0 * 1024.0);
+		ProfilingTotalData.TotalNumTasks = TreeInfo.NumberOfTasks;
 	}
 
 	//TODO: could turn this into a tree instead of expanding into a list
@@ -784,6 +852,10 @@ FReply SPCGEditorGraphProfilingView::Refresh()
 				{
 					isLess = A->CallTime.OutputDataMemorySize < B->CallTime.OutputDataMemorySize;
 				}
+				else if (SortingColumn == PCGEditorGraphProfilingView::NAME_NumberOfTasks)
+				{
+					isLess = A->NumberOfTasks < B->NumberOfTasks;
+				}
 
 				return SortMode == EColumnSortMode::Ascending ? isLess : !isLess;
 			});
diff --git a/Engine/Plugins/PCG/Source/PCGEditor/Private/Widgets/SPCGEditorGraphProfilingView.h b/Engine/Plugins/PCG/Source/PCGEditor/Private/Widgets/SPCGEditorGraphProfilingView.h
--- a/Engine/Plugins/PCG/Source/PCGEditor/Private/Widgets/SPCGEditorGraphProfilingView.h
+++ b/Engine/Plugins/PCG/Source/PCGEditor/Private/Widgets/SPCGEditorGraphProfilingView.h
@@ -26,10 +26,34 @@ struct FPCGProfilingListViewItem
 	FString Name;
 
 	PCGUtils::FCallTime CallTime;
+	int32 NumberOfTasks = 0;
 
 	bool bHasData = false;
 };
 
+struct FPCGProfilingTotalData
+{
+	bool bShowTotalTime = true;
+	bool bShowTotalWallTime = true;
+	bool bShowTotalDataSize = true;
+	bool bShowTotalTasks = true;
+	
+	int32 TotalNumTasks = 0;
+	double TotalTime = 0.0;
+	double TotalWallTime = 0.0;
+	double TotalDataSizeInMB = 0.0;
+
+	FText GetTotalTimeLabel() const;
+	FText GetTotalWallTimeLabel() const;
+	FText GetTotalDataSizeInMBLabel() const;
+	FText GetTotalNumTasksLabel() const;
+
+	TSharedRef<SWidget> MakeOptionsWidget();
+	TSharedRef<SWidget> OnOptionsMenu();
+
+	void Reset();
+};
+
 typedef TSharedPtr<FPCGProfilingListViewItem> PCGProfilingListViewItemPtr;
 
 class SPCGProfilingListViewItemRow : public SMultiColumnTableRow<PCGProfilingListViewItemPtr>
@@ -72,9 +96,6 @@ private:
 	void OnItemDoubleClicked(PCGProfilingListViewItemPtr Item);
 	void OnSortColumnHeader(const EColumnSortPriority::Type SortPriority, const FName& ColumnId, const EColumnSortMode::Type NewSortMode);
 	EColumnSortMode::Type GetColumnSortMode(const FName ColumnId) const;
-	FText GetTotalTimeLabel() const;
-	FText GetTotalWallTimeLabel() const;
-	FText GetTotalDataSizeInMB() const;
 
 	FReply OnListViewKeyDown(const FGeometry& InGeometry, const FKeyEvent& InKeyEvent) const;
 	void CopySelectionToClipboard() const;
@@ -107,9 +128,7 @@ private:
 	FName SortingColumn = NAME_None;
 	EColumnSortMode::Type SortMode = EColumnSortMode::Type::None;
 
-	double TotalTime = 0.0;
-	double TotalWallTime = 0.0;
-	double TotalDataSizeInMB = 0.0;
+	FPCGProfilingTotalData ProfilingTotalData;
 
 	bool bNeedsRefresh = false;
 
