PCG: exposed bIgnorePCGGeneratedComponents for the get actor data mode Components and PCG Data

#rnx

diff --git a/Engine/Plugins/PCG/Source/PCG/Public/Elements/PCGDataFromActor.h b/Engine/Plugins/PCG/Source/PCG/Public/Elements/PCGDataFromActor.h
--- a/Engine/Plugins/PCG/Source/PCG/Public/Elements/PCGDataFromActor.h
+++ b/Engine/Plugins/PCG/Source/PCG/Public/Elements/PCGDataFromActor.h
@@ -92,7 +92,7 @@ public:
 	EPCGGetDataFromActorMode Mode = EPCGGetDataFromActorMode::ParseActorComponents;
 
 	/** Ignores any component that was spawned by PCG. */
-	UPROPERTY(BlueprintReadWrite, EditAnywhere, Category = "Data Retrieval Settings", meta = (EditCondition = "Mode == EPCGGetDataFromActorMode::ParseActorComponents || Mode == EPCGGetDataFromActorMode::GetComponentsReference", EditConditionHides))
+	UPROPERTY(BlueprintReadWrite, EditAnywhere, Category = "Data Retrieval Settings", meta = (EditCondition = "Mode == EPCGGetDataFromActorMode::ParseActorComponents || Mode == EPCGGetDataFromActorMode::GetComponentsReference || Mode == EPCGGetDataFromActorMode::GetDataFromPCGComponentOrParseComponents", EditConditionHides))
 	bool bIgnorePCGGeneratedComponents = true;
 
 	/** Also produces a single point data at the actor location. */
