[PCG] Don't produce volume data from higen grid size node if bounds output pin is not connected.

[RN] [PCG] Optimized higen grid size node to not produce cell volume data output when output pin is not connected.

diff --git a/Engine/Plugins/PCG/Source/PCG/Private/Elements/PCGHiGenGridSize.cpp b/Engine/Plugins/PCG/Source/PCG/Private/Elements/PCGHiGenGridSize.cpp
--- a/Engine/Plugins/PCG/Source/PCG/Private/Elements/PCGHiGenGridSize.cpp
+++ b/Engine/Plugins/PCG/Source/PCG/Private/Elements/PCGHiGenGridSize.cpp
@@ -45,6 +45,21 @@ uint32 UPCGHiGenGridSizeSettings::GetGridSize() const
 	return (HiGenGridSize == EPCGHiGenGrid::Unbounded) ? PCGHiGenGrid::UnboundedGridSize() : PCGHiGenGrid::GridToGridSize(GetGrid());
 }
 
+bool UPCGHiGenGridSizeSettings::ShouldOutputCellVolumeData() const
+{
+#if !WITH_EDITOR
+	if (const UPCGNode* Node = Cast<UPCGNode>(GetOuter()))
+	{
+		return Node->IsOutputPinConnected(PCGHiGenGridSizeConstants::CellVolumeOutputLabel);
+	}
+	else
+#endif
+	{
+		// In editor, we always want to generate this data for inspection & to prevent caching issues.
+		return true;
+	}
+}
+
 #if WITH_EDITOR
 FName UPCGHiGenGridSizeSettings::GetDefaultNodeName() const
 {
@@ -246,13 +261,16 @@ bool FPCGHiGenGridSizeElement::ExecuteInternal(FPCGContext* Context) const
 		Data.bIsUsedMultipleTimes = true;
 	}
 
-	UPCGVolumeData* VolumeData = FPCGContext::NewObject_AnyThread<UPCGVolumeData>(Context);
-	check(VolumeData);
-	VolumeData->Initialize(Context->ExecutionSource->GetExecutionState().GetBounds());
+	if (Settings->ShouldOutputCellVolumeData())
+	{
+		UPCGVolumeData* VolumeData = FPCGContext::NewObject_AnyThread<UPCGVolumeData>(Context);
+		check(VolumeData);
+		VolumeData->Initialize(Context->ExecutionSource->GetExecutionState().GetBounds());
 
-	FPCGTaggedData& VolumeOutput = Context->OutputData.TaggedData.Emplace_GetRef();
-	VolumeOutput.Data = VolumeData;
-	VolumeOutput.Pin = PCGHiGenGridSizeConstants::CellVolumeOutputLabel;
+		FPCGTaggedData& VolumeOutput = Context->OutputData.TaggedData.Emplace_GetRef();
+		VolumeOutput.Data = VolumeData;
+		VolumeOutput.Pin = PCGHiGenGridSizeConstants::CellVolumeOutputLabel;
+	}
 
 	return true;
 }
@@ -262,12 +280,14 @@ void FPCGHiGenGridSizeElement::GetDependenciesCrc(const FPCGGetDependenciesCrcPa
 	FPCGCrc Crc;
 	IPCGElement::GetDependenciesCrc(InParams, Crc);
 
-	// The grid cell volume output depends on the component transform.
-	// NOTE: It might be interesting to only incorporate the transform if the cell volume output pin is connected (and if we have a node obviously).
-	const UPCGData* ActorData = InParams.ExecutionSource ? InParams.ExecutionSource->GetExecutionState().GetSelfData() : nullptr;
-	if (ActorData)
+	if (const UPCGHiGenGridSizeSettings* Settings = Cast<UPCGHiGenGridSizeSettings>(InParams.Settings); Settings && Settings->ShouldOutputCellVolumeData())
 	{
-		Crc.Combine(ActorData->GetOrComputeCrc(/*bFullDataCrc=*/false));
+		// The grid cell volume output depends on the component transform.
+		const UPCGData* ActorData = InParams.ExecutionSource ? InParams.ExecutionSource->GetExecutionState().GetSelfData() : nullptr;
+		if (ActorData)
+		{
+			Crc.Combine(ActorData->GetOrComputeCrc(/*bFullDataCrc=*/false));
+		}
 	}
 
 	OutCrc = Crc;
diff --git a/Engine/Plugins/PCG/Source/PCG/Public/Elements/PCGHiGenGridSize.h b/Engine/Plugins/PCG/Source/PCG/Public/Elements/PCGHiGenGridSize.h
--- a/Engine/Plugins/PCG/Source/PCG/Public/Elements/PCGHiGenGridSize.h
+++ b/Engine/Plugins/PCG/Source/PCG/Public/Elements/PCGHiGenGridSize.h
@@ -18,6 +18,8 @@ public:
 	EPCGHiGenGrid GetGrid() const;
 	uint32 GetGridSize() const;
 
+	bool ShouldOutputCellVolumeData() const;
+
 	//~Begin UPCGSettings interface
 #if WITH_EDITOR
 	virtual FName GetDefaultNodeName() const override;
@@ -32,13 +34,6 @@ public:
 	virtual bool HasDynamicPins() const override { return true; }
 	virtual FPCGDataTypeIdentifier GetCurrentPinTypesID(const UPCGPin* InPin) const override;
 
-	UPROPERTY(BlueprintReadWrite, EditAnywhere, Category = Settings, meta = (DisplayName = "HiGen Grid Size"))
-	EPCGHiGenGrid HiGenGridSize = EPCGHiGenGrid::Grid256;
-
-	/** This property is hidden from the user and drives the behavior of the Grid Size node. */
-	UPROPERTY()
-	bool bShowInputPin = true;
-
 protected:
 #if WITH_EDITOR
 	virtual EPCGChangeType GetChangeTypeForProperty(const FName& InPropertyName) const override;
@@ -50,6 +45,13 @@ protected:
 	virtual TArray<FPCGPinProperties> OutputPinProperties() const override;
 	virtual FPCGElementPtr CreateElement() const override;
 	//~End UPCGSettings interface
+
+	UPROPERTY(BlueprintReadWrite, EditAnywhere, Category = Settings, meta = (DisplayName = "HiGen Grid Size"))
+	EPCGHiGenGrid HiGenGridSize = EPCGHiGenGrid::Grid256;
+
+	/** This property is hidden from the user and drives the behavior of the Grid Size node. */
+	UPROPERTY()
+	bool bShowInputPin = true;
 };
 
 class FPCGHiGenGridSizeElement : public IPCGElement
