[PCG] GPU SM Spawner: Mirror logic used on CPU side to decide whether to produce output points or not.

Mirrors similar logic on CPU side.

Significantly reduces work and memory bandwidth.

[RN] [PCG] Optimized GPU Static Mesh Spawner to not generate output point data if output pin is not connected, matching similar logic in CPU path.

diff --git a/Engine/Plugins/PCG/Shaders/Private/Elements/PCGStaticMeshSpawner.usf b/Engine/Plugins/PCG/Shaders/Private/Elements/PCGStaticMeshSpawner.usf
--- a/Engine/Plugins/PCG/Shaders/Private/Elements/PCGStaticMeshSpawner.usf
+++ b/Engine/Plugins/PCG/Shaders/Private/Elements/PCGStaticMeshSpawner.usf
@@ -5,15 +5,19 @@
 
 // When spawning by attribute, we don't know in advance how many instances of each mesh will be spawned, so the worst case NumPoints * NumPrims is allocated in the GPU Scene.
 
+#define PCG_PRODUCE_OUTPUT_POINTS {ProduceOutputPoints}
+
 [numthreads(64, 1, 1)]
 void PCGStaticMeshSpawnerCS(uint3 GroupId : SV_GroupID, uint GroupIndex : SV_GroupIndex)
 {
+#if PCG_PRODUCE_OUTPUT_POINTS
 	// Mark the kernel as having executed. Must run before we early out via thread index, because the kernel is still 'executed' even if the number of
 	// threads to iterate on is zero. Even if GetNumThreads() returns 0, the kernel will still have been dispatched on a single thread to set this flag.
 	if (all(GroupId == 0) && GroupIndex == 0)
 	{
 		Out_SetAsExecutedInternal();
 	}
+#endif // PCG_PRODUCE_OUTPUT_POINTS
 
 	const uint ThreadIndex = GetUnWrappedDispatchThreadId(GroupId, GroupIndex, 64);
 
@@ -23,17 +27,21 @@ void PCGStaticMeshSpawnerCS(uint3 GroupId : SV_GroupID, uint GroupIndex : SV_Gro
 		return;
 	}
 
+#if PCG_PRODUCE_OUTPUT_POINTS
 	uint Out_DataIndex, Out_ElementIndex;
 	if (!Out_GetThreadData(ThreadIndex, Out_DataIndex, Out_ElementIndex))
 	{
 		// Should never happen.
 		return;
 	}
+#endif // PCG_PRODUCE_OUTPUT_POINTS
 
 	if (In_IsPointRemoved(In_DataIndex, In_ElementIndex))
 	{
+#if PCG_PRODUCE_OUTPUT_POINTS
 		// Propagate the 'Removed' status to the output point if input point is removed. Otherwise, this point will not be culled.
 		Out_RemovePoint(Out_DataIndex, Out_ElementIndex);
+#endif // PCG_PRODUCE_OUTPUT_POINTS
 
 		return;
 	}
@@ -110,6 +118,7 @@ void PCGStaticMeshSpawnerCS(uint3 GroupId : SV_GroupID, uint GroupIndex : SV_Gro
 		}
 	}
 
+#if PCG_PRODUCE_OUTPUT_POINTS
 	// Write output point data (pass through).
 	Out_SetPosition(Out_DataIndex, Out_ElementIndex, In_GetPosition(In_DataIndex, In_ElementIndex));
 	Out_SetRotation(Out_DataIndex, Out_ElementIndex, In_GetRotation(In_DataIndex, In_ElementIndex));
@@ -143,4 +152,7 @@ void PCGStaticMeshSpawnerCS(uint3 GroupId : SV_GroupID, uint GroupIndex : SV_Gro
 
 		Out_SetStringKey(Out_DataIndex, Out_ElementIndex, SMSpawner_GetSelectedMeshAttributeId(), OutputMeshStringKey);
 	}
+#endif // PCG_PRODUCE_OUTPUT_POINTS
 }
+
+#undef PCG_PRODUCE_OUTPUT_POINTS
diff --git a/Engine/Plugins/PCG/Source/PCG/Private/Elements/PCGStaticMeshSpawner.cpp b/Engine/Plugins/PCG/Source/PCG/Private/Elements/PCGStaticMeshSpawner.cpp
--- a/Engine/Plugins/PCG/Source/PCG/Private/Elements/PCGStaticMeshSpawner.cpp
+++ b/Engine/Plugins/PCG/Source/PCG/Private/Elements/PCGStaticMeshSpawner.cpp
@@ -76,8 +76,17 @@ void UPCGStaticMeshSpawnerSettings::CreateKernels(FPCGGPUCompilationContext& InO
 			CreateParams.NodeInputPinsToWire.Add(PCGStaticMeshSpawner::PrimitiveTablePinLabel);
 		}
 
+		const UPCGNode* Node = Cast<UPCGNode>(GetOuter());
+		const bool bProduceOutputPoints = !Node || Node->IsOutputPinConnected(PCGPinConstants::DefaultOutputLabel) || bIsInspecting || bDebug;
+
+		if (!bProduceOutputPoints)
+		{
+			CreateParams.NodeOutputPinsToWire.Empty();
+		}
+
 		SpawnerKernel = PCGKernelHelpers::CreateKernel<UPCGStaticMeshSpawnerKernel>(InOutContext, CreateParams, OutKernels, OutEdges);
 		SpawnerKernel->SetUseRawInstanceCountsBuffer(bUseRawBuffer);
+		SpawnerKernel->SetProduceOutputPoints(bProduceOutputPoints);
 	}
 
 	// If doing by-attribute selection, add analysis kernel that will count how many instances of each mesh are present.
diff --git a/Engine/Plugins/PCG/Source/PCG/Private/Elements/PCGStaticMeshSpawnerKernel.cpp b/Engine/Plugins/PCG/Source/PCG/Private/Elements/PCGStaticMeshSpawnerKernel.cpp
--- a/Engine/Plugins/PCG/Source/PCG/Private/Elements/PCGStaticMeshSpawnerKernel.cpp
+++ b/Engine/Plugins/PCG/Source/PCG/Private/Elements/PCGStaticMeshSpawnerKernel.cpp
@@ -168,6 +168,19 @@ int UPCGStaticMeshSpawnerKernel::ComputeThreadCount(const UPCGDataBinding* InBin
 }
 
 #if WITH_EDITOR
+FString UPCGStaticMeshSpawnerKernel::GetCookedSource(FPCGGPUCompilationContext& InOutContext) const
+{
+	FString SourceFile;
+	ensure(LoadShaderSourceFile(GetSourceFilePath(), EShaderPlatform::SP_PCD3D_SM5, &SourceFile, nullptr));
+
+	TMap<FString, FStringFormatArg> Defines =
+	{
+		{ TEXT("ProduceOutputPoints"), GetProduceOutputPoints() ? 1 : 0 },
+	};
+
+	return FString::Format(*SourceFile, Defines);
+}
+
 void UPCGStaticMeshSpawnerKernel::CreateAdditionalInputDataInterfaces(FPCGGPUCompilationContext& InOutContext, UObject* InObjectOuter, TArray<TObjectPtr<UComputeDataInterface>>& OutDataInterfaces) const
 {
 	Super::CreateAdditionalInputDataInterfaces(InOutContext, InObjectOuter, OutDataInterfaces);
@@ -221,7 +234,10 @@ void UPCGStaticMeshSpawnerKernel::GetInputPins(TArray<FPCGPinProperties>& OutPin
 
 void UPCGStaticMeshSpawnerKernel::GetOutputPins(TArray<FPCGPinPropertiesGPU>& OutPins) const
 {
-	OutPins.Emplace(PCGPinConstants::DefaultOutputLabel, EPCGDataType::Point);
+	if (GetProduceOutputPoints())
+	{
+		OutPins.Emplace(PCGPinConstants::DefaultOutputLabel, EPCGDataType::Point);
+	}
 }
 
 #if WITH_EDITOR
diff --git a/Engine/Plugins/PCG/Source/PCG/Public/Elements/PCGStaticMeshSpawnerKernel.h b/Engine/Plugins/PCG/Source/PCG/Public/Elements/PCGStaticMeshSpawnerKernel.h
--- a/Engine/Plugins/PCG/Source/PCG/Public/Elements/PCGStaticMeshSpawnerKernel.h
+++ b/Engine/Plugins/PCG/Source/PCG/Public/Elements/PCGStaticMeshSpawnerKernel.h
@@ -24,6 +24,7 @@ public:
 	virtual int ComputeThreadCount(const UPCGDataBinding* Binding) const override;
 #if WITH_EDITOR
 	virtual const TCHAR* GetSourceFilePath() const override { return TEXT("/Plugin/PCG/Private/Elements/PCGStaticMeshSpawner.usf"); }
+	virtual FString GetCookedSource(FPCGGPUCompilationContext& InOutContext) const override;
 	virtual FString GetEntryPoint() const override { return TEXT("PCGStaticMeshSpawnerCS"); }
 	void CreateAdditionalInputDataInterfaces(FPCGGPUCompilationContext& InOutContext, UObject* InObjectOuter, TArray<TObjectPtr<UComputeDataInterface>>& OutDataInterfaces) const override;
 	void CreateAdditionalOutputDataInterfaces(FPCGGPUCompilationContext& InOutContext, UObject* InObjectOuter, TArray<TObjectPtr<UComputeDataInterface>>& OutDataInterfaces) const override;
@@ -36,6 +37,9 @@ public:
 
 	void SetUseRawInstanceCountsBuffer(bool bInUseRawInstanceCountsBuffer) { bUseRawInstanceCountsBuffer = bInUseRawInstanceCountsBuffer; }
 
+	bool GetProduceOutputPoints() const { return bProduceOutputPoints; }
+	void SetProduceOutputPoints(bool bInProduceOutputPoints) { bProduceOutputPoints = bInProduceOutputPoints; }
+
 protected:
 #if WITH_EDITOR
 	virtual bool PerformStaticValidation() override;
@@ -45,4 +49,7 @@ protected:
 	/** Consume the instance counts as a raw buffer instead of an attribute set. */
 	UPROPERTY()
 	bool bUseRawInstanceCountsBuffer = false;
+
+	UPROPERTY()
+	bool bProduceOutputPoints = true;
 };
