[PCG] Added support for partial clear PCG link

#Synced-CL 48040603

diff --git a/Engine/Plugins/PCG/Source/PCG/Private/PCGComponent.cpp b/Engine/Plugins/PCG/Source/PCG/Private/PCGComponent.cpp
--- a/Engine/Plugins/PCG/Source/PCG/Private/PCGComponent.cpp
+++ b/Engine/Plugins/PCG/Source/PCG/Private/PCGComponent.cpp
@@ -1174,7 +1174,9 @@ AActor* UPCGComponent::ClearPCGLink(const FPCGMoveResourceParams& InParams)
 	{
 		if (bHasMovedResources)
 		{
-			Cleanup(true);
+#if WITH_EDITOR
+			DirtyGenerated(EPCGComponentDirtyFlag::Input, /*bDispatchToLocalComponents=*/true);
+#endif
 		}
 		else
 		{
@@ -1398,28 +1400,43 @@ bool UPCGComponent::MoveResourcesToNewActor(const FPCGMoveResourceParams& Params
 #if WITH_EDITOR
 		Scope.AddResources(this, ManagedResourceContainer.GetConstManagedResourcesNoLock());
 #endif
-		for (TObjectPtr<UPCGManagedResource>& GeneratedResource : ManagedResourceContainer.GetManagedResourcesNoLock())
+		TArray<TSoftObjectPtr<AActor>> TentativeMovedGeneratedActors;
+		TArray<TObjectPtr<UPCGManagedResource>>& ManagedResources = ManagedResourceContainer.GetManagedResourcesNoLock();
+
+		for(int ResourceIndex = ManagedResources.Num() - 1; ResourceIndex >=0; --ResourceIndex)
 		{
+			TObjectPtr<UPCGManagedResource> GeneratedResource = ManagedResources[ResourceIndex];
+			TentativeMovedGeneratedActors.Reset();
+
 			if (GeneratedResource)
 			{
 #if WITH_EDITOR
 				if (UPCGManagedActors* ManagedActors = Cast<UPCGManagedActors>(GeneratedResource))
 				{
-					MovedActors.Append(ManagedActors->GetConstGeneratedActors());
+					TentativeMovedGeneratedActors = ManagedActors->GetConstGeneratedActors();
 				}
 #endif
-				GeneratedResource->MoveResource(PotentiallyRetargetedParams);
-				TSet<TSoftObjectPtr<AActor>> Dummy;
-				GeneratedResource->ReleaseIfUnused(Dummy);
-				bHasMovedResources = true;
+				if (GeneratedResource->MoveResource(PotentiallyRetargetedParams))
+				{
+					TSet<TSoftObjectPtr<AActor>> Dummy;
+					GeneratedResource->ReleaseIfUnused(Dummy);
+					bHasMovedResources = true;
+
+#if WITH_EDITOR
+					if (!TentativeMovedGeneratedActors.IsEmpty())
+					{
+						MovedActors.Append(TentativeMovedGeneratedActors);
+					}
+#endif
+
+					ManagedResources.RemoveAtSwap(ResourceIndex);
+				}
 			}
 			else
 			{
 				UE_LOG(LogPCG, Error, TEXT("[UPCGComponent::MoveResourcesToNewActor] Null generated resource encountered on actor \"%s\" and will be skipped."), *Owner->GetFName().ToString());
 			}
 		}
-
-		ManagedResourceContainer.GetManagedResourcesNoLock().Empty();
 	}
 
 	if (CreatedChildActor)
diff --git a/Engine/Plugins/PCG/Source/PCG/Private/PCGManagedResource.cpp b/Engine/Plugins/PCG/Source/PCG/Private/PCGManagedResource.cpp
--- a/Engine/Plugins/PCG/Source/PCG/Private/PCGManagedResource.cpp
+++ b/Engine/Plugins/PCG/Source/PCG/Private/PCGManagedResource.cpp
@@ -32,6 +32,31 @@ static TAutoConsoleVariable<bool> CVarLoadManagedActorsForCleanup(
 	true,
 	TEXT("When enabled unloaded actors will get loaded before clean up"));
 
+bool FPCGMoveResourceParams::Accepts(const UObject* InObject) const
+{
+	if (!InObject) // Never accept nullptrs
+	{
+		return false;
+	}
+	else if (RequiredTags.IsEmpty()) // Always accept if there are no requirements
+	{
+		return true;
+	}
+	else if (const AActor* Actor = Cast<AActor>(InObject))
+	{
+		return Algo::AllOf(RequiredTags, [Actor](FName Tag) { return Actor->Tags.Contains(Tag); });
+	}
+	else if (const UActorComponent* Component = Cast<UActorComponent>(InObject))
+	{
+		return Algo::AllOf(RequiredTags, [Component](FName Tag) { return Component->ComponentHasTag(Tag); });
+	}
+	else
+	{
+		// Unknown object type, assume we accept it.
+		return true;
+	}
+}
+
 void UPCGManagedResource::PostApplyToComponent()
 {
 	// Nothing - apply to component should already properly remap most of everything we need to do
@@ -267,14 +292,19 @@ bool UPCGManagedActors::MoveResource(const FPCGMoveResourceParams& Params)
 {
 	AActor* TargetActor = Params.GetTarget<AActor>();
 	check(TargetActor);
+	bool bMovedResource = false;
 
-	for (TSoftObjectPtr<AActor>& Actor : GeneratedActorsArray)
+	for(int ActorIndex = GeneratedActorsArray.Num() - 1; ActorIndex >=0; --ActorIndex)
 	{
-		if (!Actor.IsValid())
+		TSoftObjectPtr<AActor> Actor = GeneratedActorsArray[ActorIndex];
+
+		if (!Actor.IsValid() || !Params.Accepts(Actor.Get()))
 		{
 			continue;
 		}
 
+		bMovedResource = true;
+
 		const bool bWasAttached = (Actor->GetAttachParentActor() != nullptr);
 
 		if (bWasAttached)
@@ -283,11 +313,11 @@ bool UPCGManagedActors::MoveResource(const FPCGMoveResourceParams& Params)
 			Actor->SetOwner(nullptr);
 			Actor->AttachToActor(TargetActor, FAttachmentTransformRules::KeepWorldTransform);
 		}
-	}
 
-	GeneratedActorsArray.Empty();
+		GeneratedActorsArray.RemoveAtSwap(ActorIndex);
+	}
 
-	return true;
+	return bMovedResource;
 }
 
 void UPCGManagedActors::MarkAsUsed()
@@ -688,7 +718,7 @@ bool UPCGManagedComponentBase::MoveResource(const FPCGMoveResourceParams& Params
 	{
 		TSoftObjectPtr<UActorComponent> GeneratedComponent = GeneratedComponents[ComponentIndex];
 
-		if (!GeneratedComponent.IsValid())
+		if (!GeneratedComponent.IsValid() || !Params.Accepts(GeneratedComponent.Get()))
 		{
 			continue;
 		}
diff --git a/Engine/Plugins/PCG/Source/PCG/Public/PCGManagedResource.h b/Engine/Plugins/PCG/Source/PCG/Public/PCGManagedResource.h
--- a/Engine/Plugins/PCG/Source/PCG/Public/PCGManagedResource.h
+++ b/Engine/Plugins/PCG/Source/PCG/Public/PCGManagedResource.h
@@ -28,6 +28,7 @@ struct FPCGMoveResourceParams
 	UClass* TemplateTargetClass = nullptr;
 	UObject* Target = nullptr;
 	const UObject* ExpectedPreviousOwner = nullptr;
+	TArray<FName> RequiredTags;
 
 	template<typename ObjectType>
 	TSubclassOf<ObjectType> GetTemplateClass() const
@@ -53,6 +54,8 @@ struct FPCGMoveResourceParams
 	{
 		return Cast<ObjectType>(ExpectedPreviousOwner);
 	}
+
+	bool Accepts(const UObject* InObject) const;
 };
 
 /** 
