[PCG] When removing a node from a graph, remove it also from the graph inner objetcs.

#rnx
#Synced-CL 47950886

diff --git a/Engine/Plugins/PCG/Source/PCG/Private/PCGGraph.cpp b/Engine/Plugins/PCG/Source/PCG/Private/PCGGraph.cpp
--- a/Engine/Plugins/PCG/Source/PCG/Private/PCGGraph.cpp
+++ b/Engine/Plugins/PCG/Source/PCG/Private/PCGGraph.cpp
@@ -1274,6 +1274,9 @@ void UPCGGraph::RemoveNodes_Internal(TArrayView<UPCGNode*> InNodes)
 		// Add the node to the transaction, to make sure we reconnect everything correctly on Undo/Redo
 		Node->Modify();
 
+		// Remove the node from the graph inners.
+		Node->Rename(nullptr, GetTransientPackage(), REN_DontCreateRedirectors);
+
 		Nodes.Remove(Node);
 	}
 
