[PCG] Bugfixes regarding overrides for PCG

* Add an option in the GetAllOverridableParams to discard a certain container depth
* Deprecate the multiple options to include/exclude metadata/property flags in favor of single functor that will allow to implement any complex logic.
* Discard all properties with a container depth more than one for all overrides.
* Discard all nested properties that are not editable (except the one that are explictly marked overridable)

#Synced-CL 48103248

diff --git a/Engine/Plugins/PCG/Source/PCG/Private/Elements/PCGExecuteBlueprint.cpp b/Engine/Plugins/PCG/Source/PCG/Private/Elements/PCGExecuteBlueprint.cpp
--- a/Engine/Plugins/PCG/Source/PCG/Private/Elements/PCGExecuteBlueprint.cpp
+++ b/Engine/Plugins/PCG/Source/PCG/Private/Elements/PCGExecuteBlueprint.cpp
@@ -426,7 +426,8 @@ TArray<FPCGSettingsOverridableParam> UPCGBlueprintSettings::GatherOverridablePar
 		PCGSettingsHelpers::FPCGGetAllOverridableParamsConfig Config;
 		Config.bExcludeSuperProperties = true;
 		Config.bExtractArrays = true;
-		Config.ExcludePropertyFlags = CPF_DisableEditOnInstance | CPF_EditConst | CPF_BlueprintReadOnly;
+		Config.MaxContainersNum = 1; // Do not support more than 1 container deep.
+		Config.ShouldKeepPropertyFunc = [](const FProperty* InProperty, int32) { return InProperty && !InProperty->HasAnyPropertyFlags(CPF_DisableEditOnInstance | CPF_EditConst | CPF_BlueprintReadOnly); };
 		OverridableParams.Append(PCGSettingsHelpers::GetAllOverridableParams(BPClass, Config));
 	}
 
diff --git a/Engine/Plugins/PCG/Source/PCG/Private/Helpers/PCGPropertyHelpers.cpp b/Engine/Plugins/PCG/Source/PCG/Private/Helpers/PCGPropertyHelpers.cpp
--- a/Engine/Plugins/PCG/Source/PCG/Private/Helpers/PCGPropertyHelpers.cpp
+++ b/Engine/Plugins/PCG/Source/PCG/Private/Helpers/PCGPropertyHelpers.cpp
@@ -370,8 +370,7 @@ UPCGParamData* PCGPropertyHelpers::ExtractPropertyAsAttributeSet(const PCGProper
 			// Can only get exposed properties and visible if requested
 			if (Parameters.bPropertyNeedsToBeVisible)
 			{
-				Config.ExcludePropertyFlags = ExcludePropertyFlags;
-				Config.IncludePropertyFlags = IncludePropertyFlags;
+				Config.ShouldKeepPropertyFunc = [](const FProperty* InProperty, int32) -> bool { return InProperty && InProperty->HasAnyPropertyFlags(IncludePropertyFlags) && !InProperty->HasAnyPropertyFlags(ExcludePropertyFlags); };
 			}
 			TArray<FPCGSettingsOverridableParam> AllChildProperties = PCGSettingsHelpers::GetAllOverridableParams(UnderlyingClass, Config);
 
diff --git a/Engine/Plugins/PCG/Source/PCG/Private/Helpers/PCGSettingsHelpers.cpp b/Engine/Plugins/PCG/Source/PCG/Private/Helpers/PCGSettingsHelpers.cpp
--- a/Engine/Plugins/PCG/Source/PCG/Private/Helpers/PCGSettingsHelpers.cpp
+++ b/Engine/Plugins/PCG/Source/PCG/Private/Helpers/PCGSettingsHelpers.cpp
@@ -266,7 +266,7 @@ namespace PCGSettingsHelpers
 	}
 
 	template <uint32 N>
-	TArray<FPCGSettingsOverridableParam> GetAllOverridableParams_Internal(const UStruct* InClass, const FPCGGetAllOverridableParamsConfig& InConfig, TArray<const FProperty*, TInlineAllocator<N>>& InAlreadySeenProperties)
+	TArray<FPCGSettingsOverridableParam> GetAllOverridableParams_Internal(const UStruct* InClass, const FPCGGetAllOverridableParamsConfig& InConfig, TArray<const FProperty*, TInlineAllocator<N>>& InAlreadySeenProperties, int32 Depth)
 	{
 		TArray<FName> LabelCache;
 
@@ -331,6 +331,7 @@ namespace PCGSettingsHelpers
 					return bFound;
 				};
 
+PRAGMA_DISABLE_DEPRECATION_WARNINGS
 				if (!InConfig.IncludeMetadataValues.IsEmpty())
 				{
 					bValid &= HasMetadata(InConfig.IncludeMetadataValues, InConfig.bIncludeMetadataIsConjunction);
@@ -340,13 +341,15 @@ namespace PCGSettingsHelpers
 				{
 					bValid &= !HasMetadata(InConfig.ExcludeMetadataValues, InConfig.bExcludeMetadataIsConjunction);
 				}
+PRAGMA_ENABLE_DEPRECATION_WARNINGS
 #endif // WITH_EDITOR
 
 				auto HasPropertyFlags = [Property](uint64 InFlags, bool bIsConjunction) -> bool
 				{
 					return bIsConjunction ? Property->HasAllPropertyFlags(InFlags) : Property->HasAnyPropertyFlags(InFlags);
 				};
-
+				
+PRAGMA_DISABLE_DEPRECATION_WARNINGS
 				if (InConfig.IncludePropertyFlags != 0)
 				{
 					bValid &= HasPropertyFlags(InConfig.IncludePropertyFlags, InConfig.bIncludePropertyFlagsIsConjunction);
@@ -356,6 +359,7 @@ namespace PCGSettingsHelpers
 				{
 					bValid &= !HasPropertyFlags(InConfig.ExcludePropertyFlags, InConfig.bExcludePropertyFlagsIsConjunction);
 				}
+PRAGMA_ENABLE_DEPRECATION_WARNINGS
 
 				// Don't allow to override the seed if the settings doesn't use the seed.
 				if (!InConfig.bUseSeed && Property->GetOwnerClass() && Property->GetOwnerClass()->IsChildOf<UPCGSettings>())
@@ -363,15 +367,20 @@ namespace PCGSettingsHelpers
 					bValid &= (Property->GetFName() != GET_MEMBER_NAME_CHECKED(UPCGSettings, Seed));
 				}
 
+				if (InConfig.ShouldKeepPropertyFunc.IsSet())
+				{
+					bValid &= InConfig.ShouldKeepPropertyFunc(Property, Depth);
+				}
+
 				if (!bValid)
 				{
 					continue;
 				}
 
-				auto RecursiveExtraction = [&InConfig, &Res, &LabelCache, Property, InClass, &GatherAliases, &InAlreadySeenProperties](const UStruct* NextClass, const FProperty* ContainerProperty)
+				auto RecursiveExtraction = [&InConfig, &Res, &LabelCache, Property, InClass, &GatherAliases, &InAlreadySeenProperties, Depth](const UStruct* NextClass, const FProperty* ContainerProperty)
 				{
-					// Reached max depth
-					if (InConfig.MaxStructDepth == 0)
+					// Reached max depth or max number of containers
+					if (InConfig.MaxStructDepth == 0 || (ContainerProperty && InConfig.MaxContainersNum == 0))
 					{
 						return;
 					}
@@ -381,9 +390,11 @@ namespace PCGSettingsHelpers
 					InAlreadySeenProperties.Add(ContainerProperty ? ContainerProperty : Property);
 					RecurseConfig.bUseSeed = true;
 #if WITH_EDITOR
+PRAGMA_DISABLE_DEPRECATION_WARNINGS
 					RecurseConfig.IncludeMetadataValues.Remove(PCGObjectMetadata::Overridable);
 					RecurseConfig.IncludeMetadataValues.Remove(PCGObjectMetadata::OverridableCPUAndGPU);
 					RecurseConfig.IncludeMetadataValues.Remove(PCGObjectMetadata::OverridableCPUAndGPUWithReadback);
+PRAGMA_ENABLE_DEPRECATION_WARNINGS
 
 					// Also check if we have some filtering on the child properties
 					TArray<FString> ChildPropertiesFiltering;
@@ -399,7 +410,12 @@ namespace PCGSettingsHelpers
 						RecurseConfig.MaxStructDepth--;
 					}
 
-					for (FPCGSettingsOverridableParam& ChildParam : GetAllOverridableParams_Internal(NextClass, RecurseConfig, InAlreadySeenProperties))
+					if (ContainerProperty && RecurseConfig.MaxContainersNum > 0)
+					{
+						RecurseConfig.MaxContainersNum--;
+					}
+
+					for (FPCGSettingsOverridableParam& ChildParam : GetAllOverridableParams_Internal(NextClass, RecurseConfig, InAlreadySeenProperties, Depth + 1))
 					{
 						FName Label = ChildParam.Label;
 						bool bHasNameClash = false;
@@ -537,7 +553,7 @@ namespace PCGSettingsHelpers
 		// Array of already seen properties to catch infinite recursion
 		TArray<const FProperty*, TInlineAllocator<InlineAllocator>> AlreadySeenProperties;
 
-		return GetAllOverridableParams_Internal(InClass, InConfig, AlreadySeenProperties);
+		return GetAllOverridableParams_Internal(InClass, InConfig, AlreadySeenProperties, 0);
 	}
 
 	// Thread-local variables used for the FPinTypeScopeHelper implementation
@@ -600,8 +616,7 @@ PCG_API TArray<FPCGSettingsPropertyDefinition> UPCGSettingsHelpers::GetCommonlyU
 	using namespace PCGSettingsHelpers;
 
 	FPCGGetAllOverridableParamsConfig Config;
-	Config.ExcludePropertyFlags = EPropertyFlags::CPF_AdvancedDisplay;
-	Config.IncludePropertyFlags = EPropertyFlags::CPF_Edit;
+	Config.ShouldKeepPropertyFunc = [](const FProperty* InProperty, int32) -> bool { return InProperty && InProperty->HasAnyPropertyFlags(CPF_Edit) && !InProperty->HasAnyPropertyFlags(CPF_AdvancedDisplay); };
 	Config.StopClass = UPCGSettings::StaticClass();
 
 	TArray<FPCGSettingsPropertyDefinition> Properties;
diff --git a/Engine/Plugins/PCG/Source/PCG/Private/PCGSettings.cpp b/Engine/Plugins/PCG/Source/PCG/Private/PCGSettings.cpp
--- a/Engine/Plugins/PCG/Source/PCG/Private/PCGSettings.cpp
+++ b/Engine/Plugins/PCG/Source/PCG/Private/PCGSettings.cpp
@@ -1014,11 +1014,35 @@ TArray<FPCGSettingsOverridableParam> UPCGSettings::GatherOverridableParams() con
 {
 	PCGSettingsHelpers::FPCGGetAllOverridableParamsConfig Config;
 	Config.bExtractArrays = true;
+	Config.MaxContainersNum = 1; // Do not support more than 1 container deep.
 	Config.bUseSeed = UseSeed();
-	Config.IncludeMetadataValues.Add(PCGObjectMetadata::Overridable);
-	Config.IncludeMetadataValues.Add(PCGObjectMetadata::OverridableCPUAndGPU);
-	Config.IncludeMetadataValues.Add(PCGObjectMetadata::OverridableCPUAndGPUWithReadback);
-	Config.ExcludeMetadataValues.Add(PCGObjectMetadata::NotOverridable);
+
+	Config.ShouldKeepPropertyFunc = [](const FProperty* InProperty, int32 Depth) -> bool
+	{
+		check(InProperty);
+		if (InProperty->HasMetaData(PCGObjectMetadata::Overridable)
+			|| InProperty->HasMetaData(PCGObjectMetadata::OverridableCPUAndGPU)
+			|| InProperty->HasMetaData(PCGObjectMetadata::OverridableCPUAndGPUWithReadback))
+		{
+			// All properties explicitly marked "Overridable" are overridable
+			return true;
+		}
+		else if (InProperty->HasMetaData(PCGObjectMetadata::NotOverridable))
+		{
+			// All properties explicitly marked "Not Overridable" are not overridable
+			return false;
+		}
+		else if (!InProperty->HasAnyPropertyFlags(CPF_Edit))
+		{
+			// Then for recursive structs, make sure to not take properties that are not editable.
+			return false;
+		}
+		else
+		{
+			// If we have none of that, we only allow overrides for nested properties.
+			return Depth > 0;
+		}
+	};
 
 	return PCGSettingsHelpers::GetAllOverridableParams(GetClass(), Config);
 }
@@ -1135,6 +1159,7 @@ void UPCGSettings::InitializeCachedOverridableParams(bool bReset)
 		// Do not support multi containers
 		if (Param.NumContainers > 1)
 		{
+			CachedOverridableParams.RemoveAt(i--);
 			continue;
 		}
 
diff --git a/Engine/Plugins/PCG/Source/PCG/Private/PCGSubgraph.cpp b/Engine/Plugins/PCG/Source/PCG/Private/PCGSubgraph.cpp
--- a/Engine/Plugins/PCG/Source/PCG/Private/PCGSubgraph.cpp
+++ b/Engine/Plugins/PCG/Source/PCG/Private/PCGSubgraph.cpp
@@ -279,8 +279,9 @@ TArray<FPCGSettingsOverridableParam> UPCGBaseSubgraphSettings::GatherOverridable
 			{
 				PCGSettingsHelpers::FPCGGetAllOverridableParamsConfig Config;
 				Config.bExcludeSuperProperties = true;
-				Config.ExcludePropertyFlags = CPF_DisableEditOnInstance;
+				Config.ShouldKeepPropertyFunc = [](const FProperty* InProperty, int32) { return InProperty && !InProperty->HasAnyPropertyFlags(CPF_DisableEditOnInstance); };
 				Config.bExtractArrays = true;
+				Config.MaxContainersNum = 1; // Do not support more than 1 container deep.
 				OverridableParams.Append(PCGSettingsHelpers::GetAllOverridableParams(ScriptStruct, Config));
 			}
 		}
diff --git a/Engine/Plugins/PCG/Source/PCG/Public/Helpers/PCGSettingsHelpers.h b/Engine/Plugins/PCG/Source/PCG/Public/Helpers/PCGSettingsHelpers.h
--- a/Engine/Plugins/PCG/Source/PCG/Public/Helpers/PCGSettingsHelpers.h
+++ b/Engine/Plugins/PCG/Source/PCG/Public/Helpers/PCGSettingsHelpers.h
@@ -129,6 +129,14 @@ namespace PCGSettingsHelpers
 
 	struct FPCGGetAllOverridableParamsConfig
 	{
+		PRAGMA_DISABLE_DEPRECATION_WARNINGS
+		FPCGGetAllOverridableParamsConfig() = default;
+		FPCGGetAllOverridableParamsConfig(const FPCGGetAllOverridableParamsConfig&) = default;
+		FPCGGetAllOverridableParamsConfig(FPCGGetAllOverridableParamsConfig&&) = default;
+		FPCGGetAllOverridableParamsConfig& operator=(const FPCGGetAllOverridableParamsConfig&) = default;
+		FPCGGetAllOverridableParamsConfig& operator=(FPCGGetAllOverridableParamsConfig&&) = default;
+		PRAGMA_ENABLE_DEPRECATION_WARNINGS
+
 		// If we don't use the seed, don't add it as override.
 		bool bUseSeed = false;
 
@@ -141,30 +149,42 @@ namespace PCGSettingsHelpers
 
 #if WITH_EDITOR
 		// List of metadata values to find in property metadata. Only works in editor builds as metadata on property is not available elsewise.
+		UE_DEPRECATED(5.8, "Replaced by custom logic in ShouldKeepPropertyFunc")
 		TArray<FName> IncludeMetadataValues;
 
 		// If the metadata values is a conjunction (all values need to be present to keep), or disjunction (any values needs to be present to keep)
+		UE_DEPRECATED(5.8, "Replaced by custom logic in ShouldKeepPropertyFunc")
 		bool bIncludeMetadataIsConjunction = false;
 
 		// List of metadata values to find in property metadata. Only works in editor builds as metadata on property is not available elsewise.
+		UE_DEPRECATED(5.8, "Replaced by custom logic in ShouldKeepPropertyFunc")
 		TArray<FName> ExcludeMetadataValues;
 
 		// If the exclude values is a conjunction (all values need to be present to discard), or disjunction (any values needs to be present to discard)
+		UE_DEPRECATED(5.8, "Replaced by custom logic in ShouldKeepPropertyFunc")
 		bool bExcludeMetadataIsConjunction = false;
 #endif // WITH_EDITOR
 
+		UE_DEPRECATED(5.8, "Replaced by custom logic in ShouldKeepPropertyFunc")
 		// Flags to exclude in property flags
 		uint64 ExcludePropertyFlags = 0;
 
 		// If the exclude flags is a conjunction (all flags need to be present to discard), or disjunction (any flag needs to be present to discard)
+		UE_DEPRECATED(5.8, "Replaced by custom logic in ShouldKeepPropertyFunc")
 		bool bExcludePropertyFlagsIsConjunction = false;
 
 		// Flags to include in property flags. Note that if exclusion says discard, it will be discarded.
+		UE_DEPRECATED(5.8, "Replaced by custom logic in ShouldKeepPropertyFunc")
 		uint64 IncludePropertyFlags = 0;
 
 		// If the include flags is a conjunction (all flags need to be present to keep), or disjunction (any flag needs to be present to keep)
+		UE_DEPRECATED(5.8, "Replaced by custom logic in ShouldKeepPropertyFunc")
 		bool bIncludePropertyFlagsIsConjunction = false;
 
+		// To accomodate all the different combinations for property flags and metadata flags, the caller can provide a function to evaluate
+		// if the property should be kept or discarded.
+		TFunction<bool(const FProperty* /*Property*/, int32 /*Depth*/)> ShouldKeepPropertyFunc;
+
 		// Max depth for structs of structs. -1 = no limit
 		int32 MaxStructDepth = -1;
 
@@ -173,6 +193,11 @@ namespace PCGSettingsHelpers
 
 		// If you want to also extract arrays
 		bool bExtractArrays = false;
+
+		// When extracting arrays of arrays, discard any property that is within this number or more containers.
+		// For example, a value of 1 will discard any property that is in an array of array.
+		// -1 = no limit
+		int32 MaxContainersNum = -1;
 	};
 
 	PCG_API TArray<FPCGSettingsOverridableParam> GetAllOverridableParams(const UStruct* InClass, const FPCGGetAllOverridableParamsConfig& InConfig);
