Add some structure logging and contexts for error messages

#Synced-CL Unknown

diff --git a/Engine/Plugins/PCG/Source/PCG/Private/Graph/PCGGraphCompiler.cpp b/Engine/Plugins/PCG/Source/PCG/Private/Graph/PCGGraphCompiler.cpp
--- a/Engine/Plugins/PCG/Source/PCG/Private/Graph/PCGGraphCompiler.cpp
+++ b/Engine/Plugins/PCG/Source/PCG/Private/Graph/PCGGraphCompiler.cpp
@@ -18,6 +18,7 @@
 #include "Graph/PCGPinDependencyExpression.h"
 
 #include "HAL/IConsoleManager.h"
+#include "Logging/StructuredLog.h"
 #include "Misc/ScopeRWLock.h"
 
 namespace PCGGraphCompiler
@@ -262,7 +263,7 @@ TArray<FPCGGraphTask> FPCGGraphCompiler::CompileGraph(UPCGGraph* InGraph, FPCGTa
 					}
 					else
 					{
-						UE_LOG(LogPCG, Warning, TEXT("Invalid inbound edge on subgraph"));
+						UE_LOGFMT(LogPCG, Warning, "Invalid inbound edge on subgraph");
 					}
 				}
 			}
@@ -337,7 +338,7 @@ TArray<FPCGGraphTask> FPCGGraphCompiler::CompileGraph(UPCGGraph* InGraph, FPCGTa
 				{
 					if (!InboundEdge->IsValid())
 					{
-						UE_LOG(LogPCG, Warning, TEXT("Unbound inbound edge"));
+						UE_LOGFMT(LogPCG, Warning, "Unbound inbound edge");
 						continue;
 					}
 
@@ -348,7 +349,7 @@ TArray<FPCGGraphTask> FPCGGraphCompiler::CompileGraph(UPCGGraph* InGraph, FPCGTa
 					}
 					else
 					{
-						UE_LOG(LogPCG, Error, TEXT("Inconsistent node linkage on node '%s'"), *Node->GetFName().ToString());
+						UE_LOGFMT(LogPCG, Error, "Inconsistent node linkage on node '{Node}'", FSoftObjectPath(Node));
 						return TArray<FPCGGraphTask>();
 					}
 				}
@@ -365,7 +366,7 @@ TArray<FPCGGraphTask> FPCGGraphCompiler::CompileGraph(UPCGGraph* InGraph, FPCGTa
 			{
 				if (!OutboundEdge->IsValid())
 				{
-					UE_LOG(LogPCG, Warning, TEXT("Unbound outbound edge"));
+					UE_LOGFMT(LogPCG, Warning, "Unbound outbound edge");
 					continue;
 				}
 
@@ -1120,9 +1121,9 @@ void FPCGGraphCompiler::CalculateDynamicActivePinDependencies(FPCGTaskId InTaskI
 
 	auto LogTooManyOutputPinsWarning = [](const UPCGNode* InUpstreamNode)
 	{
-		UE_LOG(LogPCG, Warning, TEXT("Number of output pins (%d) on node %s exceeded maximum limit (%d) for dynamic node culling. "
-			"Dynamic culling will not function for this node's output pins."),
-			InUpstreamNode->GetOutputPins().Num(), *InUpstreamNode->GetName(), PCGPinIdHelpers::MaxOutputPins);
+		UE_LOGFMT(LogPCG, Warning, "Number of output pins ({OutputPinCount}) on node {NodeName} exceeded maximum limit ({Limit}) for dynamic node culling. "
+			"Dynamic culling will not function for this node's output pins.",
+			InUpstreamNode->GetOutputPins().Num(), InUpstreamNode->GetName(), PCGPinIdHelpers::MaxOutputPins);
 	};
 
 	const UPCGNode* Node = InOutCompiledTasks[InTaskId].Node;
@@ -1242,7 +1243,8 @@ void FPCGGraphCompiler::CalculateDynamicActivePinDependencies(FPCGTaskId InTaskI
 TArray<FPCGGraphTask> FPCGGraphCompiler::GetCompiledTasks(UPCGGraph* InGraph, uint32 GenerationGridSize, FPCGStackContext& OutStackContext, bool bIsTopGraph)
 {
 	TRACE_CPUPROFILER_EVENT_SCOPE(FPCGGraphCompiler::GetCompiledTasks);
-
+	UE_LOG_CONTEXT("PCGGraph", FSoftObjectPath(InGraph));
+	
 #if !WITH_EDITOR
 	// In standalone builds, try using cooked tasks first.
 	if (InGraph)
@@ -1415,7 +1417,7 @@ void FPCGGraphCompiler::CompileTopGraph(UPCGGraph* InGraph, uint32 GenerationGri
 
 		if (CountBefore > CompiledTasks.Num())
 		{
-			UE_LOG(LogPCG, Warning, TEXT("One or more GPU nodes were culled from the graph!"));
+			UE_LOGFMT(LogPCG, Warning, "One or more GPU nodes were culled from the graph!");
 		}
 	}
 
diff --git a/Engine/Plugins/PCG/Source/PCG/Private/Graph/PCGGraphCompilerGPU.cpp b/Engine/Plugins/PCG/Source/PCG/Private/Graph/PCGGraphCompilerGPU.cpp
--- a/Engine/Plugins/PCG/Source/PCG/Private/Graph/PCGGraphCompilerGPU.cpp
+++ b/Engine/Plugins/PCG/Source/PCG/Private/Graph/PCGGraphCompilerGPU.cpp
@@ -23,6 +23,7 @@
 #include "ComputeFramework/ComputeFramework.h"
 #include "ComputeFramework/ComputeKernel.h"
 #include "HAL/IConsoleManager.h"
+#include "Logging/StructuredLog.h"
 #include "Misc/App.h"
 #include "Shader/ShaderTypes.h"
 #include "UObject/Package.h"
@@ -903,7 +904,8 @@ void FPCGGraphCompilerGPU::BuildComputeGraphStaticData(
 			{
 				if (InOutComputeGraph->StaticAttributeTable.AddAttribute(Key) == INDEX_NONE)
 				{
-					UE_LOG(LogPCG, Error, TEXT("FPCGGraphCompilerGPU: Static attribute table exceeded maximum size (%d), use the 'Dump Data Descriptions' setting on the GPU node(s) to list attributes that are present."), PCGDataCollectionPackingConstants::MAX_NUM_CUSTOM_ATTRS);
+					UE_LOGFMT(LogPCG, Error, "FPCGGraphCompilerGPU: Static attribute table exceeded maximum size ({Maximum}), use the 'Dump Data Descriptions' setting on the GPU node(s) to list attributes that are present.",
+						PCGDataCollectionPackingConstants::MAX_NUM_CUSTOM_ATTRS);
 					break;
 				}
 			}
@@ -1332,7 +1334,7 @@ void FPCGGraphCompilerGPU::CompileComputeGraph(
 			UPCGComputeDataInterface* const* FoundDI = InOutputPinToDataInterface.Find({ Input.TaskId, Input.UpstreamPin->Label });
 			if (!FoundDI || !*FoundDI)
 			{
-				UE_LOG(LogPCG, Warning, TEXT("Failed to find data interface associated with pin '%s'."), *Input.UpstreamPin->Label.ToString());
+				UE_LOGFMT(LogPCG, Warning, "Failed to find data interface associated with pin '{Label}'.", Input.UpstreamPin->Label);
 				continue;
 			}
 
@@ -1365,7 +1367,7 @@ void FPCGGraphCompilerGPU::CompileComputeGraph(
 			UPCGComputeDataInterface* const* FoundDI = InOutputPinToDataInterface.Find({ TaskId, OutputPinProperties.Label });
 			if (!FoundDI || !*FoundDI)
 			{
-				UE_LOG(LogPCG, Warning, TEXT("Failed to find data interface associated with pin '%s'."), *OutputPinProperties.Label.ToString());
+				UE_LOGFMT(LogPCG, Warning, "Failed to find data interface associated with pin '{Label}'.", OutputPinProperties.Label);
 				continue;
 			}
 
