[PCG] Added query (logs) and isolate editor toolmodes in PCG.

#Synced-CL 48040603

diff --git a/Engine/Plugins/PCG/Source/PCGEditor/Private/EditorMode/PCGEdMode.cpp b/Engine/Plugins/PCG/Source/PCGEditor/Private/EditorMode/PCGEdMode.cpp
--- a/Engine/Plugins/PCG/Source/PCGEditor/Private/EditorMode/PCGEdMode.cpp
+++ b/Engine/Plugins/PCG/Source/PCGEditor/Private/EditorMode/PCGEdMode.cpp
@@ -9,6 +9,7 @@
 #include "EditorMode/Tools/PCGInteractiveToolSettings.h"
 #include "EditorMode/Tools/Line/PCGDrawSplineTool.h"
 #include "EditorMode/Tools/Paint/PCGPaintTool.h"
+#include "EditorMode/Tools/Query/PCGQueryTool.h"
 #include "EditorMode/Tools/Select/PCGSelectTool.h"
 #include "EditorMode/Tools/Volume/PCGVolumeTool.h"
 
@@ -70,6 +71,25 @@ namespace PCGEditorMode
 				ThisMode->GetInteractiveToolsContext()->StartTool(Args[0]);
 			}
 		}));
+
+	static FAutoConsoleCommand CommandSetGraphInTool(
+		TEXT("pcg.tool.SetGraph"),
+		TEXT("Sets graph on the current tool."),
+		FConsoleCommandWithArgsDelegate::CreateLambda([](const TArray<FString>& Args)
+		{
+			UEdMode* ThisMode = GLevelEditorModeTools().GetActiveScriptableMode(UPCGEditorMode::EM_PCGEditorModeId);
+			UPCGEditorMode* EditorMode = Cast<UPCGEditorMode>(ThisMode);
+			UPCGInteractiveToolSettings* ToolSettings = EditorMode ? Cast<UPCGInteractiveToolSettings>(EditorMode->GetCurrentToolSettings()) : nullptr;
+			if (ToolSettings && !Args.IsEmpty())
+			{
+				FSoftObjectPath GraphPath(Args[0]);
+				UObject* GraphObject = GraphPath.TryLoad();
+				if (UPCGGraphInterface* Graph = Cast<UPCGGraphInterface>(GraphObject))
+				{
+					ToolSettings->SetToolGraph(Graph);
+				}
+			}
+		}));
 }
 
 UPCGEditorMode::UPCGEditorMode()
@@ -211,6 +231,18 @@ void UPCGEditorMode::Enter()
 		RegisterTool(ToolManagerCommands.EnableSelectTool, UPCGInteractiveToolSettings_Select::StaticGetToolTag().ToString(), SelectToolBuilder);
 	}
 
+	{
+		UPCGQueryToolBuilder* QueryToolBuilder = NewObject<UPCGQueryToolBuilder>();
+		QueryToolBuilder->SetToolClass(UPCGQueryTool::StaticClass());
+		RegisterTool(ToolManagerCommands.EnableQueryTool, UPCGQueryTool::StaticGetToolTag().ToString(), QueryToolBuilder);
+	}
+
+	{
+		UPCGQueryToolBuilder* QueryToolBuilder = NewObject<UPCGQueryToolBuilder>();
+		QueryToolBuilder->SetToolClass(UPCGIsolateTool::StaticClass());
+		RegisterTool(ToolManagerCommands.EnableIsolateTool, UPCGIsolateTool::StaticGetToolTag().ToString(), QueryToolBuilder);
+	}
+
 	ToolManager->SelectActiveToolType(EToolSide::Left, "DrawSpline");
 
 	ConfigureRealTimeViewportsOverride(/*bEnable=*/true);
@@ -283,7 +315,7 @@ bool UPCGEditorMode::IsSelectionAllowed(AActor* InActor, bool bInSelection) cons
 		return true;
 	}
 
-	if (const UPCGInteractiveToolSettings* ToolSettings = GetCurrentToolSettings())
+	if (const UPCGInteractiveToolBaseSettings* ToolSettings = GetCurrentToolSettings())
 	{
 		return ToolSettings->IsSelectionAllowed(InActor, bInSelection);
 	}
@@ -450,7 +482,7 @@ void UPCGEditorMode::ConfigureRealTimeViewportsOverride(const bool bEnable)
 	}
 }
 
-UPCGInteractiveToolSettings* UPCGEditorMode::GetCurrentToolSettings() const
+UPCGInteractiveToolBaseSettings* UPCGEditorMode::GetCurrentToolSettings() const
 {
 	UInteractiveToolManager* ToolManager = GetToolManager();
 
@@ -462,12 +494,12 @@ UPCGInteractiveToolSettings* UPCGEditorMode::GetCurrentToolSettings() const
 
 			UObject** FoundProperties = AllToolProperties.FindByPredicate([](const UObject* Candidate)
 			{
-				return Candidate->IsA<UPCGInteractiveToolSettings>();
+				return Candidate->IsA<UPCGInteractiveToolBaseSettings>();
 			});
 
 			if (FoundProperties)
 			{
-				return Cast<UPCGInteractiveToolSettings>(*FoundProperties);
+				return Cast<UPCGInteractiveToolBaseSettings>(*FoundProperties);
 			}
 		}
 	}
diff --git a/Engine/Plugins/PCG/Source/PCGEditor/Private/EditorMode/PCGEdMode.h b/Engine/Plugins/PCG/Source/PCGEditor/Private/EditorMode/PCGEdMode.h
--- a/Engine/Plugins/PCG/Source/PCGEditor/Private/EditorMode/PCGEdMode.h
+++ b/Engine/Plugins/PCG/Source/PCGEditor/Private/EditorMode/PCGEdMode.h
@@ -7,7 +7,7 @@
 
 #include "PCGEdMode.generated.h"
 
-class UPCGInteractiveToolSettings;
+class UPCGInteractiveToolBaseSettings;
 
 /** The unique state for the editor to be in, specific to PCG's toolkits. */
 UCLASS(Transient)
@@ -42,7 +42,7 @@ public:
 	virtual void ModeTick(float DeltaTime) override;
 	//~End UEdMode interface
 
-	UPCGInteractiveToolSettings* GetCurrentToolSettings() const;
+	UPCGInteractiveToolBaseSettings* GetCurrentToolSettings() const;
 
 protected:
 	virtual void CreateToolkit() override;
diff --git a/Engine/Plugins/PCG/Source/PCGEditor/Private/EditorMode/PCGEdModeCommands.cpp b/Engine/Plugins/PCG/Source/PCGEditor/Private/EditorMode/PCGEdModeCommands.cpp
--- a/Engine/Plugins/PCG/Source/PCGEditor/Private/EditorMode/PCGEdModeCommands.cpp
+++ b/Engine/Plugins/PCG/Source/PCGEditor/Private/EditorMode/PCGEdModeCommands.cpp
@@ -9,6 +9,7 @@
 #include "EditorMode/Tools/Helpers/PCGEdModeEditorUtilities.h"
 #include "EditorMode/Tools/Line/PCGDrawSplineTool.h"
 #include "EditorMode/Tools/Paint/PCGPaintTool.h"
+#include "EditorMode/Tools/Query/PCGQueryTool.h"
 #include "EditorMode/Tools/Select/PCGSelectTool.h"
 #include "EditorMode/Tools/Volume/PCGVolumeTool.h"
 #include "Logging/StructuredLog.h"
@@ -54,6 +55,7 @@ void FPCGEditorModePaletteCommands::RegisterCommands()
 	PCG_CONTEXT_COMMAND_INFO(LoadPaintContextPalette, "Paint", "PCGEditorMode.Context.Paint");
 	PCG_CONTEXT_COMMAND_INFO(LoadVolumeContextPalette, "Volume", "PCGEditorMode.Context.Volume");
 	PCG_CONTEXT_COMMAND_INFO(LoadSelectContextPalette, "Select", "PCGEditorMode.Context.Select");
+	PCG_CONTEXT_COMMAND_INFO(LoadQueryContextPalette, "Query", "PCGEditorMode.Context.Query");
 }
 
 #undef PCG_CONTEXT_COMMAND_INFO
@@ -93,6 +95,8 @@ void FPCGEditorModeToolCommands::RegisterCommands()
 	PCG_TOOL_COMMAND_INFO(EnablePaintTool, "Paint", "Paint", "Use the '{ToolTag}' tool to apply PCG points with a brush.", UPCGInteractiveToolSettings_PaintTool::StaticGetToolTag(), "PCGEditorMode.Context.Paint");
 	PCG_TOOL_COMMAND_INFO(EnableVolumeTool, "Volume", "Volume", "Create or updates a volume in the scene using the '{ToolTag}' tool.", UPCGInteractiveToolSettings_Volume::StaticGetToolTag(), "PCGEditorMode.Tools.Volume");
 	PCG_TOOL_COMMAND_INFO(EnableSelectTool, "Select", "Select", "Apply PCG graph to selection using the '{ToolTag}' tool.", UPCGInteractiveToolSettings_Select::StaticGetToolTag(), "PCGEditorMode.Tools.Select");
+	PCG_TOOL_COMMAND_INFO(EnableQueryTool, "Query", "Query", "Queries scene data from raycasts.", UPCGQueryTool::StaticGetToolTag(), "PCGEditorMode.Tools.Query");
+	PCG_TOOL_COMMAND_INFO(EnableIsolateTool, "Isolate", "Isolate", "Isolates content from the actor on a tag basis.", UPCGIsolateTool::StaticGetToolTag(), "PCGEditorMode.Tools.Isolate");
 }
 
 void FPCGEditorModeToolCommands::RegisterAllToolCommands()
@@ -101,10 +105,14 @@ void FPCGEditorModeToolCommands::RegisterAllToolCommands()
 	FPCGPaintCommands::Register();
 	FPCGVolumeCommands::Register();
 	FPCGSelectCommands::Register();
+	FPCGQueryCommands::Register();
+	FPCGIsolateCommands::Register();
 }
 
 void FPCGEditorModeToolCommands::UnregisterAllToolCommands()
 {
+	FPCGIsolateCommands::Unregister();
+	FPCGQueryCommands::Unregister();
 	FPCGSelectCommands::Unregister();
 	FPCGVolumeCommands::Unregister();
 	FPCGPaintCommands::Unregister();
@@ -136,6 +144,11 @@ TArray<TSharedPtr<FUICommandInfo>> FPCGEditorModeToolCommands::GetSelectContextT
 	return { EnableSelectTool };
 }
 
+TArray<TSharedPtr<FUICommandInfo>> FPCGEditorModeToolCommands::GetQueryContextToolCommands()
+{
+	return { EnableQueryTool, EnableIsolateTool };
+}
+
 /**
  * Macro for defining a constructor and a helper function to retrieve 
  * @param ToolCommandClassName The tool's class name
@@ -158,6 +171,8 @@ PCG_DEFINE_TOOL_COMMANDS(FPCGDrawSurfaceCommands, "DrawSurface", "Draw Surface",
 PCG_DEFINE_TOOL_COMMANDS(FPCGPaintCommands, "Paint", "Paint", UPCGPaintTool);
 PCG_DEFINE_TOOL_COMMANDS(FPCGVolumeCommands, "Volume", "Volume", UPCGVolumeTool);
 PCG_DEFINE_TOOL_COMMANDS(FPCGSelectCommands, "Select", "Select", UPCGSelectTool);
+PCG_DEFINE_TOOL_COMMANDS(FPCGQueryCommands, "Query", "Query", UPCGQueryTool);
+PCG_DEFINE_TOOL_COMMANDS(FPCGIsolateCommands, "Isolate", "Isolate", UPCGIsolateTool);
 
 #undef PCG_DEFINE_TOOL_COMMANDS
 
diff --git a/Engine/Plugins/PCG/Source/PCGEditor/Private/EditorMode/PCGEdModeCommands.h b/Engine/Plugins/PCG/Source/PCGEditor/Private/EditorMode/PCGEdModeCommands.h
--- a/Engine/Plugins/PCG/Source/PCGEditor/Private/EditorMode/PCGEdModeCommands.h
+++ b/Engine/Plugins/PCG/Source/PCGEditor/Private/EditorMode/PCGEdModeCommands.h
@@ -27,6 +27,7 @@ public:
 	TSharedPtr<FUICommandInfo> LoadPaintContextPalette;
 	TSharedPtr<FUICommandInfo> LoadVolumeContextPalette;
 	TSharedPtr<FUICommandInfo> LoadSelectContextPalette;
+	TSharedPtr<FUICommandInfo> LoadQueryContextPalette;
 
 	// @todo_pcg: Selection/Filtering actions.
 
@@ -76,6 +77,11 @@ public:
 	/*** PCG EDITOR MODE TOOLS - Select Context **/
 	TSharedPtr<FUICommandInfo> EnableSelectTool;
 	TArray<TSharedPtr<FUICommandInfo>> GetSelectContextToolCommands();
+
+	/*** PCG EDITOR MODE TOOLS - Query Context **/
+	TSharedPtr<FUICommandInfo> EnableQueryTool;
+	TSharedPtr<FUICommandInfo> EnableIsolateTool;
+	TArray<TSharedPtr<FUICommandInfo>> GetQueryContextToolCommands();
 };
 
 /**
@@ -98,5 +104,7 @@ DECLARE_TOOL_COMMANDS(FPCGDrawSurfaceCommands);
 DECLARE_TOOL_COMMANDS(FPCGPaintCommands);
 DECLARE_TOOL_COMMANDS(FPCGVolumeCommands);
 DECLARE_TOOL_COMMANDS(FPCGSelectCommands);
+DECLARE_TOOL_COMMANDS(FPCGQueryCommands);
+DECLARE_TOOL_COMMANDS(FPCGIsolateCommands);
 
 #undef DECLARE_TOOL_COMMANDS
diff --git a/Engine/Plugins/PCG/Source/PCGEditor/Private/EditorMode/PCGEdModeStyle.cpp b/Engine/Plugins/PCG/Source/PCGEditor/Private/EditorMode/PCGEdModeStyle.cpp
--- a/Engine/Plugins/PCG/Source/PCGEditor/Private/EditorMode/PCGEdModeStyle.cpp
+++ b/Engine/Plugins/PCG/Source/PCGEditor/Private/EditorMode/PCGEdModeStyle.cpp
@@ -65,18 +65,17 @@ FPCGEditorModeStyle::FPCGEditorModeStyle()
 		Set("PCGEditorMode.Context.Paint", new IMAGE_BRUSH_SVG("Icons/EdMode/Paint", Icon20x20));
 		Set("PCGEditorMode.Context.Volume", new IMAGE_BRUSH_SVG("Icons/EdMode/Volume", Icon20x20));
 		Set("PCGEditorMode.Context.Select", new IMAGE_BRUSH_SVG("Icons/EdMode/Select", Icon20x20));
+		Set("PCGEditorMode.Context.Query", new IMAGE_BRUSH_SVG("Icons/EdMode/Select", Icon20x20));
 	}
 
 	// Tools
 	{
 		Set("PCGEditorMode.Tools.DrawSpline", new IMAGE_BRUSH_SVG("Icons/EdMode/DrawSpline", Icon20x20));
 		Set("PCGEditorMode.Tools.DrawSurface", new IMAGE_BRUSH_SVG("Icons/EdMode/DrawSurface", Icon20x20));
-
 		Set("PCGEditorMode.Tools.Paint", new IMAGE_BRUSH_SVG("Icons/EdMode/Paint", Icon20x20));
-
 		Set("PCGEditorMode.Tools.Volume", new IMAGE_BRUSH_SVG("Icons/EdMode/Volume", Icon20x20));
-
 		Set("PCGEditorMode.Tools.Select", new IMAGE_BRUSH_SVG("Icons/EdMode/Select", Icon20x20));
+		Set("PCGEditorMode.Tools.Query", new IMAGE_BRUSH_SVG("Icons/EdMode/Select", Icon20x20));
 	}
 
 	// Backdrops and backgrounds
diff --git a/Engine/Plugins/PCG/Source/PCGEditor/Private/EditorMode/PCGEdModeToolkit.cpp b/Engine/Plugins/PCG/Source/PCGEditor/Private/EditorMode/PCGEdModeToolkit.cpp
--- a/Engine/Plugins/PCG/Source/PCGEditor/Private/EditorMode/PCGEdModeToolkit.cpp
+++ b/Engine/Plugins/PCG/Source/PCGEditor/Private/EditorMode/PCGEdModeToolkit.cpp
@@ -22,6 +22,7 @@
 #include "Styling/SlateIconFinder.h"
 #include "Tools/Line/PCGDrawSplineTool.h"
 #include "Tools/Paint/PCGPaintTool.h"
+#include "Tools/Query/PCGQueryTool.h"
 #include "Tools/Select/PCGSelectTool.h"
 #include "Tools/Volume/PCGVolumeTool.h"
 #include "Widgets/SPCGToolPresetSection.h"
@@ -530,6 +531,12 @@ void FPCGEditorModeToolkit::RegisterPalettes()
 		ToolkitBuilder->AddPalette(MoveTemp(SelectToolPalette));
 	}
 
+	{
+		FString ToolIdentifier = UPCGQueryTool::StaticGetToolTag().ToString();
+		TSharedPtr<FPCGToolPalette> QueryToolPalette = MakeShared<FPCGToolPalette>(PaletteCommands.LoadQueryContextPalette, ToolCommands.GetQueryContextToolCommands());
+		ToolkitBuilder->AddPalette(MoveTemp(QueryToolPalette));
+	}
+
 	ToolkitBuilder->SetActivePaletteOnLoad(PaletteCommands.LoadSplineContextPalette.Get());
 
 	ToolkitBuilder->UpdateWidget();
diff --git a/Engine/Plugins/PCG/Source/PCGEditor/Private/EditorMode/Tools/Helpers/PCGEdModeSceneQueryHelpers.cpp b/Engine/Plugins/PCG/Source/PCGEditor/Private/EditorMode/Tools/Helpers/PCGEdModeSceneQueryHelpers.cpp
--- a/Engine/Plugins/PCG/Source/PCGEditor/Private/EditorMode/Tools/Helpers/PCGEdModeSceneQueryHelpers.cpp
+++ b/Engine/Plugins/PCG/Source/PCGEditor/Private/EditorMode/Tools/Helpers/PCGEdModeSceneQueryHelpers.cpp
@@ -117,18 +117,6 @@ namespace UE::PCG::EditorMode::Scene
 		return ViewportRay(Ray.WorldRay, Params);
 	}
 
-	TOptional<FHitResult> ViewportRay_HitResult(const FRay& WorldRay, const FViewRayParams& Params)
-	{
-		if (TOptional<FHitResult> Result = TraceToNearestObject(WorldRay, Params))
-		{
-			return Result;
-		}
-		else
-		{
-			return {};
-		}
-	}
-
 	TOptional<FViewHitResult> ViewportRayAgainstPlane(const FInputDeviceRay& Ray, const FViewRayParams& Params, const FPlane& Plane)
 	{
 		bool bHitPlane = false;
diff --git a/Engine/Plugins/PCG/Source/PCGEditor/Private/EditorMode/Tools/PCGInteractiveToolBuilder.cpp b/Engine/Plugins/PCG/Source/PCGEditor/Private/EditorMode/Tools/PCGInteractiveToolBuilder.cpp
--- a/Engine/Plugins/PCG/Source/PCGEditor/Private/EditorMode/Tools/PCGInteractiveToolBuilder.cpp
+++ b/Engine/Plugins/PCG/Source/PCGEditor/Private/EditorMode/Tools/PCGInteractiveToolBuilder.cpp
@@ -21,6 +21,11 @@ bool UE::PCG::EditorMode::Tool::BuildTool(UInteractiveTool* InTool)
 		return false;
 	}
 
+	if ((Cast<UPCGInteractiveTool>(InTool) && !Cast<UPCGInteractiveTool>(InTool)->RequiresSettings()))
+	{
+		return true;
+	}
+
 	const FObjectProperty* SettingsProperty = nullptr;
 	const void* PropertyAddress = nullptr;
 	UClass* SettingsClass = nullptr;
@@ -29,7 +34,7 @@ bool UE::PCG::EditorMode::Tool::BuildTool(UInteractiveTool* InTool)
 		if(PropertyIt.Key()->IsA(FObjectProperty::StaticClass()))
 		{
 			const FObjectProperty* ObjectProperty = CastFieldChecked<FObjectProperty>(PropertyIt.Key());
-			if(ObjectProperty->PropertyClass->IsChildOf(UPCGInteractiveToolSettings::StaticClass()) && ObjectProperty->PropertyClass->HasAnyClassFlags(EClassFlags::CLASS_Abstract) == false)
+			if(ObjectProperty->PropertyClass->IsChildOf(UPCGInteractiveToolBaseSettings::StaticClass()) && ObjectProperty->PropertyClass->HasAnyClassFlags(EClassFlags::CLASS_Abstract) == false)
 			{
 				SettingsProperty = ObjectProperty;
 				PropertyAddress = PropertyIt.Value();
@@ -41,12 +46,12 @@ bool UE::PCG::EditorMode::Tool::BuildTool(UInteractiveTool* InTool)
 	
 	if(SettingsProperty == nullptr || PropertyAddress == nullptr || SettingsClass == nullptr)
 	{
-		FText Message = FText::FormatOrdered(LOCTEXT("PCGSettingsPropertyMissingMessage", "Settings {0} of tool {1} does not have a valid PCGInteractiveToolSettings property. Did you forget to mark it as UPROPERTY?"), SettingsClass ? SettingsClass->GetClass()->GetDisplayNameText() : FText::GetEmpty(), InTool->GetClass()->GetDisplayNameText());
+		FText Message = FText::FormatOrdered(LOCTEXT("PCGSettingsPropertyMissingMessage", "Settings {0} of tool {1} does not have a valid PCGInteractiveToolBaseSettings property. Did you forget to mark it as UPROPERTY?"), SettingsClass ? SettingsClass->GetClass()->GetDisplayNameText() : FText::GetEmpty(), InTool->GetClass()->GetDisplayNameText());
 		InTool->GetToolManager()->DisplayMessage(Message, EToolMessageLevel::UserError);
 		return false;
 	}
 	
-	UPCGInteractiveToolSettings* Settings = NewObject<UPCGInteractiveToolSettings>(InTool, SettingsClass);
+	UPCGInteractiveToolBaseSettings* Settings = NewObject<UPCGInteractiveToolBaseSettings>(InTool, SettingsClass);
 	SettingsProperty->SetObjectPropertyValue(const_cast<void*>(PropertyAddress), Settings);
 	if(TValueOrError<bool, FText> Result = Settings->Initialize(InTool); Result.HasError() || (Result.HasValue() && Result.GetValue() == false))
 	{
diff --git a/Engine/Plugins/PCG/Source/PCGEditor/Private/EditorMode/Tools/PCGInteractiveToolSettings.cpp b/Engine/Plugins/PCG/Source/PCGEditor/Private/EditorMode/Tools/PCGInteractiveToolSettings.cpp
--- a/Engine/Plugins/PCG/Source/PCGEditor/Private/EditorMode/Tools/PCGInteractiveToolSettings.cpp
+++ b/Engine/Plugins/PCG/Source/PCGEditor/Private/EditorMode/Tools/PCGInteractiveToolSettings.cpp
@@ -29,7 +29,7 @@ namespace UE::PCG::EditorMode::Tool
 	{
 		for (UObject* ToolPropertyObject : InTool->GetToolProperties(false))
 		{
-			if (UPCGInteractiveToolSettings* PCGToolSettings = Cast<UPCGInteractiveToolSettings>(ToolPropertyObject))
+			if (UPCGInteractiveToolBaseSettings* PCGToolSettings = Cast<UPCGInteractiveToolBaseSettings>(ToolPropertyObject))
 			{
 				PCGToolSettings->SaveProperties(InTool);
 
@@ -48,6 +48,11 @@ namespace UE::PCG::EditorMode::Tool
 	}
 }
 
+TValueOrError<bool, FText> UPCGInteractiveToolBaseSettings::Initialize(UInteractiveTool* InTool)
+{
+	return MakeValue(true);
+}
+
 UPCGInteractiveToolSettings::UPCGInteractiveToolSettings()
 	: NewPCGComponentName(GetDefault<UPCGEditorModeSettings>()->DefaultNewPCGComponentName)
 {
diff --git a/Engine/Plugins/PCG/Source/PCGEditor/Private/EditorMode/Tools/PCGInteractiveToolSettings.h b/Engine/Plugins/PCG/Source/PCGEditor/Private/EditorMode/Tools/PCGInteractiveToolSettings.h
--- a/Engine/Plugins/PCG/Source/PCGEditor/Private/EditorMode/Tools/PCGInteractiveToolSettings.h
+++ b/Engine/Plugins/PCG/Source/PCGEditor/Private/EditorMode/Tools/PCGInteractiveToolSettings.h
@@ -21,8 +21,32 @@ namespace UE::PCG::EditorMode::Tool
 	extern TAutoConsoleVariable<bool> CVarCommitPropertyChangedEventOnTimer;
 }
 
+/**
+* Base class for interactive tool settings in PCG.
+* It is useful to derive from this class instead of UPCGInteractiveToolSettings when the tool does not require to create an actor,
+* use a graph or a PCG component.
+*/
 UCLASS(Abstract)
-class UPCGInteractiveToolSettings : public UInteractiveToolPropertySet, public FTickableEditorObject, public IKeyInputBehaviorTarget
+class UPCGInteractiveToolBaseSettings : public UInteractiveToolPropertySet
+{
+	GENERATED_BODY()
+
+public:
+	virtual TValueOrError<bool, FText> Initialize(UInteractiveTool* InTool);
+	virtual void Apply(UInteractiveTool* OwningTool) {}
+	virtual void Cancel(UInteractiveTool* OwningTool) {}
+	virtual void Shutdown() {}
+	virtual bool IsSelectionAllowed(AActor* InActor, bool bInSelection) const { return true; }
+
+	/** Tool tag for use with filtering compatible graphs. */
+	virtual FName GetToolTag() const PURE_VIRTUAL(UPCGInteractiveToolBaseSettings::GetToolTag, return NAME_None;);
+};
+
+/**
+* Base class for tools that need to manage creating an actor, a pcg component and graph tool settings in PCG.
+*/
+UCLASS(Abstract)
+class UPCGInteractiveToolSettings : public UPCGInteractiveToolBaseSettings, public FTickableEditorObject, public IKeyInputBehaviorTarget
 {
 	GENERATED_BODY()
 
@@ -30,17 +54,17 @@ public:
 	UPCGInteractiveToolSettings();
 	virtual ~UPCGInteractiveToolSettings() override = default;
 
-	virtual TValueOrError<bool, FText> Initialize(UInteractiveTool* InTool);
-	virtual void Apply(UInteractiveTool* OwningTool);
-	virtual void Cancel(UInteractiveTool* OwningTool);
-	virtual void Shutdown();
+	virtual TValueOrError<bool, FText> Initialize(UInteractiveTool* InTool) override;
+	virtual void Apply(UInteractiveTool* OwningTool) override;
+	virtual void Cancel(UInteractiveTool* OwningTool) override;
+	virtual void Shutdown() override;
 
 	//~Begin interface UInteractiveToolPropertySet
 	virtual void RestoreProperties(UInteractiveTool* RestoreToTool, const FString& CacheIdentifier = TEXT("")) override;
 	//~End interface UInteractiveToolPropertySet
 
 	/** By default, when a tool is active, we can only select the Edit Actor in the Outliner. */
-	virtual bool IsSelectionAllowed(AActor* InActor, bool bInSelection) const;
+	virtual bool IsSelectionAllowed(AActor* InActor, bool bInSelection) const override;
 
 	/** Whether tool data can be reset. Controls visibility of the Reset button in the UI. */
 	virtual bool CanResetToolData(FName DataInstanceIdentifier);
@@ -65,9 +89,6 @@ public:
 	/** A soft object path to the graph asset to load by default for this tool. */
 	virtual FSoftObjectPath GetDefaultGraph() const;
 	
-	/** Tool tag for use with filtering compatible graphs. */
-	virtual FName GetToolTag() const PURE_VIRTUAL(UPCGInteractiveToolSettings::GetToolTag, return NAME_None;);
-
 	/** The working data type for this settings class. */
 	virtual const UScriptStruct* GetWorkingDataType() const PURE_VIRTUAL(UPCGInteractiveToolSettings::GetWorkingDataType, return FPCGInteractiveToolWorkingData::StaticStruct(););
 	
@@ -275,3 +296,12 @@ private:
 
 	TWeakObjectPtr<UInteractiveTool> Tool = nullptr;
 };
+
+UCLASS(Abstract)
+class UPCGInteractiveTool : public UInteractiveTool
+{
+	GENERATED_BODY()
+
+public:
+	virtual bool RequiresSettings() const { return true; }
+};
\ No newline at end of file
diff --git a/Engine/Plugins/PCG/Source/PCGEditor/Private/EditorMode/Tools/Query/PCGQueryTool.cpp b/Engine/Plugins/PCG/Source/PCGEditor/Private/EditorMode/Tools/Query/PCGQueryTool.cpp
new file mode 100644
--- /dev/null
+++ b/Engine/Plugins/PCG/Source/PCGEditor/Private/EditorMode/Tools/Query/PCGQueryTool.cpp
@@ -0,0 +1,349 @@
+// Copyright Epic Games, Inc. All Rights Reserved.
+
+#include "EditorMode/Tools/Query/PCGQueryTool.h"
+
+#include "PCGEditorModule.h"
+#include "EditorMode/PCGEdModeStyle.h"
+
+#include "PCGManagedResource.h"
+
+#include "EditorViewportClient.h"
+#include "Components/InstancedStaticMeshComponent.h"
+#include "Elements/Framework/TypedElementRegistry.h"
+#include "Elements/Interfaces/TypedElementObjectInterface.h"
+#include "InteractiveToolManager.h"
+#include "BaseBehaviors/SingleClickBehavior.h"
+
+#include UE_INLINE_GENERATED_CPP_BY_NAME(PCGQueryTool)
+
+#define LOCTEXT_NAMESPACE "UPCGQueryTool"
+
+namespace PCGQueryTool
+{
+	namespace Constants
+	{
+		const FName ToolTag = TEXT("QueryTool");
+		const FText ToolDisplayName = LOCTEXT("QueryToolName", "Query Tool");
+		const FText ToolDescription = LOCTEXT("QueryToolDisplayMessage", "Queries the scene from a user input (e.g. clicking) to retrieve data (tags, etc.) and optionally affects the selected content.");
+	}
+}
+
+namespace PCGIsolateTool
+{
+	namespace Constants
+	{
+		const FName ToolTag = TEXT("IsolateTool");
+	}
+}
+
+FName UPCGQueryTool::StaticGetToolTag()
+{
+	return PCGQueryTool::Constants::ToolTag;
+}
+
+void UPCGQueryTool::Setup()
+{
+	Super::Setup();
+
+	ClickBehavior = NewObject<USingleClickInputBehavior>();
+	ClickBehavior->Initialize(this);
+	AddInputBehavior(ClickBehavior);
+}
+
+void UPCGQueryTool::Shutdown(const EToolShutdownType ShutdownType)
+{
+	// We have to call this before the super function as it would remove the property sets before we can route the shutdown function to the pcg settings
+	UE::PCG::EditorMode::Tool::Shutdown(this, ShutdownType);
+
+	Super::Shutdown(ShutdownType);
+}
+
+HHitProxy* UPCGQueryTool::GetHitProxy(const FInputDeviceRay& ClickPosition) const
+{
+	FViewport* FocusedViewport = GetToolManager()->GetContextQueriesAPI()->GetFocusedViewport();
+	if (FocusedViewport)
+	{
+		if (HHitProxy* HitProxy = FocusedViewport->GetHitProxy(ClickPosition.ScreenPosition.X, ClickPosition.ScreenPosition.Y))
+		{
+			return HitProxy;
+		}
+	}
+	
+	return nullptr;
+}
+
+UActorComponent* UPCGQueryTool::GetHitComponent(HHitProxy* HitProxy) const
+{
+	if (!HitProxy)
+	{
+		return nullptr;
+	}
+	else if (HInstancedStaticMeshInstance* ISMProxy = HitProxyCast<HInstancedStaticMeshInstance>(HitProxy))
+	{
+		return ISMProxy->Component.Get();
+	}
+	// implementation note: also catches HActor.
+	else if (TTypedElement<ITypedElementObjectInterface> ObjectInterface = UTypedElementRegistry::GetInstance()->GetElement<ITypedElementObjectInterface>(HitProxy->GetElementHandle()))
+	{
+		return ObjectInterface.GetObjectAs<UActorComponent>();
+	}
+
+	// Unknown/incompatible proxy type
+	return nullptr;
+}
+
+UActorComponent* UPCGQueryTool::GetHitComponent(const FInputDeviceRay& ClickPosition) const
+{
+	return GetHitComponent(GetHitProxy(ClickPosition));
+}
+
+FInputRayHit UPCGQueryTool::IsHitByClick(const FInputDeviceRay& ClickPosition)
+{
+	using namespace UE::PCG::EditorMode;
+
+	if (UActorComponent* HitComponent = GetHitComponent(ClickPosition))
+	{
+		FInputRayHit Result;
+		Result.bHit = true;
+		Result.HitObject = HitComponent;
+
+		return Result;
+	}
+
+	return {};
+}
+
+void UPCGQueryTool::OnClicked(const FInputDeviceRay& ClickPosition)
+{
+	using namespace UE::PCG::EditorMode;
+
+	if (HHitProxy* HitProxy = GetHitProxy(ClickPosition))
+	{
+		UActorComponent* HitComponent = GetHitComponent(HitProxy);
+
+		int32 InstanceIndex = INDEX_NONE;
+		if (HInstancedStaticMeshInstance* ISMProxy = HitProxyCast<HInstancedStaticMeshInstance>(HitProxy))
+		{
+			InstanceIndex = ISMProxy->InstanceIndex;
+		}
+
+		OnHitProxy(HitProxy, HitComponent, InstanceIndex);
+	}
+}
+
+void UPCGQueryTool::VisitHierarchy(UActorComponent* HitComponent, int32 InstanceIndex, TFunctionRef<void(UActorComponent*, AActor*, int32)> Func)
+{
+	if (!HitComponent)
+	{
+		return;
+	}
+
+	// Start from the hit component, get its name + tags
+	UActorComponent* CurrentComponent = HitComponent;
+
+	while (CurrentComponent)
+	{
+		Func(CurrentComponent, nullptr, InstanceIndex);
+		// The instance index is valid just for the first component we're visiting
+		InstanceIndex = INDEX_NONE;
+		
+		UActorComponent* ParentComponent = nullptr;
+		if (USceneComponent* SceneComponent = Cast<USceneComponent>(CurrentComponent))
+		{
+			ParentComponent = SceneComponent->GetAttachParent();
+		}
+
+		CurrentComponent = (CurrentComponent == ParentComponent ? nullptr : ParentComponent);
+	}
+
+	// go up the root hierarchy
+	AActor* CurrentActor = HitComponent->GetOwner();
+
+	while (CurrentActor)
+	{
+		Func(nullptr, CurrentActor, INDEX_NONE);
+
+		AActor* NewActor = CurrentActor->GetSceneOutlinerParent();
+		CurrentActor = (CurrentActor == NewActor ? nullptr : NewActor);
+	}
+}
+
+void UPCGQueryTool::OnHitProxy(HHitProxy* HitProxy, UActorComponent* HitComponent, int32 InstanceIndex)
+{
+	auto OutputToLog = [](UActorComponent* CurrentComponent, AActor* CurrentActor, int32 InstanceIndex)
+	{
+		if(CurrentComponent)
+		{
+			UE_LOG(LogPCGEditor, Warning, TEXT("Component: '%s'"), *CurrentComponent->GetName());
+			
+			if (InstanceIndex != INDEX_NONE)
+			{
+				UE_LOG(LogPCGEditor, Warning, TEXT("Instance: '%d'"), InstanceIndex);
+			}
+
+			if (!CurrentComponent->ComponentTags.IsEmpty())
+			{
+				UE_LOG(LogPCGEditor, Warning, TEXT("Tags: "));
+				for (FName Tag : CurrentComponent->ComponentTags)
+				{
+					UE_LOG(LogPCGEditor, Warning, TEXT("      %s"), *Tag.ToString());
+				}
+			}
+		}
+		else if (CurrentActor)
+		{
+			UE_LOG(LogPCGEditor, Warning, TEXT("Actor: '%s'"), *CurrentActor->GetActorLabel());
+			if (!CurrentActor->Tags.IsEmpty())
+			{
+				UE_LOG(LogPCGEditor, Warning, TEXT("Tags: "));
+				for (FName Tag : CurrentActor->Tags)
+				{
+					UE_LOG(LogPCGEditor, Warning, TEXT("      %s"), *Tag.ToString());
+				}
+			}
+		}
+	};
+
+	VisitHierarchy(HitComponent, InstanceIndex, OutputToLog);
+}
+
+UInteractiveTool* UPCGQueryToolBuilder::BuildTool(const FToolBuilderState& SceneState) const
+{
+	UPCGQueryTool* NewTool = NewObject<UPCGQueryTool>(SceneState.ToolManager, ToolClass);
+
+	NewTool->SetToolInfo(
+	{
+		.ToolDisplayName = PCGQueryTool::Constants::ToolDisplayName,
+		.ToolDisplayMessage = PCGQueryTool::Constants::ToolDescription,
+		.ToolIcon = FPCGEditorModeStyle::Get().GetBrush("PCGEditorMode.Tools.Query")
+	});
+
+	if (!UE::PCG::EditorMode::Tool::BuildTool(NewTool))
+	{
+		return nullptr;
+	}
+	else
+	{
+		return NewTool;
+	}
+}
+
+void UPCGQueryToolBuilder::SetToolClass(TSubclassOf<UPCGQueryTool> InClass)
+{
+	ToolClass = InClass;
+}
+
+void UPCGInteractiveToolSettings_Isolate::Apply(UInteractiveTool* OwningTool)
+{
+	if (AActor* Actor = SelectedActor.Get())
+	{
+		if (UPCGComponent* PCGComponent = Actor->GetComponentByClass<UPCGComponent>())
+		{
+			FPCGMoveResourceParams Params
+			{
+				.RequiredTags = SelectedTags
+			};
+
+			AActor* IsolatedActor = PCGComponent->ClearPCGLink(Params);
+
+			if (IsolatedActor)
+			{
+				// @todo_pcg: Push tags onto actor?
+				// @todo_pcg: Clear PCG generated component?
+			}
+		}
+	}
+}
+
+FName UPCGInteractiveToolSettings_Isolate::GetToolTag() const
+{
+	return UPCGIsolateTool::StaticGetToolTag();
+}
+
+bool UPCGInteractiveToolSettings_Isolate::IsSelectionAllowed(AActor* InActor, bool bInSelection) const
+{
+	return InActor && InActor->GetComponentByClass<UPCGComponent>() != nullptr;
+}
+
+FName UPCGIsolateTool::StaticGetToolTag()
+{
+	return PCGIsolateTool::Constants::ToolTag;
+}
+
+void UPCGIsolateTool::Setup()
+{
+	AddToolPropertySource(ToolSettings);
+	Super::Setup();
+}
+
+bool UPCGIsolateTool::CanAccept() const
+{
+	// Can return true if we picked an actor that has a pcg component.
+	if (UPCGInteractiveToolSettings_Isolate* Settings = GetSettings())
+	{
+		if (AActor* SelectedActor = Settings->SelectedActor.Get())
+		{
+			return SelectedActor->GetComponentByClass<UPCGComponent>() != nullptr;
+		}
+	}
+
+	return false;
+}
+
+UPCGInteractiveToolSettings_Isolate* UPCGIsolateTool::GetSettings() const
+{
+	TArray<UObject*> AllToolProperties = GetToolProperties(false);
+
+	UObject** FoundProperties = AllToolProperties.FindByPredicate([](const UObject* Candidate)
+		{
+			return Candidate->IsA<UPCGInteractiveToolSettings_Isolate>();
+		});
+
+	if (FoundProperties)
+	{
+		return Cast<UPCGInteractiveToolSettings_Isolate>(*FoundProperties);
+	}
+
+	return nullptr;
+}
+
+void UPCGIsolateTool::OnHitProxy(HHitProxy* HitProxy, UActorComponent* HitComponent, int InstanceIndex)
+{
+	AActor* FoundActor = nullptr;
+	TSet<FName> FoundTags;
+	int32 FoundInstanceIndex = INDEX_NONE;
+
+	auto GatherTags = [&FoundTags, &FoundActor, &FoundInstanceIndex](UActorComponent* CurrentComponent, AActor* CurrentActor, int32 InstanceIndex)
+	{
+		if(CurrentComponent)
+		{
+			FoundTags.Append(CurrentComponent->ComponentTags);
+		}
+
+		if (CurrentActor)
+		{
+			if (!FoundActor)
+			{
+				FoundActor = CurrentActor;
+			}
+
+			FoundTags.Append(CurrentActor->Tags);
+		}
+
+		if (InstanceIndex != INDEX_NONE)
+		{
+			FoundInstanceIndex = InstanceIndex;
+		}
+	};
+
+	VisitHierarchy(HitComponent, InstanceIndex, GatherTags);
+
+	if (UPCGInteractiveToolSettings_Isolate* Settings = GetSettings())
+	{
+		Settings->SelectedActor = FoundActor;
+		Settings->Tags = FoundTags;
+		Settings->InstanceIndex = FoundInstanceIndex;
+	}
+}
+
+#undef LOCTEXT_NAMESPACE
\ No newline at end of file
diff --git a/Engine/Plugins/PCG/Source/PCGEditor/Private/EditorMode/Tools/Query/PCGQueryTool.h b/Engine/Plugins/PCG/Source/PCGEditor/Private/EditorMode/Tools/Query/PCGQueryTool.h
new file mode 100644
--- /dev/null
+++ b/Engine/Plugins/PCG/Source/PCGEditor/Private/EditorMode/Tools/Query/PCGQueryTool.h
@@ -0,0 +1,131 @@
+// Copyright Epic Games, Inc. All Rights Reserved.
+
+#pragma once
+
+#include "EditorMode/Tools/PCGInteractiveToolBuilder.h"
+#include "EditorMode/Tools/PCGInteractiveToolCommon.h"
+#include "EditorMode/Tools/PCGInteractiveToolSettings.h"
+#include "EditorMode/Tools/Helpers/PCGEdModeSceneQueryHelpers.h"
+
+#include "BaseBehaviors/BehaviorTargetInterfaces.h"
+
+#include "PCGQueryTool.generated.h"
+
+class HHitProxy;
+class UActorComponent;
+class USingleClickInputBehavior;
+
+/**
+* Simple editor tool to gather information from a picked object in the world. 
+* By default logs visited hierarchy tags (from touched component, to actor, ...) but can be reimplemented in a derived class.
+*/
+UCLASS()
+class UPCGQueryTool : public UPCGInteractiveTool, public IClickBehaviorTarget
+{
+	GENERATED_BODY()
+
+public:
+	//~Begin UInteractiveTool interface
+	virtual void Setup() override;
+	virtual void Shutdown(EToolShutdownType ShutdownType) override;
+	virtual bool HasCancel() const override { return false; }
+	virtual bool HasAccept() const override { return false; }
+	virtual bool CanAccept() const override { return false; }
+	//~End UInteractiveTool interface
+
+	//~Begin UPCGInteractiveTool interface
+	virtual bool RequiresSettings() const override { return false; }
+	//~End UPCGInteractiveTool interface
+
+	//~Begin IClickBehaviorTarget interface
+	virtual FInputRayHit IsHitByClick(const FInputDeviceRay& ClickPosition) override;
+	virtual void OnClicked(const FInputDeviceRay& ClickPosition) override;
+	//~End IClickBehaviorTarget interface
+
+	static FName StaticGetToolTag();
+
+protected:
+	HHitProxy* GetHitProxy(const FInputDeviceRay& ClickPosition) const;
+	UActorComponent* GetHitComponent(HHitProxy* HitProxy) const;
+	UActorComponent* GetHitComponent(const FInputDeviceRay& ClickPosition) const;
+	virtual void OnHitProxy(HHitProxy* HitProxy, UActorComponent* HitComponent, int InstanceIndex);
+	void VisitHierarchy(UActorComponent* HitComponent, int32 InstanceIndex, TFunctionRef<void(UActorComponent*, AActor*, int32)> Func);
+
+private:
+	UPROPERTY(Instanced)
+	TObjectPtr<USingleClickInputBehavior> ClickBehavior = nullptr;
+};
+
+UCLASS(Transient)
+class UPCGQueryToolBuilder : public UPCGInteractiveToolBuilder
+{
+	GENERATED_BODY()
+
+public:
+	virtual bool CanBuildTool(const FToolBuilderState& SceneState) const override { return true; }
+	virtual UInteractiveTool* BuildTool(const FToolBuilderState& SceneState) const override;
+
+	void SetToolClass(TSubclassOf<UPCGQueryTool> InClass);
+
+private:
+	UPROPERTY()
+	TSubclassOf<UPCGQueryTool> ToolClass;
+};
+
+UCLASS()
+class UPCGInteractiveToolSettings_Isolate : public UPCGInteractiveToolBaseSettings
+{
+	GENERATED_BODY()
+
+public:
+	virtual void Apply(UInteractiveTool* OwningTool) override;
+	virtual bool IsSelectionAllowed(AActor* InActor, bool bInSelection) const override;
+	virtual FName GetToolTag() const override;
+	//virtual const UScriptStruct* GetWorkingDataType() const { return nullptr; }
+
+	UPROPERTY(VisibleAnywhere, Category = PCG)
+	TWeakObjectPtr<AActor> SelectedActor;
+
+	UPROPERTY(VisibleAnywhere, Category = PCG)
+	TSet<FName> Tags;
+
+	UPROPERTY(VisibleAnywhere, Category = PCG)
+	int32 InstanceIndex = INDEX_NONE;
+
+	UPROPERTY(EditAnywhere, Category = PCG)
+	TArray<FName> SelectedTags;
+};
+
+/**
+* Tool that extracts the artifacts (through Clear PCG Link) from an actor with a PCG component according to selected tags.
+* Tags are gathered through picking (from the Query tool), and listed so the user can choose a combination.
+*/
+UCLASS()
+class UPCGIsolateTool : public UPCGQueryTool
+{
+	GENERATED_BODY()
+
+public:
+	//~Begin UInteractiveTool interface
+	virtual void Setup() override;
+	virtual bool HasCancel() const override { return true; }
+	virtual bool HasAccept() const override { return true; }
+	virtual bool CanAccept() const override;
+	//~End UInteractiveTool interface
+
+	//~Begin UPCGInteractiveTool interface
+	virtual bool RequiresSettings() const override { return true; }
+	//~End UPCGInteractiveTool interface
+
+	static FName StaticGetToolTag();
+
+protected:
+	//~Begin UPCGQueryTool interface
+	virtual void OnHitProxy(HHitProxy* HitProxy, UActorComponent* HitComponent, int InstanceIndex) override;
+	//~End UPCGQueryTool interface
+
+	UPCGInteractiveToolSettings_Isolate* GetSettings() const;
+
+	UPROPERTY(Transient)
+	TObjectPtr<UPCGInteractiveToolSettings_Isolate> ToolSettings;
+};
\ No newline at end of file
diff --git a/Engine/Plugins/PCG/Source/PCGEditor/Private/EditorMode/Widgets/SPCGToolPresetSection.cpp b/Engine/Plugins/PCG/Source/PCGEditor/Private/EditorMode/Widgets/SPCGToolPresetSection.cpp
--- a/Engine/Plugins/PCG/Source/PCGEditor/Private/EditorMode/Widgets/SPCGToolPresetSection.cpp
+++ b/Engine/Plugins/PCG/Source/PCGEditor/Private/EditorMode/Widgets/SPCGToolPresetSection.cpp
@@ -37,7 +37,7 @@ void SPCGToolPresetSection::UpdatePresets()
 		return;
 	}
 
-	UPCGInteractiveToolSettings* ToolSettings = EdMode->GetCurrentToolSettings();
+	UPCGInteractiveToolSettings* ToolSettings = Cast<UPCGInteractiveToolSettings>(EdMode->GetCurrentToolSettings());
 
 	if(ToolSettings == nullptr)
 	{
@@ -163,7 +163,7 @@ void SPCGToolPresetSection::ActivatePreset(TWeakPtr<FPresetItem> PresetItemWeakP
 
 	UPCGEditorMode* EdMode = CastChecked<UPCGEditorMode>(Context->ThisSection.Pin()->OwningEditorMode.Get());
 
-	if (UPCGInteractiveToolSettings* ToolSettings = EdMode->GetCurrentToolSettings())
+	if (UPCGInteractiveToolSettings* ToolSettings = Cast<UPCGInteractiveToolSettings>(EdMode->GetCurrentToolSettings()))
 	{
 		if(UPCGGraphInterface* Graph = PresetItemWeakPtr.Pin()->Item.LoadSynchronous())
 		{
@@ -193,7 +193,7 @@ ECheckBoxState SPCGToolPresetSection::IsPresetActive(TWeakPtr<FPresetItem> Prese
 	
 	UPCGEditorMode* EdMode = CastChecked<UPCGEditorMode>(Context->ThisSection.Pin()->OwningEditorMode.Get());
 	
-	if (UPCGInteractiveToolSettings* ToolSettings = EdMode->GetCurrentToolSettings())
+	if (UPCGInteractiveToolSettings* ToolSettings = Cast<UPCGInteractiveToolSettings>(EdMode->GetCurrentToolSettings()))
 	{
 		if (UPCGGraphInterface* Graph = PresetItem->Item.LoadSynchronous())
 		{
