[PCG] Add end/join controls to create polygon 2d & offset polygon 2d.

#Synced-CL 48172364

diff --git a/Engine/Plugins/PCG/Source/PCG/Private/Elements/Polygon/PCGCreatePolygon2D.cpp b/Engine/Plugins/PCG/Source/PCG/Private/Elements/Polygon/PCGCreatePolygon2D.cpp
--- a/Engine/Plugins/PCG/Source/PCG/Private/Elements/Polygon/PCGCreatePolygon2D.cpp
+++ b/Engine/Plugins/PCG/Source/PCG/Private/Elements/Polygon/PCGCreatePolygon2D.cpp
@@ -236,8 +236,8 @@ bool FPCGCreatePolygonElement::ExecuteInternal(FPCGContext* Context) const
 			TArray<UE::Geometry::FGeneralPolygon2d> OffsetPolygons;
 
 			//Note: same values as FGeometryScriptOpenPathOffsetOptions, but could be parametrized in advanced settings
-			const UE::Geometry::EPolygonOffsetEndType EndType = UE::Geometry::EPolygonOffsetEndType::Square;
-			const UE::Geometry::EPolygonOffsetJoinType JoinType = UE::Geometry::EPolygonOffsetJoinType::Square;
+			const UE::Geometry::EPolygonOffsetEndType EndType = PCGPolygon2DUtils::GetEndType(Settings->EndType);
+			const UE::Geometry::EPolygonOffsetJoinType JoinType = PCGPolygon2DUtils::GetJoinType(Settings->JoinType);
 			const double MiterLimit = 2.0;
 			const double MaxStepsPerRadian = 10.0;
 			const double DefaultStepsPerRadianScale = 1.0;
diff --git a/Engine/Plugins/PCG/Source/PCG/Private/Elements/Polygon/PCGOffsetPolygon2D.cpp b/Engine/Plugins/PCG/Source/PCG/Private/Elements/Polygon/PCGOffsetPolygon2D.cpp
--- a/Engine/Plugins/PCG/Source/PCG/Private/Elements/Polygon/PCGOffsetPolygon2D.cpp
+++ b/Engine/Plugins/PCG/Source/PCG/Private/Elements/Polygon/PCGOffsetPolygon2D.cpp
@@ -95,7 +95,7 @@ bool FPCGOffsetPolygonElement::ExecuteInternal(FPCGContext* Context) const
 		}
 
 		const double MiterLimit = 2.0;
-		const UE::Geometry::EPolygonOffsetJoinType JoinType = UE::Geometry::EPolygonOffsetJoinType::Square;
+		const UE::Geometry::EPolygonOffsetJoinType JoinType = PCGPolygon2DUtils::GetJoinType(Settings->JoinType);
 		const UE::Geometry::EPolygonOffsetEndType EndType = UE::Geometry::EPolygonOffsetEndType::Polygon;
 
 		if (Settings->Operation == EPCGPolygonOffsetOperation::Offset)
diff --git a/Engine/Plugins/PCG/Source/PCG/Private/Elements/Polygon/PCGPolygon2DUtils.cpp b/Engine/Plugins/PCG/Source/PCG/Private/Elements/Polygon/PCGPolygon2DUtils.cpp
--- a/Engine/Plugins/PCG/Source/PCG/Private/Elements/Polygon/PCGPolygon2DUtils.cpp
+++ b/Engine/Plugins/PCG/Source/PCG/Private/Elements/Polygon/PCGPolygon2DUtils.cpp
@@ -5,6 +5,42 @@
 namespace PCGPolygon2DUtils
 {
 
+UE::Geometry::EPolygonOffsetJoinType GetJoinType(EPCGPolygonJoinType JoinType)
+{
+	using namespace UE::Geometry;
+	switch (JoinType)
+	{
+	case EPCGPolygonJoinType::Square:
+		return EPolygonOffsetJoinType::Square;
+	case EPCGPolygonJoinType::Round:
+		return EPolygonOffsetJoinType::Round;
+	case EPCGPolygonJoinType::Miter:
+		return EPolygonOffsetJoinType::Miter;
+	default:
+		checkNoEntry();
+	}
+	return EPolygonOffsetJoinType::Square;
+}
+
+UE::Geometry::EPolygonOffsetEndType GetEndType(EPCGPolygonEndType EndType)
+{
+	using namespace UE::Geometry;
+	switch (EndType)
+	{
+	case EPCGPolygonEndType::Butt:
+		return EPolygonOffsetEndType::Butt;
+	case EPCGPolygonEndType::Square:
+		return EPolygonOffsetEndType::Square;
+	case EPCGPolygonEndType::Round:
+		return EPolygonOffsetEndType::Round;
+	case EPCGPolygonEndType::Polygon:
+		return EPolygonOffsetEndType::Polygon;
+	default:
+		checkNoEntry();
+	}
+	return EPolygonOffsetEndType::Square;
+}
+
 TArray<FPCGPinProperties> DefaultPolygonInputPinProperties()
 {
 	TArray<FPCGPinProperties> PinProperties;
diff --git a/Engine/Plugins/PCG/Source/PCG/Public/Elements/Polygon/PCGCreatePolygon2D.h b/Engine/Plugins/PCG/Source/PCG/Public/Elements/Polygon/PCGCreatePolygon2D.h
--- a/Engine/Plugins/PCG/Source/PCG/Public/Elements/Polygon/PCGCreatePolygon2D.h
+++ b/Engine/Plugins/PCG/Source/PCG/Public/Elements/Polygon/PCGCreatePolygon2D.h
@@ -80,4 +80,12 @@ public:
 	/** Default axis to use when the input has an ambiguous shape (lines, collinear polygons, etc.) */
 	UPROPERTY(EditAnywhere, BlueprintReadWrite, Category = Settings, AdvancedDisplay, meta = (PCG_Overridable))
 	FVector DefaultPolygonNormal = FVector::ZAxisVector;
+
+	/** Controls behavior of how joining points are resolved. See GeometryScript for more details. */
+	UPROPERTY(EditAnywhere, BlueprintReadWrite, Category = Settings, AdvancedDisplay, meta = (EditCondition = "InputType != EPCGCreatePolygonInputType::ForceClosed", PCG_Overridable))
+	EPCGPolygonJoinType JoinType = EPCGPolygonJoinType::Square;
+
+	/** Controls behavior of how end points are resolved. See GeometryScript for more details. */
+	UPROPERTY(EditAnywhere, BlueprintReadWrite, Category = Settings, AdvancedDisplay, meta = (EditCondition = "InputType != EPCGCreatePolygonInputType::ForceClosed", PCG_Overridable))
+	EPCGPolygonEndType EndType = EPCGPolygonEndType::Square;
 };
\ No newline at end of file
diff --git a/Engine/Plugins/PCG/Source/PCG/Public/Elements/Polygon/PCGOffsetPolygon2D.h b/Engine/Plugins/PCG/Source/PCG/Public/Elements/Polygon/PCGOffsetPolygon2D.h
--- a/Engine/Plugins/PCG/Source/PCG/Public/Elements/Polygon/PCGOffsetPolygon2D.h
+++ b/Engine/Plugins/PCG/Source/PCG/Public/Elements/Polygon/PCGOffsetPolygon2D.h
@@ -62,4 +62,8 @@ public:
 
 	UPROPERTY(BlueprintReadWrite, EditAnywhere, Category = Settings, meta = (EditCondition = "!bUseOffsetFromInput", EditConditionHides, PCG_Overridable))
 	double Offset = 100.0;
+
+	/** Controls behavior of how joining points are resolved. See GeometryScript for more details. */
+	UPROPERTY(EditAnywhere, BlueprintReadWrite, Category = Settings, AdvancedDisplay, meta = (PCG_Overridable))
+	EPCGPolygonJoinType JoinType = EPCGPolygonJoinType::Square;
 };
\ No newline at end of file
diff --git a/Engine/Plugins/PCG/Source/PCG/Public/Elements/Polygon/PCGPolygon2DUtils.h b/Engine/Plugins/PCG/Source/PCG/Public/Elements/Polygon/PCGPolygon2DUtils.h
--- a/Engine/Plugins/PCG/Source/PCG/Public/Elements/Polygon/PCGPolygon2DUtils.h
+++ b/Engine/Plugins/PCG/Source/PCG/Public/Elements/Polygon/PCGPolygon2DUtils.h
@@ -4,8 +4,39 @@
 
 #include "PCGSettings.h"
 
+#include "Curve/PolygonOffsetUtils.h"
+
+#include "PCGPolygon2DUtils.generated.h"
+
+UENUM(BlueprintType)
+enum class EPCGPolygonJoinType : uint8
+{
+	/* Uniform squaring on all convex edge joins. */
+	Square,
+	/* Arcs on all convex edge joins. */
+	Round,
+	/* Squaring of convex edge joins with acute angles ("spikes"). Use in combination with MiterLimit. */
+	Miter,
+};
+
+UENUM(BlueprintType)
+enum class EPCGPolygonEndType : uint8
+{
+	/* Offsets only one side of a closed path */
+	Polygon UMETA(Hidden),
+	/* Offsets both sides of a path, with square blunt ends */
+	Butt,
+	/* Offsets both sides of a path, with square extended ends */
+	Square,
+	/* Offsets both sides of a path, with round extended ends */
+	Round,
+};
+
 namespace PCGPolygon2DUtils
 {
+	UE::Geometry::EPolygonOffsetJoinType GetJoinType(EPCGPolygonJoinType JoinType);
+	UE::Geometry::EPolygonOffsetEndType GetEndType(EPCGPolygonEndType EndType);
+
 	namespace Constants
 	{
 		const FLazyName ClipPolysLabel = TEXT("ClipPolygons");
