[PCG] Fix UPCGPointArrayData::GetTransformCopy not respecting inheritance
- revert collapse to use GetTransformCopy again now that it is fixed.

#Synced-CL 48211924

diff --git a/Engine/Plugins/PCG/Source/PCG/Private/Data/PCGPointArrayData.cpp b/Engine/Plugins/PCG/Source/PCG/Private/Data/PCGPointArrayData.cpp
--- a/Engine/Plugins/PCG/Source/PCG/Private/Data/PCGPointArrayData.cpp
+++ b/Engine/Plugins/PCG/Source/PCG/Private/Data/PCGPointArrayData.cpp
@@ -256,6 +256,17 @@ void UPCGPointArrayData::CopyUnallocatedPropertiesFrom(const UPCGBasePointData*
 	}
 }
 
+TArray<FTransform> UPCGPointArrayData::GetTransformsCopy() const
+{
+	TArray<FTransform> TransformsCopy;
+	const TConstPCGValueRange<FTransform> ReadValueRange = GetConstTransformValueRange();
+	
+	TransformsCopy.Reserve(ReadValueRange.Num());
+	Algo::Transform(ReadValueRange, TransformsCopy, [](const FTransform& Transform) { return Transform; });
+
+	return TransformsCopy;
+}
+
 FPCGPointTransform::ValueRange UPCGPointArrayData::GetTransformValueRange(bool bAllocate)
 { 
 	FlattenPropertyIfNeeded<EPCGPointNativeProperties::Transform>();
diff --git a/Engine/Plugins/PCG/Source/PCG/Private/Elements/PCGCollapsePoints.cpp b/Engine/Plugins/PCG/Source/PCG/Private/Elements/PCGCollapsePoints.cpp
--- a/Engine/Plugins/PCG/Source/PCG/Private/Elements/PCGCollapsePoints.cpp
+++ b/Engine/Plugins/PCG/Source/PCG/Private/Elements/PCGCollapsePoints.cpp
@@ -10,7 +10,6 @@
 #include "SpatialAlgo/PCGOctreeQueries.h"
 
 #include "Algo/MaxElement.h"
-#include "Algo/Transform.h"
 
 #include UE_INLINE_GENERATED_CPP_BY_NAME(PCGCollapsePoints)
 
@@ -308,15 +307,11 @@ PRAGMA_DISABLE_DEPRECATION_WARNINGS
 		OutState.SourceData = Cast<UPCGPointData>(PointData);
 PRAGMA_ENABLE_DEPRECATION_WARNINGS
 
-		const int32 NumPoints = PointData->GetNumPoints();
-
 		// Get points, octree will be rebuilt later
-		OutState.PointTransforms.Reserve(NumPoints);
-
-		const TConstPCGValueRange<FTransform> ReadValueRange = PointData->GetConstTransformValueRange();
-		Algo::Transform(ReadValueRange, OutState.PointTransforms, [](const FTransform& Transform) { return Transform; });
+		OutState.PointTransforms = PointData->GetTransformsCopy();
 
 		// Get weights
+		const int32 NumPoints = PointData->GetNumPoints();
 		OutState.Weights.SetNumUninitialized(NumPoints);
 
 		ensure(OutState.PointTransforms.Num() == OutState.Weights.Num());
diff --git a/Engine/Plugins/PCG/Source/PCG/Public/Data/PCGPointArrayData.h b/Engine/Plugins/PCG/Source/PCG/Public/Data/PCGPointArrayData.h
--- a/Engine/Plugins/PCG/Source/PCG/Public/Data/PCGPointArrayData.h
+++ b/Engine/Plugins/PCG/Source/PCG/Public/Data/PCGPointArrayData.h
@@ -7,8 +7,6 @@
 
 #include "PCGPointArrayData.generated.h"
 
-#define UE_API PCG_API
-
 extern PCG_API TAutoConsoleVariable<bool> CVarPCGEnablePointArrayDataParenting;
 
 UCLASS(MinimalAPI, BlueprintType, ClassGroup = (Procedural))
@@ -17,58 +15,58 @@ class UPCGPointArrayData : public UPCGBasePointData
 	GENERATED_BODY()
 public:	
 	//~Begin UObject Interface
-	UE_API virtual void GetResourceSizeEx(FResourceSizeEx& CumulativeResourceSize) override;
+	PCG_API virtual void GetResourceSizeEx(FResourceSizeEx& CumulativeResourceSize) override;
 	//~End UObject Interface
 
-	UE_API virtual void Flatten() override;
+	PCG_API virtual void Flatten() override;
 
-	UE_API virtual void VisitDataNetwork(TFunctionRef<void(const UPCGData*)> Action) const override;
-	UE_API virtual const UPCGPointData* ToPointData(FPCGContext* Context, const FBox& InBounds = FBox(EForceInit::ForceInit)) const override;
+	PCG_API virtual void VisitDataNetwork(TFunctionRef<void(const UPCGData*)> Action) const override;
+	PCG_API virtual const UPCGPointData* ToPointData(FPCGContext* Context, const FBox& InBounds = FBox(EForceInit::ForceInit)) const override;
 	virtual const UPCGPointArrayData* ToPointArrayData(FPCGContext* Context, const FBox& InBounds = FBox(EForceInit::ForceInit)) const override { return this; }
-	UE_API virtual UPCGSpatialData* CopyInternal(FPCGContext* Context) const override;
-	UE_API virtual void CopyPropertiesTo(UPCGBasePointData* To, int32 ReadStartIndex, int32 WriteStartIndex, int32 Count, EPCGPointNativeProperties Properties) const override;
-	UE_API virtual EPCGPointNativeProperties GetAllocatedProperties(bool bWithInheritance = true) const override;
-	UE_API virtual bool SupportsSpatialDataInheritance() const override;
+	PCG_API virtual UPCGSpatialData* CopyInternal(FPCGContext* Context) const override;
+	PCG_API virtual void CopyPropertiesTo(UPCGBasePointData* To, int32 ReadStartIndex, int32 WriteStartIndex, int32 Count, EPCGPointNativeProperties Properties) const override;
+	PCG_API virtual EPCGPointNativeProperties GetAllocatedProperties(bool bWithInheritance = true) const override;
+	PCG_API virtual bool SupportsSpatialDataInheritance() const override;
 	virtual bool HasSpatialDataParent() const override { return ParentData != nullptr; }
 
 	//~Begin UPCGBasePointData Interface
 	virtual bool IsValidRef(const PCGPointOctree::FPointRef& InPointRef) const override { return InPointRef.Index >= 0 && InPointRef.Index < GetNumPoints(); }
 	
 	virtual int32 GetNumPoints() const override { return PointArray.GetNumPoints(); }
-	UE_API virtual void SetNumPoints(int32 InNumPoints, bool bInitializeValues = true) override;
-	UE_API void RemoveAt(int32 Index);
+	PCG_API virtual void SetNumPoints(int32 InNumPoints, bool bInitializeValues = true) override;
+	PCG_API void RemoveAt(int32 Index);
 	
-	UE_API virtual void AllocateProperties(EPCGPointNativeProperties Properties) override;
-	UE_API virtual void FreeProperties(EPCGPointNativeProperties Properties) override;
-	UE_API virtual void MoveRange(int32 RangeStartIndex, int32 MoveToIndex, int32 NumElements) override;
-	UE_API virtual void CopyUnallocatedPropertiesFrom(const UPCGBasePointData* InPointData) override;
-
-	virtual TArray<FTransform> GetTransformsCopy() const { return PointArray.GetTransformCopy(); }
-
-	UE_API virtual FPCGPointTransform::ValueRange GetTransformValueRange(bool bAllocate = true) override;
-	UE_API virtual FPCGPointDensity::ValueRange GetDensityValueRange(bool bAllocate = true) override;
-	UE_API virtual FPCGPointBoundsMin::ValueRange GetBoundsMinValueRange(bool bAllocate = true) override;
-	UE_API virtual FPCGPointBoundsMax::ValueRange GetBoundsMaxValueRange(bool bAllocate = true) override;
-	UE_API virtual FPCGPointColor::ValueRange GetColorValueRange(bool bAllocate = true) override;
-	UE_API virtual FPCGPointSteepness::ValueRange GetSteepnessValueRange(bool bAllocate = true) override;
-	UE_API virtual FPCGPointSeed::ValueRange GetSeedValueRange(bool bAllocate = true) override;
-	UE_API virtual FPCGPointMetadataEntry::ValueRange GetMetadataEntryValueRange(bool bAllocate = true) override;
-
-	UE_API virtual FPCGPointTransform::ConstValueRange GetConstTransformValueRange() const override;
-	UE_API virtual FPCGPointDensity::ConstValueRange GetConstDensityValueRange() const override;
-	UE_API virtual FPCGPointBoundsMin::ConstValueRange GetConstBoundsMinValueRange() const override;
-	UE_API virtual FPCGPointBoundsMax::ConstValueRange GetConstBoundsMaxValueRange() const override;
-	UE_API virtual FPCGPointColor::ConstValueRange GetConstColorValueRange() const override;
-	UE_API virtual FPCGPointSteepness::ConstValueRange GetConstSteepnessValueRange() const override;
-	UE_API virtual FPCGPointSeed::ConstValueRange GetConstSeedValueRange() const override;
-	UE_API virtual FPCGPointMetadataEntry::ConstValueRange GetConstMetadataEntryValueRange() const override;
+	PCG_API virtual void AllocateProperties(EPCGPointNativeProperties Properties) override;
+	PCG_API virtual void FreeProperties(EPCGPointNativeProperties Properties) override;
+	PCG_API virtual void MoveRange(int32 RangeStartIndex, int32 MoveToIndex, int32 NumElements) override;
+	PCG_API virtual void CopyUnallocatedPropertiesFrom(const UPCGBasePointData* InPointData) override;
+
+	PCG_API virtual TArray<FTransform> GetTransformsCopy() const override;
+
+	PCG_API virtual FPCGPointTransform::ValueRange GetTransformValueRange(bool bAllocate = true) override;
+	PCG_API virtual FPCGPointDensity::ValueRange GetDensityValueRange(bool bAllocate = true) override;
+	PCG_API virtual FPCGPointBoundsMin::ValueRange GetBoundsMinValueRange(bool bAllocate = true) override;
+	PCG_API virtual FPCGPointBoundsMax::ValueRange GetBoundsMaxValueRange(bool bAllocate = true) override;
+	PCG_API virtual FPCGPointColor::ValueRange GetColorValueRange(bool bAllocate = true) override;
+	PCG_API virtual FPCGPointSteepness::ValueRange GetSteepnessValueRange(bool bAllocate = true) override;
+	PCG_API virtual FPCGPointSeed::ValueRange GetSeedValueRange(bool bAllocate = true) override;
+	PCG_API virtual FPCGPointMetadataEntry::ValueRange GetMetadataEntryValueRange(bool bAllocate = true) override;
+
+	PCG_API virtual FPCGPointTransform::ConstValueRange GetConstTransformValueRange() const override;
+	PCG_API virtual FPCGPointDensity::ConstValueRange GetConstDensityValueRange() const override;
+	PCG_API virtual FPCGPointBoundsMin::ConstValueRange GetConstBoundsMinValueRange() const override;
+	PCG_API virtual FPCGPointBoundsMax::ConstValueRange GetConstBoundsMaxValueRange() const override;
+	PCG_API virtual FPCGPointColor::ConstValueRange GetConstColorValueRange() const override;
+	PCG_API virtual FPCGPointSteepness::ConstValueRange GetConstSteepnessValueRange() const override;
+	PCG_API virtual FPCGPointSeed::ConstValueRange GetConstSeedValueRange() const override;
+	PCG_API virtual FPCGPointMetadataEntry::ConstValueRange GetConstMetadataEntryValueRange() const override;
 	//~End UPCGBasePointData Interface
 		
 protected:
-	UE_API virtual void InitializeSpatialDataInternal(const FPCGInitializeFromDataParams& InParams) override;
+	PCG_API virtual void InitializeSpatialDataInternal(const FPCGInitializeFromDataParams& InParams) override;
 
 private:
-	UE_API void FlattenPropertiesIfNeeded(EPCGPointNativeProperties Property = EPCGPointNativeProperties::All);
+	void FlattenPropertiesIfNeeded(EPCGPointNativeProperties Property = EPCGPointNativeProperties::All);
 
 	template<EPCGPointNativeProperties Property>
 	FPCGPointArrayProperty<typename TPCGPointNativeProperty<Property>::Type>* GetProperty(bool bWithInheritance = true)
@@ -219,5 +217,3 @@ private:
 	UPROPERTY(VisibleAnywhere, Category = ParentData, meta = (BitMask, BitmaskEnum = "/Script/PCG.EPCGPointNativeProperties"))
 	uint32 InheritedProperties = 0;
 };
-
-#undef UE_API
