[PCG] Add edit anywhere to the FEnumSelector to make sure it is gathered as override.

#rnx
#Synced-CL 48169160

diff --git a/Engine/Plugins/PCG/Source/PCG/Public/Elements/ControlFlow/PCGControlFlow.h b/Engine/Plugins/PCG/Source/PCG/Public/Elements/ControlFlow/PCGControlFlow.h
--- a/Engine/Plugins/PCG/Source/PCG/Public/Elements/ControlFlow/PCGControlFlow.h
+++ b/Engine/Plugins/PCG/Source/PCG/Public/Elements/ControlFlow/PCGControlFlow.h
@@ -27,7 +27,7 @@ struct FEnumSelector
 	UPROPERTY(DisplayName="Enum Class", meta=(PCG_NotOverridable))
 	TObjectPtr<UEnum> Class = nullptr;
 
-	UPROPERTY(DisplayName="Enum Value")
+	UPROPERTY(EditAnywhere, DisplayName="Enum Value", Category="Settings")
 	int64 Value = 0;
 
 	UE_API FText GetDisplayName() const;
