[PCG] Fix undo bug : remove MarkSubObjectsAsGarbage() as this was causing issues and reason for this change was to avoid some references being reported when creating a level instance of a PCG Actor that had actor references in its output. (This was properly fixed in CL 48271548)

#Synced-CL 48233187

diff --git a/Engine/Plugins/PCG/Source/PCG/Private/PCGComponent.cpp b/Engine/Plugins/PCG/Source/PCG/Private/PCGComponent.cpp
--- a/Engine/Plugins/PCG/Source/PCG/Private/PCGComponent.cpp
+++ b/Engine/Plugins/PCG/Source/PCG/Private/PCGComponent.cpp
@@ -1720,36 +1720,6 @@ void UPCGComponent::OnUnregister()
 	Super::OnUnregister();
 }
 
-void UPCGComponent::MarkSubObjectsAsGarbage()
-{
-	// First move out the generated output
-	ClearGraphGeneratedOutput(/*bClearLoadedPreviewData=*/true);
-	ClearPerPinGeneratedOutput();
-
-	auto ModifyAndMarkForGarbage = [](UObject* SubObject)
-	{
-		// @todo_pcg: revisit this but for now this fixes a crash when adding elements to GridsDependentOnWorldStreaming on a Blueprint PCG Component
-		// ModifyAndMarkForGarbage main use is for UPCGData that would have references on the PCG Component's owning actor.
-		if (!IsValid(SubObject) || SubObject->IsA<UPCGSchedulingPolicyBase>())
-		{
-			return;
-		}
-
-#if WITH_EDITOR
-		// Modify object for Undo/Redo support
-		if (GEditor)
-		{
-			// We don't have the info here to know if we should dirty or not but we can rely on the fact that this component will have been modified prior to its OnComponentDestroyed call
-			SubObject->Modify(/*bAlwaysMarkDirty=*/false);
-		}
-#endif
-		SubObject->MarkAsGarbage();
-	};
-
-	// Make sure all outered objects of this component are marked as garbage so that references from those objects to our owner actor do not get reported
-	ForEachObjectWithOuter(this, ModifyAndMarkForGarbage, true);
-}
-
 void UPCGComponent::OnComponentDestroyed(bool bDestroyingHierarchy)
 {
 	TRACE_CPUPROFILER_EVENT_SCOPE(UPCGComponent::OnComponentDestroyed);
@@ -1772,14 +1742,6 @@ void UPCGComponent::OnComponentDestroyed(bool bDestroyingHierarchy)
 	}
 
 #if WITH_EDITOR
-	// Only do this in Editor as it hasn't proven to cause issues at runtime and the call can be costly:
-	// If Component is getting destroyed and it is not a BP constructed component or we are not currently reconstructing the BP
-	// We need to mark subobjects (GeneratedOutputData, ManagedResources) for Garbage so that they aren't found by reference checking code paths
-	if (!IsCreatedByConstructionScript() || !GIsReconstructingBlueprintInstances)
-	{
-		MarkSubObjectsAsGarbage();
-	}
-
 	// Don't do the unregister in OnUnregister, because we have flows where the component gets Unregistered/Registered without getting destroyed.
 	if (UPCGSubsystem* Subsystem = GetSubsystem())
 	{
@@ -4157,9 +4119,6 @@ void FPCGComponentInstanceData::ApplyToComponent(UActorComponent* Component, con
 			{
 				ConstructionSourceComponent->GraphInstance->TeardownCallbacks();
 			}
-			// Make sure previous component destroys all its sub-objects to avoid leaving references to the new component
-			// This can happen in Copy/Paste code paths (ex: Level Instance creation)
-			const_cast<UPCGComponent*>(ConstructionSourceComponent)->MarkSubObjectsAsGarbage();
 		}
 
 		// Make sure to apply latest editing mode
diff --git a/Engine/Plugins/PCG/Source/PCG/Public/PCGComponent.h b/Engine/Plugins/PCG/Source/PCG/Public/PCGComponent.h
--- a/Engine/Plugins/PCG/Source/PCG/Public/PCGComponent.h
+++ b/Engine/Plugins/PCG/Source/PCG/Public/PCGComponent.h
@@ -603,8 +603,6 @@ protected:
 	/** Purges Actors and Components generated by PCG but are no longer managed by any PCG Component */
 	static UE_API void PurgeUnlinkedResources(const AActor* InActor);
 
-	UE_API void MarkSubObjectsAsGarbage();
-
 protected:
 	UPROPERTY(BlueprintReadOnly, VisibleAnywhere, Category = Settings, Instanced, meta = (NoResetToDefault, DisplayPriority = 100))
 	TObjectPtr<UPCGGraphInstance> GraphInstance;
