[PCG] Demote error to verbose in landscape data interface which was being thrown when generating areas that are not loaded in editor.

[RN] [PCG] Demote an error that occurred when generating landscape scattering in unloaded areas of the world in editor to verbose.
#Synced-CL 48138972

diff --git a/Engine/Plugins/PCG/Source/PCG/Private/Compute/DataInterfaces/PCGLandscapeDataInterface.cpp b/Engine/Plugins/PCG/Source/PCG/Private/Compute/DataInterfaces/PCGLandscapeDataInterface.cpp
--- a/Engine/Plugins/PCG/Source/PCG/Private/Compute/DataInterfaces/PCGLandscapeDataInterface.cpp
+++ b/Engine/Plugins/PCG/Source/PCG/Private/Compute/DataInterfaces/PCGLandscapeDataInterface.cpp
@@ -85,6 +85,7 @@ FPCGLandscapeResource::FPCGLandscapeResource(const FResourceKey& InKey)
 	const ULandscapeInfo* LandscapeInfo = ResourceKey.Source->GetLandscapeInfo();
 	if (!LandscapeInfo)
 	{
+		UE_LOG(LogPCG, Warning, TEXT("FPCGLandscapeResource: Could not retrieve landscape info to setup landscape height data."));
 		return;
 	}
 
@@ -366,6 +367,8 @@ void UPCGLandscapeDataProvider::InitFromLandscape(UPCGDataBinding* InBinding, AL
 	// If the height can't be accessed via RVT, build a physics based height texture instead.
 	if (!HeightVirtualTexture)
 	{
+		UE_CLOG(bAllowSampleVirtualTexture, LogPCG, Verbose, TEXT("UPCGLandscapeDataProvider: No runtime virtual texture found to use for landscape heights, falling back to uploading data from landscape collision components."));
+
 		InitializeFromLandscapeCollision(InBinding, InLandscape, Bounds);
 	}
 }
@@ -704,7 +707,8 @@ void FPCGLandscapeDataProviderProxy::SetShaderParameters(FParameters& OutShaderP
 
 	if (!bSetHeightVT && !bSetCollisionHeightTexture)
 	{
-		UE_LOG(LogPCG, Error, TEXT("PCGLandscapeDataInterface: Failed to initialize a height texture. If you are trying to use a virtual texture, make sure the RVT volumes are created."))
+		// Can be hit frequently in editor in unloaded areas for spatially loaded landscapes.
+		UE_LOG(LogPCG, Verbose, TEXT("PCGLandscapeDataInterface: No height texture initialized. If you are trying to use a virtual texture, make sure the RVT volumes are created. This may also occur when generating in areas where the landscape is not loaded."));
 	}
 
 	OutShaderParameters.LandscapeGridSize = (FVector2f)LandscapeGridSize;
