[PCG] Array overrides bugfixes

* Refactor `PCGPropertyHelpers::ExtractPropertyAsAttributeSet` to re-use as much as possible accessor code
* This will also fix the issue where empty arrays where not creating attributes
* Also add information when an attribute has no entries when overriding
* Make sure when we have no entries but we try to write into an array, we do not write an array of 1 default value, but an empty array

#rnx
#Synced-CL 48169160

diff --git a/Engine/Plugins/PCG/Source/PCG/Private/Helpers/PCGPropertyHelpers.cpp b/Engine/Plugins/PCG/Source/PCG/Private/Helpers/PCGPropertyHelpers.cpp
--- a/Engine/Plugins/PCG/Source/PCG/Private/Helpers/PCGPropertyHelpers.cpp
+++ b/Engine/Plugins/PCG/Source/PCG/Private/Helpers/PCGPropertyHelpers.cpp
@@ -6,6 +6,7 @@
 #include "PCGElement.h"
 #include "PCGModule.h"
 #include "PCGParamData.h"
+#include "Elements/Metadata/PCGMetadataElementCommon.h"
 #include "Helpers/PCGHelpers.h"
 #include "Helpers/PCGSettingsHelpers.h"
 #include "Metadata/PCGMetadata.h"
@@ -274,18 +275,13 @@ UPCGParamData* PCGPropertyHelpers::ExtractPropertyAsAttributeSet(const PCGProper
 {
 	check(Parameters.Container && Parameters.Class);
 
-	// TODO: Remove this after we've removed the deprecated PropertySelector.
-	TArrayView<const FPCGAttributePropertySelector> PropertySelectors;
+	TConstArrayView<FPCGAttributePropertySelector> PropertySelectors = MakeArrayView(Parameters.PropertySelectors);
+	static const FPCGAttributePropertySelector DefaultSelector{};
 
-	if (Parameters.PropertySelectors.Num() > 0)
+	// To keep previous behavior where no selectors meant a single default selector. 
+	if (PropertySelectors.IsEmpty())
 	{
-		PropertySelectors = TArrayView<const FPCGAttributePropertySelector>(Parameters.PropertySelectors);
-	}
-	else
-	{
-		PRAGMA_DISABLE_DEPRECATION_WARNINGS
-		PropertySelectors = TArrayView<const FPCGAttributePropertySelector>(&Parameters.PropertySelector, 1);
-		PRAGMA_ENABLE_DEPRECATION_WARNINGS
+		PropertySelectors = MakeArrayView(&DefaultSelector, 1);
 	}
 
 	UPCGParamData* ParamData = nullptr;
@@ -324,12 +320,9 @@ UPCGParamData* PCGPropertyHelpers::ExtractPropertyAsAttributeSet(const PCGProper
 			Property = SetProperty->ElementProp;
 		}
 
-		using ExtractablePropertyTuple = TTuple<FString, const FProperty*>;
+		using ExtractablePropertyTuple = TTuple<FString /*PropertyName*/, const FProperty* /*PropertyToExtract*/, const FProperty* /*ContainerProperty*/>;
 		TArray<ExtractablePropertyTuple> ExtractableProperties;
 
-		using GetAddressFunc = TFunction<const void* (const void*)>;
-		GetAddressFunc AddressFunc;
-
 		// Force extraction if the property is not supported by accessors.
 		const bool bShouldExtract = Parameters.bShouldExtract || !PCGAttributeAccessorHelpers::IsPropertyAccessorSupported(Property);
 
@@ -340,26 +333,24 @@ UPCGParamData* PCGPropertyHelpers::ExtractPropertyAsAttributeSet(const PCGProper
 		if (ExtractRoot || ((Property->IsA<FStructProperty>() || Property->IsA<FObjectProperty>()) && bShouldExtract))
 		{
 			const UStruct* UnderlyingClass = nullptr;
+			const FProperty* ContainerProperty = nullptr;
 
 			if (ExtractRoot)
 			{
 				UnderlyingClass = Parameters.Class;
-				// Identity
-				AddressFunc = [](const void* InAddress) { return InAddress; };
 			}
 			else if (const FStructProperty* StructProperty = CastField<FStructProperty>(Property))
 			{
 				UnderlyingClass = StructProperty->Struct;
-				AddressFunc = [StructProperty](const void* InAddress) { return InAddress ? StructProperty->ContainerPtrToValuePtr<void>(InAddress) : nullptr; };
+				ContainerProperty = Property;
 			}
 			else if (ObjectProperty)
 			{
 				UnderlyingClass = ObjectProperty->PropertyClass;
-				AddressFunc = [ObjectProperty](const void* InAddress) { return InAddress ? ObjectProperty->GetObjectPropertyValue_InContainer(InAddress) : nullptr; };
+				ContainerProperty = Property;
 			}
 
 			check(UnderlyingClass);
-			check(!!AddressFunc);
 
 			// Re-use code from overridable params
 			// Limit ourselves to not recurse into more structs.
@@ -378,7 +369,7 @@ UPCGParamData* PCGPropertyHelpers::ExtractPropertyAsAttributeSet(const PCGProper
 			{
 				if (ensure(!Param.PropertiesNames.IsEmpty()) && ensure(!Param.Properties.IsEmpty()) && ensure(Param.Properties[0]))
 				{
-					ExtractableProperties.Emplace(Param.PropertiesNames[0].ToString(), Param.Properties[0]);
+					ExtractableProperties.Emplace(Param.PropertiesNames[0].ToString(), Param.Properties[0], ContainerProperty);
 				}
 			}
 		}
@@ -405,10 +396,7 @@ UPCGParamData* PCGPropertyHelpers::ExtractPropertyAsAttributeSet(const PCGProper
 				AttributeName = Parameters.OutputAttributeName.ToString();
 			}
 
-			ExtractableProperties.Emplace(std::move(AttributeName), Property);
-
-			// Identity
-			AddressFunc = [](const void* InAddress) { return InAddress; };
+			ExtractableProperties.Emplace(std::move(AttributeName), Property, nullptr);
 		}
 
 		if (ExtractableProperties.IsEmpty())
@@ -465,70 +453,109 @@ UPCGParamData* PCGPropertyHelpers::ExtractPropertyAsAttributeSet(const PCGProper
 			}
 		}
 
-		bool bHasWarnNullContainerPtr = false;
-
-		for(int32 ElementAddressIndex = 0; ElementAddressIndex < ElementAddresses.Num(); ++ElementAddressIndex)
+		for (ExtractablePropertyTuple& ExtractableProperty : ExtractableProperties)
 		{
-			const void* ElementAddress = ElementAddresses[ElementAddressIndex];
-			PCGMetadataEntryKey EntryKey = Entries[ElementAddressIndex];
+			FString AttributeNameStr = ExtractableProperty.Get<0>();
+			FName AttributeName = NAME_None;
+			const FProperty* FinalProperty = ExtractableProperty.Get<1>();
+			const FProperty* ContainerProperty = ExtractableProperty.Get<2>();
 
-			// Offset the address if needed
-			const void* ContainerPtr = AddressFunc(ElementAddress);
+			// Make sure the Attribute name is sanitized, to prevent cases where property names have unsupported characters.
+			if (Parameters.bStrictSanitizeOutputAttributeNames)
+			{
+				AttributeName = MakeObjectNameFromDisplayLabel(AttributeNameStr, NAME_None);
+			}
+			else
+			{
+				FPCGMetadataAttributeBase::SanitizeName(AttributeNameStr);
+				AttributeName = *AttributeNameStr;
+			}
 
-			if (!ContainerPtr)
+			TUniquePtr<const IPCGAttributeAccessor> PropertyAccessor;
+			if (ContainerProperty)
 			{
-				if (!bHasWarnNullContainerPtr)
+				PropertyAccessor = PCGAttributeAccessorHelpers::CreatePropertyChainAccessor({ContainerProperty, FinalProperty});
+			}
+			else
+			{
+				PropertyAccessor = PCGAttributeAccessorHelpers::CreatePropertyAccessor(FinalProperty);
+			}
+
+			if (!PropertyAccessor)
+			{
+				if (!bQuiet)
 				{
-					bHasWarnNullContainerPtr = true;
-					PCGLog::LogErrorOnGraph(LOCTEXT("NullPtrContainer", "Some resolved objects were not assigned (null pointers), some results are discarded"), OptionalContext);
+					PCGLog::LogErrorOnGraph(FText::Format(LOCTEXT("ErrorCreatingPropertyAccessor", "Error while creating an attribute for property '{0}'. The property type is not supported by PCG."), FText::FromString(FinalProperty->GetName())), OptionalContext);
 				}
-				
-				continue;
+
+				bValidOperation = false;
+				break;
 			}
 
-			for (ExtractablePropertyTuple& ExtractableProperty : ExtractableProperties)
+			const FPCGAttributeAccessorKeysGenericPtrs PropertyAccessorKeys{ElementAddresses};
+
+			const FPCGAttributePropertySelector AttributeSelector = FPCGAttributePropertySelector::CreateAttributeSelector(AttributeName);
+
+			PCGAttributeAccessorHelpers::FPCGCreateAccessorWithAttributeCreationParams AttributeCreationParams =
 			{
-				FString AttributeNameStr = ExtractableProperty.Get<0>();
-				FName AttributeName = NAME_None;
-				const FProperty* FinalProperty = ExtractableProperty.Get<1>();
+				.InData = ParamData,
+				.InSelector = &AttributeSelector,
+				.InMatchingAccessor = PropertyAccessor.Get(),
+				.InMatchingKeysForDefaultValue = PropertyAccessorKeys.GetNum() > 0 ? &PropertyAccessorKeys : nullptr
+			};
+			
+			const TUniquePtr<IPCGAttributeAccessor> AttributeAccessor = PCGAttributeAccessorHelpers::CreateAccessorWithAttributeCreation(AttributeCreationParams);
+			FPCGAttributeAccessorKeysEntries AttributeKeys{TArrayView<PCGMetadataEntryKey>(Entries)};
 
-				// Make sure the Attribute name is sanitized, to prevent cases where property names have unsupported characters.
-				if (Parameters.bStrictSanitizeOutputAttributeNames)
-				{
-					AttributeName = MakeObjectNameFromDisplayLabel(AttributeNameStr, NAME_None);
-				}
-				else
+			if (!AttributeAccessor)
+			{
+				if (!bQuiet)
 				{
-					FPCGMetadataAttributeBase::SanitizeName(AttributeNameStr);
-					AttributeName = *AttributeNameStr;
+					PCGLog::LogErrorOnGraph(FText::Format(LOCTEXT("ErrorCreatingAttribute", "Error while creating an attribute for property '{0}'. Attribute creation failed."), FText::FromString(FinalProperty->GetName())), OptionalContext);
 				}
 
-				if (!Metadata->SetAttributeFromDataProperty(AttributeName, EntryKey, ContainerPtr, FinalProperty, /*bCreate=*/ true))
-				{
-					if (!bQuiet)
-					{
-						PCGLog::LogErrorOnGraph(FText::Format(LOCTEXT("ErrorCreatingAttribute", "Error while creating an attribute for property '{0}'. Either the property type is not supported by PCG or attribute creation failed."), FText::FromString(FinalProperty->GetName())), OptionalContext);
-					}
-
-					bValidOperation = false;
-					break;
-				}
+				bValidOperation = false;
+				break;
 			}
 
-			if (!bValidOperation)
+			PCGMetadataElementCommon::FCopyFromAccessorToAccessorParams Params =
+			{
+				.IterationCount = PCGMetadataElementCommon::FCopyFromAccessorToAccessorParams::Out,
+				.InAccessor = PropertyAccessor.Get(),
+				.InKeys = &PropertyAccessorKeys,
+				.OutAccessor = AttributeAccessor.Get(),
+				.OutKeys = &AttributeKeys,
+			};
+			
+			if (!PCGMetadataElementCommon::CopyFromAccessorToAccessor(Params))
 			{
+				bValidOperation = false;
 				break;
 			}
 
 			if (bShouldExtract && OptionalObjectTraversed && ObjectProperty)
 			{
-				const UObject* Object = ObjectProperty->GetPropertyValue_InContainer(ElementAddress);
-				if (IsValid(Object))
+				TArray<FSoftObjectPath> ObjectPaths;
+				ObjectPaths.SetNum(ElementAddresses.Num());
+				if (PropertyAccessor->GetRange<FSoftObjectPath>(ObjectPaths, 0, PropertyAccessorKeys, EPCGAttributeAccessorFlags::AllowBroadcastAndConstructible))
 				{
-					OptionalObjectTraversed->Add(Object);
+					for (const FSoftObjectPath& ObjectPath : ObjectPaths)
+					{
+						const UObject* Object = ObjectPath.ResolveObject();
+						if (IsValid(Object))
+						{
+							OptionalObjectTraversed->Add(Object);
+						}
+					}
 				}
+
 			}
 		}
+
+		if (!bValidOperation)
+		{
+			break;
+		}
 	}
 
 	return bValidOperation ? ParamData : nullptr;
diff --git a/Engine/Plugins/PCG/Source/PCG/Private/Metadata/Accessors/PCGAttributeAccessorHelpers.cpp b/Engine/Plugins/PCG/Source/PCG/Private/Metadata/Accessors/PCGAttributeAccessorHelpers.cpp
--- a/Engine/Plugins/PCG/Source/PCG/Private/Metadata/Accessors/PCGAttributeAccessorHelpers.cpp
+++ b/Engine/Plugins/PCG/Source/PCG/Private/Metadata/Accessors/PCGAttributeAccessorHelpers.cpp
@@ -307,6 +307,7 @@ TUniquePtr<const IPCGAttributeAccessor> PCGAttributeAccessorHelpers::CreateConst
 		if (Result && OutResult && ParamData && ParamData->Metadata)
 		{
 			OutResult->bHasMultipleDataInAttributeSet = ParamData->Metadata->GetLocalItemCount() > 1;
+			OutResult->bHasNoEntry = ParamData->Metadata->GetLocalItemCount() == 0;
 			
 			if (InParam.NumContainers > 0)
 			{
@@ -458,28 +459,53 @@ TUniquePtr<IPCGAttributeAccessor> PCGAttributeAccessorHelpers::CreateAccessor(FP
 
 TUniquePtr<IPCGAttributeAccessor> PCGAttributeAccessorHelpers::CreateAccessorWithAttributeCreation(UPCGData* InData, const FPCGAttributePropertySelector& InSelector, const IPCGAttributeAccessor* InMatchingAccessor, EPCGAttributeAccessorFlags InTypeMatching, bool bQuiet)
 {
-	TUniquePtr<IPCGAttributeAccessor> Result = CreateAccessor(InData, InSelector, bQuiet);
+	FPCGCreateAccessorWithAttributeCreationParams Params =
+	{
+		.InData = InData,
+		.InSelector = &InSelector,
+		.InMatchingAccessor = InMatchingAccessor,
+		.InTypeMatching = InTypeMatching,
+		.bQuiet = bQuiet
+	};
+
+	return CreateAccessorWithAttributeCreation(Params);
+}
 
-	if (!InSelector.IsBasicAttribute() || !InMatchingAccessor)
+TUniquePtr<IPCGAttributeAccessor> PCGAttributeAccessorHelpers::CreateAccessorWithAttributeCreation(const FPCGCreateAccessorWithAttributeCreationParams& InParams)
+{
+	if (!InParams.InData || !InParams.InSelector)
+	{
+		return nullptr;
+	}
+
+	TUniquePtr<IPCGAttributeAccessor> Result = CreateAccessor(InParams.InData, *InParams.InSelector, InParams.bQuiet);
+
+	if (!InParams.InSelector->IsBasicAttribute())
 	{
 		return Result;
 	}
 
 	bool bValid = !!Result;
+	const int16 MatchingType = InParams.InMatchingAccessor ? InParams.InMatchingAccessor->GetUnderlyingType() : InParams.InExpectedType;
+	// The underlying type of the accessor is expected to be valid (otherwise it is ill-formed), so this is mainly to catch if InExpectedType is valid.
+	if (!PCG::Private::IsPCGType(MatchingType))
+	{
+		return Result;
+	}
 
-	if (bValid && !!(InTypeMatching & EPCGAttributeAccessorFlags::StrictType))
+	if (bValid && !!(InParams.InTypeMatching & EPCGAttributeAccessorFlags::StrictType))
 	{
-		bValid &= (InMatchingAccessor->GetUnderlyingType() == Result->GetUnderlyingType());
+		bValid &= (MatchingType == Result->GetUnderlyingType());
 	}
 
-	if (bValid && !!(InTypeMatching & EPCGAttributeAccessorFlags::AllowBroadcast))
+	if (bValid && !!(InParams.InTypeMatching & EPCGAttributeAccessorFlags::AllowBroadcast))
 	{
-		bValid &= PCG::Private::IsBroadcastable(InMatchingAccessor->GetUnderlyingType(), Result->GetUnderlyingType());
+		bValid &= PCG::Private::IsBroadcastable(MatchingType, Result->GetUnderlyingType());
 	}
 
-	if (bValid && !!(InTypeMatching & EPCGAttributeAccessorFlags::AllowConstructible))
+	if (bValid && !!(InParams.InTypeMatching & EPCGAttributeAccessorFlags::AllowConstructible))
 	{
-		bValid &= PCG::Private::IsConstructible(InMatchingAccessor->GetUnderlyingType(), Result->GetUnderlyingType());
+		bValid &= PCG::Private::IsConstructible(MatchingType, Result->GetUnderlyingType());
 	}
 
 	if (!bValid)
@@ -487,29 +513,42 @@ TUniquePtr<IPCGAttributeAccessor> PCGAttributeAccessorHelpers::CreateAccessorWit
 		Result.Reset();
 
 		// We didn't find the attribute in the data, or we can't Broadcast/Construct, so create a new one
-		UPCGMetadata* Metadata = InData->MutableMetadata();
-		if (!Metadata)
+		UPCGMetadata* Metadata = InParams.InData->MutableMetadata();
+		FPCGMetadataDomain* MetadataDomain = Metadata ? Metadata->GetDefaultMetadataDomain() : nullptr;
+		if (!MetadataDomain)
 		{
 			return Result;
 		}
 
-		const FName AttributeName = InSelector.GetName();
-		if (Metadata->HasAttribute(AttributeName))
+		const FName AttributeName = InParams.InSelector->GetName();
+		if (MetadataDomain->HasAttribute(AttributeName))
 		{
-			Metadata->DeleteAttribute(AttributeName);
+			MetadataDomain->DeleteAttribute(AttributeName);
 		}
 
-		auto CreateAttributeAndAccessor = [InMatchingAccessor, AttributeName, Metadata](auto&& Dummy) -> TUniquePtr<IPCGAttributeAccessor>
+		auto CreateAttributeAndAccessor = [&InParams, AttributeName, MetadataDomain](auto&& Dummy) -> TUniquePtr<IPCGAttributeAccessor>
 		{
 			using AttributeType = std::decay_t<decltype(Dummy)>;
 			AttributeType DefaultValue = PCG::Private::MetadataTraits<AttributeType>::ZeroValue();
-			InMatchingAccessor->Get<AttributeType>(DefaultValue, FPCGAttributeAccessorKeysEntries(PCGInvalidEntryKey));
-			FPCGMetadataAttribute<AttributeType>* Attribute = Metadata->CreateAttribute<AttributeType>(AttributeName, DefaultValue, /*bAllowInterpolation=*/true, /*bOverrideParent=*/false);
 
-			return Attribute ? MakeUnique<FPCGAttributeAccessor<AttributeType>>(Attribute, Attribute->GetMetadataDomain()) : nullptr;
+			if (InParams.InMatchingAccessor)
+			{
+				if (InParams.InMatchingKeysForDefaultValue)
+				{
+					InParams.InMatchingAccessor->Get<AttributeType>(DefaultValue, *InParams.InMatchingKeysForDefaultValue);
+				}
+				else
+				{
+					InParams.InMatchingAccessor->Get<AttributeType>(DefaultValue, FPCGAttributeAccessorKeysEntries(PCGInvalidEntryKey));
+				}
+			}
+
+			FPCGMetadataAttribute<AttributeType>* Attribute = MetadataDomain->CreateAttribute<AttributeType>(AttributeName, DefaultValue, /*bAllowInterpolation=*/true, /*bOverrideParent=*/false);
+
+			return Attribute ? MakeUnique<FPCGAttributeAccessor<AttributeType>>(Attribute, MetadataDomain) : nullptr;
 		};
 
-		Result = PCGMetadataAttribute::CallbackWithRightType(InMatchingAccessor->GetUnderlyingType(), std::move(CreateAttributeAndAccessor));
+		Result = PCGMetadataAttribute::CallbackWithRightType(MatchingType, std::move(CreateAttributeAndAccessor));
 	}
 
 	return Result;
diff --git a/Engine/Plugins/PCG/Source/PCG/Private/PCGContext.cpp b/Engine/Plugins/PCG/Source/PCG/Private/PCGContext.cpp
--- a/Engine/Plugins/PCG/Source/PCG/Private/PCGContext.cpp
+++ b/Engine/Plugins/PCG/Source/PCG/Private/PCGContext.cpp
@@ -473,7 +473,16 @@ PRAGMA_ENABLE_DEPRECATION_WARNINGS
 			bool bSuccess = false;
 			if (Param.NumContainers == 1 && ensure(AccessorResult.Keys))
 			{
-				bSuccess = PCGAttributeAccessorHelpers::ExtractAllValues<PropertyType>(AttributeAccessor.Get(), AccessorResult.Keys.Get(), Values);
+				// Having an attribute with 0 entry means that we should empty the array, but the keys for attributes always have at least one key
+				// in case of default value. So keep Values empty if we know we explicitly don't have entries.
+				if (AccessorResult.bHasNoEntry)
+				{
+					bSuccess = true;
+				}
+				else
+				{
+					bSuccess = PCGAttributeAccessorHelpers::ExtractAllValues<PropertyType>(AttributeAccessor.Get(), AccessorResult.Keys.Get(), Values);
+				}
 			}
 			else
 			{
diff --git a/Engine/Plugins/PCG/Source/PCG/Public/Helpers/PCGPropertyHelpers.h b/Engine/Plugins/PCG/Source/PCG/Public/Helpers/PCGPropertyHelpers.h
--- a/Engine/Plugins/PCG/Source/PCG/Public/Helpers/PCGPropertyHelpers.h
+++ b/Engine/Plugins/PCG/Source/PCG/Public/Helpers/PCGPropertyHelpers.h
@@ -76,10 +76,7 @@ namespace PCGPropertyHelpers
 		// Class or ScriptStruct of the object/struct we want to extract the property from.
 		const UStruct* Class = nullptr;
 
-		// Selector of the property we want to extract
-		UE_DEPRECATED(5.5, "Use the selector array (PropertySelectors) instead.")
-		FPCGAttributePropertySelector PropertySelector;
-
+		// Selector of the properties we want to extract
 		TArray<FPCGAttributePropertySelector> PropertySelectors;
 
 		// Optional name of the attribute that will receive the extracted property. If @Source or @SourceName, will take the property name. 
diff --git a/Engine/Plugins/PCG/Source/PCG/Public/Metadata/Accessors/PCGAttributeAccessorHelpers.h b/Engine/Plugins/PCG/Source/PCG/Public/Metadata/Accessors/PCGAttributeAccessorHelpers.h
--- a/Engine/Plugins/PCG/Source/PCG/Public/Metadata/Accessors/PCGAttributeAccessorHelpers.h
+++ b/Engine/Plugins/PCG/Source/PCG/Public/Metadata/Accessors/PCGAttributeAccessorHelpers.h
@@ -73,6 +73,7 @@ namespace PCGAttributeAccessorHelpers
 		bool bPinConnected = false;
 		bool bHasMultipleAttributeSetsOnOverridePin = false;
 		bool bHasMultipleDataInAttributeSet = false;
+		bool bHasNoEntry = false;
 
 		// If the param is in a container, we'll need the keys
 		TUniquePtr<const IPCGAttributeAccessorKeys> Keys;
@@ -107,11 +108,29 @@ namespace PCGAttributeAccessorHelpers
 
 	/**
 	* Creates an accessor to the property or attribute pointed at by the InSelector. If the selector points to a base attribute and the attribute doesn't exist 
-	* or the types doesn't match with the Matching Accessor, create the attribute using the Matching Accessor.
+	* or the types doesn't match with the Matching Accessor or provided type, create the attribute using the Matching Accessor. The default value of this new attribute will either
+	* be the first value get from the provided keys (if provided), the default value of the matching accessor (if it is an attribute) or the zero value (if no matching accessor or not supporting default values).
 	* Note that InData must not be null if the selector points to an attribute,
 	* but in the case of properties, it either has to be the appropriate type or null.
 	* Make sure to update your selector before-hand if you want to support "@Source". Otherwise the creation will fail.
 	*/
+	struct FPCGCreateAccessorWithAttributeCreationParams
+	{
+		UPCGData* InData = nullptr;
+		const FPCGAttributePropertySelector* InSelector = nullptr;
+
+		const IPCGAttributeAccessor* InMatchingAccessor = nullptr;
+		const IPCGAttributeAccessorKeys* InMatchingKeysForDefaultValue = nullptr;
+		// Discarded if InMatchingAccessor is not null.
+		int16 InExpectedType = static_cast<uint16>(EPCGMetadataTypes::Unknown);
+
+		EPCGAttributeAccessorFlags InTypeMatching = EPCGAttributeAccessorFlags::AllowBroadcastAndConstructible;
+		bool bQuiet = false;
+	};
+
+	PCG_API TUniquePtr<IPCGAttributeAccessor> CreateAccessorWithAttributeCreation(const FPCGCreateAccessorWithAttributeCreationParams& InParams);
+
+	UE_DEPRECATED(5.8, "Use the version with the params")
 	PCG_API TUniquePtr<IPCGAttributeAccessor> CreateAccessorWithAttributeCreation(UPCGData* InData, const FPCGAttributePropertySelector& InSelector, const IPCGAttributeAccessor* InMatchingAccessor, EPCGAttributeAccessorFlags InTypeMatching = EPCGAttributeAccessorFlags::AllowBroadcastAndConstructible, bool bQuiet = false);
 
 	/**
