[PCG] Allow kernel pins to not upload data to GPU to save time and memory.

#rnx
#Synced-CL 48348345

diff --git a/Engine/Plugins/PCG/Source/PCG/Private/Compute/PCGComputeGraph.cpp b/Engine/Plugins/PCG/Source/PCG/Private/Compute/PCGComputeGraph.cpp
--- a/Engine/Plugins/PCG/Source/PCG/Private/Compute/PCGComputeGraph.cpp
+++ b/Engine/Plugins/PCG/Source/PCG/Private/Compute/PCGComputeGraph.cpp
@@ -221,7 +221,7 @@ void UPCGComputeGraph::DebugLogDataDescriptions(const UPCGDataBinding* InBinding
 
 			TSharedPtr<const FPCGProxyForGPUDataCollection> UpstreamDataCollectionInfo;
 
-			if (const FName* VirtualLabel = GetVirtualPinLabelFromBindingIndex(GetBindingIndex(InKernelPin)))
+			if (const FName* VirtualLabel = GetVirtualPinLabelFromKernelInputPin(InKernelPin))
 			{
 				TArray<FPCGTaggedData> PinInputs = InBinding->GetInputDataCollection().GetInputsByPin(*VirtualLabel);
 
diff --git a/Engine/Plugins/PCG/Source/PCG/Private/Compute/PCGComputeGraph.h b/Engine/Plugins/PCG/Source/PCG/Private/Compute/PCGComputeGraph.h
--- a/Engine/Plugins/PCG/Source/PCG/Private/Compute/PCGComputeGraph.h
+++ b/Engine/Plugins/PCG/Source/PCG/Private/Compute/PCGComputeGraph.h
@@ -51,8 +51,8 @@ public:
 	/** Get the first binding index, which uniquely identifies a kernel pin to the compute graph. */
 	int32 GetBindingIndex(const FPCGKernelPin& InKernelPin) const;
 
-	/** Get the virtual label associated with a binding index/kernel pin. This only exists if the kernel pin is an input pin on the border of the compute graph. */
-	const FName* GetVirtualPinLabelFromBindingIndex(int32 InBindingIndex) const { return CPUDataBindingToVirtualPinLabel.Find(InBindingIndex); }
+	/** Get the virtual label associated with a kernel pin. This only exists if the kernel pin is an input pin on the border of the compute graph. */
+	const FName* GetVirtualPinLabelFromKernelInputPin(const FPCGKernelPin& InKernelPin) const { return KernelPinToVirtualPinLabel.Find(InKernelPin); }
 
 	/** Get the PCG data type associated with the virtual input pin label. */
 	const FPCGDataTypeIdentifier* GetVirtualInputPinType(FName InVirtualInputPinLabel) const { return VirtualInputPinToType.Find(InVirtualInputPinLabel); }
@@ -132,6 +132,10 @@ protected:
 	UPROPERTY()
 	TMap</*BindingIndex*/int32, /*VirtualPinLabel*/FName> CPUDataBindingToVirtualPinLabel;
 
+	/** Kernel input pin to virtual pin label, populated for pins that receive data from external output pins. */
+	UPROPERTY()
+	TMap<FPCGKernelPin, /*VirtualPinLabel*/FName> KernelPinToVirtualPinLabel;
+
 	/** Maps virtual input pins to their PCG data type. */
 	UPROPERTY()
 	TMap</*VirtualPinLabel*/FName, FPCGDataTypeIdentifier> VirtualInputPinToType;
diff --git a/Engine/Plugins/PCG/Source/PCG/Private/Compute/PCGDataBinding.cpp b/Engine/Plugins/PCG/Source/PCG/Private/Compute/PCGDataBinding.cpp
--- a/Engine/Plugins/PCG/Source/PCG/Private/Compute/PCGDataBinding.cpp
+++ b/Engine/Plugins/PCG/Source/PCG/Private/Compute/PCGDataBinding.cpp
@@ -280,17 +280,12 @@ int32 UPCGDataBinding::GetFirstInputDataIndex(const UPCGComputeKernel* InKernel,
 {
 	int32 FirstDataIndex = INDEX_NONE;
 
-	const int32 BindingIndex = Graph ? Graph->GetBindingIndex(FPCGKernelPin(InKernel->GetKernelIndex(), InPinLabel, /*bIsInput=*/true)) : INDEX_NONE;
-
-	if (ensure(BindingIndex != INDEX_NONE))
+	if (const FName* VirtualLabel = Graph->GetVirtualPinLabelFromKernelInputPin(FPCGKernelPin(InKernel->GetKernelIndex(), InPinLabel, /*bIsInput=*/true)))
 	{
-		if (const FName* VirtualLabel = Graph->GetVirtualPinLabelFromBindingIndex(BindingIndex))
+		FirstDataIndex = InputDataCollection.TaggedData.IndexOfByPredicate([VirtualLabel](const FPCGTaggedData& InData)
 		{
-			FirstDataIndex = InputDataCollection.TaggedData.IndexOfByPredicate([VirtualLabel](const FPCGTaggedData& InData)
-			{
-				return InData.Pin == *VirtualLabel;
-			});
-		}
+			return InData.Pin == *VirtualLabel;
+		});
 	}
 
 	return FirstDataIndex;
@@ -300,20 +295,15 @@ TArray<int32> UPCGDataBinding::GetPinInputDataIndices(const UPCGComputeKernel* I
 {
 	TArray<int32> PinDataIndices;
 
-	const int32 BindingIndex = Graph ? Graph->GetBindingIndex(FPCGKernelPin(InKernel->GetKernelIndex(), InPinLabel, /*bIsInput=*/true)) : INDEX_NONE;
-
-	if (ensure(BindingIndex != INDEX_NONE))
+	if (const FName* VirtualLabel = Graph->GetVirtualPinLabelFromKernelInputPin(FPCGKernelPin(InKernel->GetKernelIndex(), InPinLabel, /*bIsInput=*/true)))
 	{
-		if (const FName* VirtualLabel = Graph->GetVirtualPinLabelFromBindingIndex(BindingIndex))
-		{
-			PinDataIndices.Reserve(InputDataCollection.TaggedData.Num());
+		PinDataIndices.Reserve(InputDataCollection.TaggedData.Num());
 
-			for (int32 DataIndex = 0; DataIndex < InputDataCollection.TaggedData.Num(); ++DataIndex)
+		for (int32 DataIndex = 0; DataIndex < InputDataCollection.TaggedData.Num(); ++DataIndex)
+		{
+			if (InputDataCollection.TaggedData[DataIndex].Pin == *VirtualLabel)
 			{
-				if (InputDataCollection.TaggedData[DataIndex].Pin == *VirtualLabel)
-				{
-					PinDataIndices.Add(DataIndex);
-				}
+				PinDataIndices.Add(DataIndex);
 			}
 		}
 	}
diff --git a/Engine/Plugins/PCG/Source/PCG/Private/Elements/PCGStaticMeshSpawner.cpp b/Engine/Plugins/PCG/Source/PCG/Private/Elements/PCGStaticMeshSpawner.cpp
--- a/Engine/Plugins/PCG/Source/PCG/Private/Elements/PCGStaticMeshSpawner.cpp
+++ b/Engine/Plugins/PCG/Source/PCG/Private/Elements/PCGStaticMeshSpawner.cpp
@@ -161,7 +161,7 @@ TArray<FPCGPinProperties> UPCGStaticMeshSpawnerSettings::InputPinProperties() co
 		FPCGPinProperties& PrimitiveDataPin = Properties.Emplace_GetRef(PCGStaticMeshSpawner::PrimitiveTablePinLabel, EPCGDataType::Param, /*bAllowMultipleConnections=*/false);
 		PrimitiveDataPin.SetRequiredPin();
 #if WITH_EDITOR
-		PrimitiveDataPin.Tooltip = LOCTEXT("PrimitiveDataPinTooltip", "Table of attributes that drive primitive properties. The primitive index attribute on the input points index into rows of this table.");
+		PrimitiveDataPin.Tooltip = LOCTEXT("PrimitiveDataPinTooltip", "Table of attributes that drive primitive properties. The primitive index attribute on the input points index into rows of this table. This data is only used on the CPU; GPU data will be read back to CPU.");
 #endif
 	}
 
diff --git a/Engine/Plugins/PCG/Source/PCG/Private/Elements/PCGStaticMeshSpawnerKernel.cpp b/Engine/Plugins/PCG/Source/PCG/Private/Elements/PCGStaticMeshSpawnerKernel.cpp
--- a/Engine/Plugins/PCG/Source/PCG/Private/Elements/PCGStaticMeshSpawnerKernel.cpp
+++ b/Engine/Plugins/PCG/Source/PCG/Private/Elements/PCGStaticMeshSpawnerKernel.cpp
@@ -199,6 +199,12 @@ void UPCGStaticMeshSpawnerKernel::CreateAdditionalOutputDataInterfaces(FPCGGPUCo
 	InstanceDI->InputPinProvidingData = PCGPinConstants::DefaultInputLabel;
 	OutDataInterfaces.Add(InstanceDI);
 }
+
+bool UPCGStaticMeshSpawnerKernel::UploadInputPinDataToGPU(FName InPinLabel) const
+{
+	// Primitive table is not actually needed on GPU, is only used statically on CPU for scheduling and primitive setup.
+	return InPinLabel != PCGStaticMeshSpawnerKernelConstants::PrimitiveTablePinLabel;
+}
 #endif
 
 void UPCGStaticMeshSpawnerKernel::AddStaticCreatedStrings(TArray<FString>& InOutStringTable) const
@@ -228,7 +234,6 @@ void UPCGStaticMeshSpawnerKernel::GetInputPins(TArray<FPCGPinProperties>& OutPin
 	OutPins.Emplace(PCGStaticMeshSpawnerKernelConstants::InstanceCountsPinLabel, bUseRawInstanceCountsBuffer ? FPCGDataTypeInfoRawBuffer::AsId() : FPCGDataTypeInfoParam::AsId());
 
 	// Used when spawning from primitive table. This data is not actually used by the kernel/on the GPU.
-	// todo_pcg: In the future it might be nice to signal the pin doesn't require data upload to GPU.
 	OutPins.Emplace(PCGStaticMeshSpawnerKernelConstants::PrimitiveTablePinLabel, FPCGDataTypeInfoParam::AsId());
 }
 
diff --git a/Engine/Plugins/PCG/Source/PCG/Private/Graph/PCGGraphCompilerGPU.cpp b/Engine/Plugins/PCG/Source/PCG/Private/Graph/PCGGraphCompilerGPU.cpp
--- a/Engine/Plugins/PCG/Source/PCG/Private/Graph/PCGGraphCompilerGPU.cpp
+++ b/Engine/Plugins/PCG/Source/PCG/Private/Graph/PCGGraphCompilerGPU.cpp
@@ -1002,11 +1002,17 @@ void FPCGGraphCompilerGPU::CreateDataInterfaces(
 
 		// Create DIs for all output pins regardless of outbound connections, because the kernels currently need their outputs to be bound to valid resources.
 
-		if (UPCGComputeKernel* const* Kernel = InTaskIdToKernel.Find(TaskId))
+		const UPCGComputeKernel* TaskKernel = nullptr;
+		if (UPCGComputeKernel* const* FoundKernel = InTaskIdToKernel.Find(TaskId))
+		{
+			TaskKernel = *FoundKernel;
+		}
+
+		if (TaskKernel)
 		{
 			// Get output pins from kernel.
 			TArray<FPCGPinPropertiesGPU> OutputPins;
-			(*Kernel)->GetOutputPins(OutputPins);
+			TaskKernel->GetOutputPins(OutputPins);
 
 			for (const FPCGPinPropertiesGPU& Pin : OutputPins)
 			{
@@ -1098,6 +1104,12 @@ void FPCGGraphCompilerGPU::CreateDataInterfaces(
 				continue;
 			}
 
+			if (TaskKernel && !TaskKernel->UploadInputPinDataToGPU(Input.DownstreamPin->Label))
+			{
+				UE_LOG(LogPCG, Verbose, TEXT("Kernel '%s' input pin '%s' does not require data on GPU, skipped creating data interface."), *TaskKernel->GetName(), *Input.DownstreamPin->Label.ToString());
+				continue;
+			}
+
 			UPCGComputeDataInterface* OutputDI = nullptr;
 
 			if (UPCGComputeDataInterface** FoundOutputDI = InOutPinToDataInterface.Find({ Input.TaskId, Input.UpstreamPin->Label }))
@@ -1331,6 +1343,14 @@ void FPCGGraphCompilerGPU::CompileComputeGraph(
 				continue;
 			}
 
+			// Setup maps related to virtual pins.
+			if (const FName* VirtualPin = InOriginalToVirtualPin.Find(FNodePin(TaskId, Input.DownstreamPin->Label, /*IsInput*/true)))
+			{
+				InOutComputeGraph->KernelPinToVirtualPinLabel.Add(FPCGKernelPin(KernelIndex, Input.DownstreamPin->Label, /*bInIsInput=*/true), *VirtualPin);
+				InOutComputeGraph->VirtualInputPinToType.FindOrAdd(*VirtualPin) = Input.DownstreamPin->AllowedTypes;
+				OriginalToVirtualInputPin.Add(Input.DownstreamPin->Label, *VirtualPin);
+			}
+
 			UPCGComputeDataInterface* const* FoundDI = InOutputPinToDataInterface.Find({ Input.TaskId, Input.UpstreamPin->Label });
 			if (!FoundDI || !*FoundDI)
 			{
@@ -1347,12 +1367,6 @@ void FPCGGraphCompilerGPU::CompileComputeGraph(
 
 			InputDataInterfaces.Emplace(Index, /*Prefix=*/Input.DownstreamPin->Label, /*PinLabel=*/Input.DownstreamPin->Label);
 
-			if (const FName* VirtualPin = InOriginalToVirtualPin.Find(FNodePin(TaskId, Input.DownstreamPin->Label, /*IsInput*/true)))
-			{
-				OriginalToVirtualInputPin.Add(Input.DownstreamPin->Label, *VirtualPin);
-				InOutComputeGraph->VirtualInputPinToType.FindOrAdd(*VirtualPin) = Input.DownstreamPin->AllowedTypes;
-			}
-
 			CreateDataResolverDataInterface(InputDataInterfaces, Input.DownstreamPin->Label, /*bIsInput=*/true);
 		}
 
diff --git a/Engine/Plugins/PCG/Source/PCG/Public/Compute/PCGComputeKernel.h b/Engine/Plugins/PCG/Source/PCG/Public/Compute/PCGComputeKernel.h
--- a/Engine/Plugins/PCG/Source/PCG/Public/Compute/PCGComputeKernel.h
+++ b/Engine/Plugins/PCG/Source/PCG/Public/Compute/PCGComputeKernel.h
@@ -151,6 +151,9 @@ public:
 
 	/** Create a data interface for the given output pin. */
 	PCG_API virtual UPCGComputeDataInterface* CreateOutputPinDataInterface(const PCGComputeHelpers::FCreateDataInterfaceParams& InParams) const;
+
+	/** Whether the data on the input pin is used on the GPU (and thus requires a data interface) or not. */
+	virtual bool UploadInputPinDataToGPU(FName InPinLabel) const { return true; }
 #endif
 
 	/** Compute a description of data that will be output from pin InOutputPinLabel. */
diff --git a/Engine/Plugins/PCG/Source/PCG/Public/Elements/PCGStaticMeshSpawnerKernel.h b/Engine/Plugins/PCG/Source/PCG/Public/Elements/PCGStaticMeshSpawnerKernel.h
--- a/Engine/Plugins/PCG/Source/PCG/Public/Elements/PCGStaticMeshSpawnerKernel.h
+++ b/Engine/Plugins/PCG/Source/PCG/Public/Elements/PCGStaticMeshSpawnerKernel.h
@@ -26,8 +26,9 @@ public:
 	virtual const TCHAR* GetSourceFilePath() const override { return TEXT("/Plugin/PCG/Private/Elements/PCGStaticMeshSpawner.usf"); }
 	virtual FString GetCookedSource(FPCGGPUCompilationContext& InOutContext) const override;
 	virtual FString GetEntryPoint() const override { return TEXT("PCGStaticMeshSpawnerCS"); }
-	void CreateAdditionalInputDataInterfaces(FPCGGPUCompilationContext& InOutContext, UObject* InObjectOuter, TArray<TObjectPtr<UComputeDataInterface>>& OutDataInterfaces) const override;
-	void CreateAdditionalOutputDataInterfaces(FPCGGPUCompilationContext& InOutContext, UObject* InObjectOuter, TArray<TObjectPtr<UComputeDataInterface>>& OutDataInterfaces) const override;
+	virtual void CreateAdditionalInputDataInterfaces(FPCGGPUCompilationContext& InOutContext, UObject* InObjectOuter, TArray<TObjectPtr<UComputeDataInterface>>& OutDataInterfaces) const override;
+	virtual void CreateAdditionalOutputDataInterfaces(FPCGGPUCompilationContext& InOutContext, UObject* InObjectOuter, TArray<TObjectPtr<UComputeDataInterface>>& OutDataInterfaces) const override;
+	virtual bool UploadInputPinDataToGPU(FName InPinLabel) const override;
 #endif
 	virtual void AddStaticCreatedStrings(TArray<FString>& InOutStringTable) const override;
 	virtual void GetKernelAttributeKeys(TArray<FPCGKernelAttributeKey>& OutKeys) const override;
