[PCG] Add node to do simple tag-based data replacement from one set to the other.

#Synced-CL 48285106

diff --git a/Engine/Plugins/PCG/Source/PCG/Private/Elements/Data/PCGReplaceDataByTags.cpp b/Engine/Plugins/PCG/Source/PCG/Private/Elements/Data/PCGReplaceDataByTags.cpp
new file mode 100644
--- /dev/null
+++ b/Engine/Plugins/PCG/Source/PCG/Private/Elements/Data/PCGReplaceDataByTags.cpp
@@ -0,0 +1,120 @@
+// Copyright Epic Games, Inc. All Rights Reserved.
+
+#include "Elements/Data/PCGReplaceDataByTags.h"
+
+#include "PCGContext.h"
+
+#include UE_INLINE_GENERATED_CPP_BY_NAME(PCGReplaceDataByTags)
+
+namespace PCGReplaceDataByTags
+{
+	namespace Constants
+	{
+		const FLazyName ReplacementLabel = TEXT("Replacement");
+	}
+}
+
+TArray<FPCGPinProperties> UPCGReplaceDataByTagSettings::InputPinProperties() const
+{
+	TArray<FPCGPinProperties> PinProperties;
+	PinProperties.Emplace_GetRef(PCGPinConstants::DefaultInputLabel, EPCGDataType::Any).SetRequiredPin();
+	PinProperties.Emplace_GetRef(PCGReplaceDataByTags::Constants::ReplacementLabel).SetRequiredPin();
+
+	return PinProperties;
+}
+
+TArray<FPCGPinProperties> UPCGReplaceDataByTagSettings::OutputPinProperties() const
+{
+	TArray<FPCGPinProperties> PinProperties;
+	PinProperties.Emplace(PCGPinConstants::DefaultOutputLabel, EPCGDataType::Any);
+
+	return PinProperties;
+}
+
+FPCGDataTypeIdentifier UPCGReplaceDataByTagSettings::GetCurrentPinTypesID(const UPCGPin* InPin) const
+{
+	// In this specific instance we'll do the union of both input pins.
+	if (InPin->IsOutputPin())
+	{
+		const FPCGDataTypeIdentifier InputTypeUnion = GetTypeUnionIDOfIncidentEdges(PCGPinConstants::DefaultInputLabel);
+		const FPCGDataTypeIdentifier ReplacementTypeUnion = GetTypeUnionIDOfIncidentEdges(PCGReplaceDataByTags::Constants::ReplacementLabel);
+
+		const FPCGDataTypeIdentifier FinalInputsUnion = (InputTypeUnion | ReplacementTypeUnion);
+		return FinalInputsUnion.IsValid() ? FinalInputsUnion : InPin->Properties.AllowedTypes;
+	}
+
+	return InPin->Properties.AllowedTypes;
+}
+
+bool FPCGReplaceDataByTagElement::ExecuteInternal(FPCGContext* Context) const
+{
+	TRACE_CPUPROFILER_EVENT_SCOPE(FPCGReplaceDataByTagElement::Execute);
+
+	check(Context);
+
+	const UPCGReplaceDataByTagSettings* Settings = Context->GetInputSettings<UPCGReplaceDataByTagSettings>();
+	check(Settings);
+
+	TArray<FPCGTaggedData> Inputs = Context->InputData.GetInputsByPin(PCGPinConstants::DefaultInputLabel);
+	TArray<FPCGTaggedData> Replacements = Context->InputData.GetInputsByPin(PCGReplaceDataByTags::Constants::ReplacementLabel);
+
+	TArray<FPCGTaggedData>& Outputs = Context->OutputData.TaggedData;
+
+	auto FindTag = [](const TSet<FString>& InTags, const FString& InTagToFind) -> FString
+	{
+		for(const FString& InTag : InTags)
+		{
+			// check for InTagToFind existence (either exactly the tag) or a Tag:Value.
+			if (InTag == InTagToFind || InTag.StartsWith(InTagToFind + ":"))
+			{
+				return InTag;
+			}
+		}
+
+		return FString();
+	};
+
+	// For each input, if there is a data in the replacement set that has the same tags, then use this one instead.
+	for (const FPCGTaggedData& Input : Inputs)
+	{
+		FPCGTaggedData& Output = Outputs.Add_GetRef(Input);
+
+		// If there are no specified tags, then replace data only if there is a match in the replacements that has the same tags (and potentially more).
+		for (const FPCGTaggedData& Replacement : Replacements)
+		{
+			bool bReplaceByThisData = false;
+
+			// If the settings don't specify what tags to match against, then do a full inclusion test (e.g. all tags from input are in replacement data)
+			if (Settings->Tags.IsEmpty())
+			{
+				bReplaceByThisData = !Input.Tags.IsEmpty() && Replacement.Tags.Includes(Input.Tags);
+			}
+			// Otherwise, select the tags and compare them (to support Tag:Value entries)
+			else
+			{
+				bReplaceByThisData = true;
+
+				for (const FString& Tag : Settings->Tags)
+				{
+					FString TagFromInput = FindTag(Input.Tags, Tag);
+					FString TagFromReplacement = FindTag(Replacement.Tags, Tag);
+
+					// If the tag hasn't be found (first condition) or if it has a different full tag (incl. value) then we can't use it as a replacement
+					if (TagFromInput.IsEmpty() || TagFromInput != TagFromReplacement)
+					{
+						bReplaceByThisData = false;
+						break;
+					}
+				}
+			}
+
+			if (bReplaceByThisData)
+			{
+				Output = Replacement;
+				break;
+			}
+		}
+	}
+
+	return true;
+}
\ No newline at end of file
diff --git a/Engine/Plugins/PCG/Source/PCG/Public/Elements/Data/PCGReplaceDataByTags.h b/Engine/Plugins/PCG/Source/PCG/Public/Elements/Data/PCGReplaceDataByTags.h
new file mode 100644
--- /dev/null
+++ b/Engine/Plugins/PCG/Source/PCG/Public/Elements/Data/PCGReplaceDataByTags.h
@@ -0,0 +1,45 @@
+// Copyright Epic Games, Inc. All Rights Reserved.
+
+#pragma once
+
+#include "PCGSettings.h"
+
+#include "PCGReplaceDataByTags.generated.h"
+
+class FPCGReplaceDataByTagElement : public IPCGElement
+{
+protected:
+	virtual bool ExecuteInternal(FPCGContext* Context) const override;
+	virtual bool SupportsBasePointDataInputs(FPCGContext* InContext) const override { return true; }
+	// @todo_pcg: unclear if we support gpu data because of tags
+};
+
+/**
+* Selects data from primary input or from replacement input based on common tags or tag values.
+*/
+UCLASS(BlueprintType, ClassGroup = (Procedural))
+class UPCGReplaceDataByTagSettings : public UPCGSettings
+{
+	GENERATED_BODY()
+
+public:
+	//~Begin UPCGSettings interface
+#if WITH_EDITOR
+	virtual FName GetDefaultNodeName() const override { return FName(TEXT("ReplaceDataByTag")); }
+	virtual FText GetDefaultNodeTitle() const override { return NSLOCTEXT("PCGReplaceDataByTagElement", "NodeTitle", "Replace Data By Tag"); }
+	virtual FText GetNodeTooltipText() const override { return NSLOCTEXT("PCGReplaceDataByTagElement", "NodeTooltip", "Replaces data in the input data collection by data from the replacement collection based on common tags."); }
+	virtual EPCGSettingsType GetType() const override { return EPCGSettingsType::Generic; }
+#endif
+	virtual bool HasDynamicPins() const override { return true; }
+	virtual FPCGDataTypeIdentifier GetCurrentPinTypesID(const UPCGPin* InPin) const override;
+
+protected:
+	virtual TArray<FPCGPinProperties> InputPinProperties() const override;
+	virtual TArray<FPCGPinProperties> OutputPinProperties() const override;
+	virtual FPCGElementPtr CreateElement() const override { return MakeShared<FPCGReplaceDataByTagElement>(); }
+	//~End UPCGSettings interface
+
+public:
+	UPROPERTY(EditAnywhere, BlueprintReadWrite, Category = Settings, meta = (PCG_Overridable))
+	TArray<FString> Tags;
+};
\ No newline at end of file
