[PCG] Add scope guard to FPCGGraphTask::LoadCookedData()

#rnx
#Synced-CL 47850883

diff --git a/Engine/Plugins/PCG/Source/PCG/Private/Graph/PCGGraphTask.cpp b/Engine/Plugins/PCG/Source/PCG/Private/Graph/PCGGraphTask.cpp
--- a/Engine/Plugins/PCG/Source/PCG/Private/Graph/PCGGraphTask.cpp
+++ b/Engine/Plugins/PCG/Source/PCG/Private/Graph/PCGGraphTask.cpp
@@ -105,6 +105,7 @@ void FPCGGraphTask::LoadCookedData()
 {
 	if (ElementSource == EPCGElementSource::FromNode)
 	{
+		FGCScopeGuard Guard;
 		check(NodePtr.Get());
 		Node = NodePtr.Get();
 	}
