[PCG] Metadata String operation improvements

* Make sure we move values to Callbacks to avoid unecessary copy.
* Make sure to hide SearchCase and the input sources according to the operation
* Change output type of Matches/Substring to bool
* Add more information of Matches op regarding wildcards
* Add operators:
- Left/LeftChop/Right/RightChop/Mid
- RemoveFromStart/RemoveFromEnd/RemoveAt
- Find/FindLast

#rnx
#Synced-CL 48375549

diff --git a/Engine/Plugins/PCG/Source/PCG/Private/Elements/Metadata/PCGMetadataStringOpElement.cpp b/Engine/Plugins/PCG/Source/PCG/Private/Elements/Metadata/PCGMetadataStringOpElement.cpp
--- a/Engine/Plugins/PCG/Source/PCG/Private/Elements/Metadata/PCGMetadataStringOpElement.cpp
+++ b/Engine/Plugins/PCG/Source/PCG/Private/Elements/Metadata/PCGMetadataStringOpElement.cpp
@@ -4,37 +4,66 @@
 
 #include "Elements/Metadata/PCGMetadataElementCommon.h"
 
+#include "UObject/FortniteMainBranchObjectVersion.h"
+
 #include UE_INLINE_GENERATED_CPP_BY_NAME(PCGMetadataStringOpElement)
 
 namespace PCGMetadataStringOpConstants
 {
-	const FName ToFindLabel = TEXT("Search");
-	const FName ToReplaceLabel = TEXT("Replace");
+	constexpr FLazyName ToFindLabel = "Search";
+	constexpr FLazyName ToReplaceLabel = "Replace";
+	constexpr FLazyName CountLabel = "Count";
+	constexpr FLazyName StartLabel = "Start";
+	
 }
 
 FName UPCGMetadataStringOpSettings::GetInputPinLabel(uint32 Index) const
 {
 	switch (Operation)
 	{
-	case EPCGMetadataStringOperation::ToLower: // fall-through
-	case EPCGMetadataStringOperation::ToUpper: // fall-through
+	case EPCGMetadataStringOperation::ToLower:   // fall-through
+	case EPCGMetadataStringOperation::ToUpper:   // fall-through
 	case EPCGMetadataStringOperation::TrimStart: // fall-through
-	case EPCGMetadataStringOperation::TrimEnd: // fall-through
+	case EPCGMetadataStringOperation::TrimEnd:   // fall-through
 	case EPCGMetadataStringOperation::TrimStartAndEnd:
 	{
 		return PCGPinConstants::DefaultInputLabel;
 	}
 
-	case EPCGMetadataStringOperation::Append:    // fall-through
-	case EPCGMetadataStringOperation::Substring: // fall-through
-	case EPCGMetadataStringOperation::Matches:
+	case EPCGMetadataStringOperation::Append:          // fall-through
+	case EPCGMetadataStringOperation::Substring:       // fall-through
+	case EPCGMetadataStringOperation::Matches:         // fall-through
+	case EPCGMetadataStringOperation::RemoveFromStart: // fall-through
+	case EPCGMetadataStringOperation::RemoveFromEnd:   // fall-through
+	case EPCGMetadataStringOperation::Find:            // fall-through
+	case EPCGMetadataStringOperation::FindLast:
 		switch (Index)
 		{
 		case 0: return PCGMetadataSettingsBaseConstants::DoubleInputFirstLabel;
 		case 1: return PCGMetadataSettingsBaseConstants::DoubleInputSecondLabel;
 		default: break;
 		}
-
+		
+	case EPCGMetadataStringOperation::Left:     // fall-through
+	case EPCGMetadataStringOperation::LeftChop: // fall-through
+	case EPCGMetadataStringOperation::Right:    // fall-through
+	case EPCGMetadataStringOperation::RightChop:
+		switch (Index)
+		{
+		case 0: return PCGPinConstants::DefaultInputLabel;
+		case 1: return PCGMetadataStringOpConstants::CountLabel;
+		default: break;
+		}
+		
+	case EPCGMetadataStringOperation::Mid:
+		switch (Index)
+		{
+		case 0: return PCGPinConstants::DefaultInputLabel;
+		case 1: return PCGMetadataStringOpConstants::StartLabel;
+		case 2: return PCGMetadataStringOpConstants::CountLabel;
+		default: break;
+		}
+		
 	case EPCGMetadataStringOperation::Replace:
 		switch (Index)
 		{
@@ -52,28 +81,44 @@ FName UPCGMetadataStringOpSettings::GetInputPinLabel(uint32 Index) const
 
 uint32 UPCGMetadataStringOpSettings::GetOperandNum() const
 {
-	switch (Operation)
+	if (HasThreeOperands())
+	{
+		return 3;
+	}
+	else if (HasTwoOperands())
 	{
-		case EPCGMetadataStringOperation::ToLower: // fall-through
-		case EPCGMetadataStringOperation::ToUpper: // fall-through
-		case EPCGMetadataStringOperation::TrimStart: // fall-through
-		case EPCGMetadataStringOperation::TrimEnd: // fall-through
-		case EPCGMetadataStringOperation::TrimStartAndEnd:
-			return 1;
-		case EPCGMetadataStringOperation::Append:    // fall-through
-		case EPCGMetadataStringOperation::Substring: // fall-through
-		case EPCGMetadataStringOperation::Matches:
-			return 2;
-		case EPCGMetadataStringOperation::Replace:
-			return 3;
-		default:
-			return 0;
+		return 2;
+	}
+	else
+	{
+		return 1;
 	}
 }
 
 bool UPCGMetadataStringOpSettings::IsSupportedInputType(uint16 TypeId, uint32 InputIndex, bool& bHasSpecialRequirement) const
 {
 	bHasSpecialRequirement = false;
+	
+	switch (Operation)
+	{
+	case EPCGMetadataStringOperation::Left:    // fall-through
+	case EPCGMetadataStringOperation::LeftChop: // fall-through
+	case EPCGMetadataStringOperation::Right: // fall-through
+	case EPCGMetadataStringOperation::RightChop: // fall-through
+	case EPCGMetadataStringOperation::Mid:
+	{
+		bHasSpecialRequirement = true;
+		if (InputIndex >= 1)
+		{
+			return PCG::Private::IsBroadcastableOrConstructible(TypeId, PCG::Private::MetadataTypes<int32>::Id);
+		}
+		
+		break;
+	}
+	default:
+		break;
+	}
+	
 	return true; // all types support ToString()
 }
 
@@ -99,6 +144,19 @@ FString UPCGMetadataStringOpSettings::GetAdditionalTitleInformation() const
 		return FString();
 	}
 }
+
+void UPCGMetadataStringOpSettings::PostLoad()
+{
+	Super::PostLoad();
+
+#if WITH_EDITOR
+	if (GetLinkerCustomVersion(FFortniteMainBranchObjectVersion::GUID) < FFortniteMainBranchObjectVersion::PCGMetadataOpStringTypeChange)
+	{
+		// For all previous nodes, we'll force this option to true. for retro-compatibility
+		OutputTypeAsString = true;
+	}
+#endif // WITH_EDITOR
+}
  
 #if WITH_EDITOR
 FName UPCGMetadataStringOpSettings::GetDefaultNodeName() const
@@ -135,7 +193,67 @@ FPCGElementPtr UPCGMetadataStringOpSettings::CreateElement() const
 
 uint16 UPCGMetadataStringOpSettings::GetOutputType(uint16 InputTypeId) const
 {
-	return (uint16)EPCGMetadataTypes::String;
+	if (Operation == EPCGMetadataStringOperation::Find || Operation == EPCGMetadataStringOperation::FindLast)
+	{
+		return (uint16)EPCGMetadataTypes::Integer32;
+	}
+	else if (OutputsBool() && !OutputTypeAsString)
+	{
+		return (uint16)EPCGMetadataTypes::Boolean;
+	}
+	else
+	{
+		return (uint16)EPCGMetadataTypes::String;
+	}
+}
+
+EPCGMetadataTypes UPCGMetadataStringOpSettings::GetPinInitialDefaultValueType(FName PinLabel) const
+{
+	return (PinLabel == PCGMetadataStringOpConstants::CountLabel || PinLabel == PCGMetadataStringOpConstants::StartLabel) ? EPCGMetadataTypes::Integer32 : EPCGMetadataTypes::String;
+}
+
+#if WITH_EDITOR
+FString UPCGMetadataStringOpSettings::GetPinInitialDefaultValueString(FName PinLabel) const
+{
+	return (PinLabel == PCGMetadataStringOpConstants::CountLabel || PinLabel == PCGMetadataStringOpConstants::StartLabel) ? PCG::Private::MetadataTraits<int32>::ZeroValueString() : PCG::Private::MetadataTraits<FString>::ZeroValueString();
+}
+
+bool UPCGMetadataStringOpSettings::HasSearchCase() const
+{
+	return Operation == EPCGMetadataStringOperation::Substring
+		|| Operation == EPCGMetadataStringOperation::Matches
+		|| Operation == EPCGMetadataStringOperation::Replace
+		|| Operation == EPCGMetadataStringOperation::RemoveFromStart
+		|| Operation == EPCGMetadataStringOperation::RemoveFromEnd
+		|| Operation == EPCGMetadataStringOperation::Find
+		|| Operation == EPCGMetadataStringOperation::FindLast;
+}
+#endif // WITH_EDITOR
+
+bool UPCGMetadataStringOpSettings::HasTwoOperands() const
+{
+	return Operation == EPCGMetadataStringOperation::Append
+		|| Operation == EPCGMetadataStringOperation::Substring
+		|| Operation == EPCGMetadataStringOperation::Matches
+		|| Operation == EPCGMetadataStringOperation::Left
+		|| Operation == EPCGMetadataStringOperation::LeftChop
+		|| Operation == EPCGMetadataStringOperation::Right
+		|| Operation == EPCGMetadataStringOperation::RightChop
+		|| Operation == EPCGMetadataStringOperation::RemoveFromStart
+		|| Operation == EPCGMetadataStringOperation::RemoveFromEnd
+		|| Operation == EPCGMetadataStringOperation::Find
+		|| Operation == EPCGMetadataStringOperation::FindLast;
+}
+
+bool UPCGMetadataStringOpSettings::HasThreeOperands() const
+{
+	return Operation == EPCGMetadataStringOperation::Replace
+		|| Operation == EPCGMetadataStringOperation::Mid;
+}
+
+bool UPCGMetadataStringOpSettings::OutputsBool() const
+{
+	return Operation == EPCGMetadataStringOperation::Matches || Operation == EPCGMetadataStringOperation::Substring;
 }
 
 bool FPCGMetadataStringOpElement::DoOperation(PCGMetadataOps::FOperationData& OperationData) const
@@ -176,11 +294,52 @@ bool FPCGMetadataStringOpElement::DoOperation(PCGMetadataOps::FOperationData& Op
 	{
 		return DoBinaryOp<FString, FString>(OperationData, [Settings](const FString& A, const FString& B) -> bool { return PCG::Private::MetadataTraits<FString>::Matches(A, B, Settings->SearchCase); });
 	}
+	else if (Settings->Operation == EPCGMetadataStringOperation::RemoveFromStart)
+	{
+		// Intentionally pass A by value, since `RemoveFromStart` is in place, but since we have just a single output, A will be moved (cf. PCG::Private::NAryOperation::Apply)
+		return DoBinaryOp<FString, FString>(OperationData, [Settings](FString A, const FString& B) -> FString { A.RemoveFromStart(B, Settings->SearchCase); return A; });
+	}
+	else if (Settings->Operation == EPCGMetadataStringOperation::RemoveFromEnd)
+	{
+		// Intentionally pass A by value, since `RemoveFromEnd` is in place, but since we have just a single output, A will be moved (cf. PCG::Private::NAryOperation::Apply)
+		return DoBinaryOp<FString, FString>(OperationData, [Settings](FString A, const FString& B) -> FString { A.RemoveFromEnd(B, Settings->SearchCase); return A; });
+	}
+	else if (Settings->Operation == EPCGMetadataStringOperation::Left)
+	{
+		return DoBinaryOp<FString, int32>(OperationData, [](const FString& A, const int32 Count) -> FString { return A.Left(Count); });
+	}
+	else if (Settings->Operation == EPCGMetadataStringOperation::LeftChop)
+	{
+		return DoBinaryOp<FString, int32>(OperationData, [](const FString& A, const int32 Count) -> FString { return A.LeftChop(Count); });
+	}
+	else if (Settings->Operation == EPCGMetadataStringOperation::Right)
+	{
+		return DoBinaryOp<FString, int32>(OperationData, [](const FString& A, const int32 Count) -> FString { return A.Right(Count); });
+	}
+	else if (Settings->Operation == EPCGMetadataStringOperation::RightChop)
+	{
+		return DoBinaryOp<FString, int32>(OperationData, [](const FString& A, const int32 Count) -> FString { return A.RightChop(Count); });
+	}
+	else if (Settings->Operation == EPCGMetadataStringOperation::Find)
+	{
+		return DoBinaryOp<FString, FString>(OperationData, [Settings](const FString& A, const FString& B) -> int32 { return A.Find(B, Settings->SearchCase, ESearchDir::FromStart); });
+	}
+	else if (Settings->Operation == EPCGMetadataStringOperation::FindLast)
+	{
+		return DoBinaryOp<FString, FString>(OperationData, [Settings](const FString& A, const FString& B) -> int32 { return A.Find(B, Settings->SearchCase, ESearchDir::FromEnd); });
+	}
+	else if (Settings->Operation == EPCGMetadataStringOperation::Mid)
+	{
+		return DoTernaryOp<FString, int32, int32>(OperationData, [](const FString& InValue, const int32 Start, const int32 Count) -> FString
+		{
+			return InValue.Mid(Start, Count);
+		});
+	}
 	else if (Settings->Operation == EPCGMetadataStringOperation::Replace)
 	{
-		return DoTernaryOp<FString, FString, FString>(OperationData, [](const FString& InValue, const FString& InSearch, const FString& InReplace)
+		return DoTernaryOp<FString, FString, FString>(OperationData, [Settings](const FString& InValue, const FString& InSearch, const FString& InReplace) -> FString
 		{
-			return InValue.Replace(*InSearch, *InReplace);
+			return InValue.Replace(*InSearch, *InReplace, Settings->SearchCase);
 		});
 	}
 	else
@@ -188,4 +347,4 @@ bool FPCGMetadataStringOpElement::DoOperation(PCGMetadataOps::FOperationData& Op
 		ensure(false);
 		return true;
 	}
-}
\ No newline at end of file
+}
diff --git a/Engine/Plugins/PCG/Source/PCG/Public/Elements/Metadata/PCGMetadataOpElementBase.h b/Engine/Plugins/PCG/Source/PCG/Public/Elements/Metadata/PCGMetadataOpElementBase.h
--- a/Engine/Plugins/PCG/Source/PCG/Public/Elements/Metadata/PCGMetadataOpElementBase.h
+++ b/Engine/Plugins/PCG/Source/PCG/Public/Elements/Metadata/PCGMetadataOpElementBase.h
@@ -341,7 +341,24 @@ namespace NAryOperation
 			// We use argument unpacking, since our InArgs contains all the arrays with our values.
 			// If in InArgs we have TArray<int>& IntValues, TArray<float>& FloatValues, the unpack will do:
 			// (IntValues[0], FloatValues[0])
-			auto OutValue = InCallbacks.template Get<OutputIndex>()(InArgs[0]...);
+
+			// Implementation note: If we have just a single callback, it is safe to move InArgs, as it will only be used once.
+			// But if we have multiple callbacks, it is not until the last one. So only move if we are at the last callback.
+			// Since we move backward, it is if OutputIndex == 0
+			auto InvokeCallback = [&InCallbacks, &InArgs...](const int32 i)
+			{
+				constexpr bool bCanMove = OutputIndex == 0;
+				if constexpr (bCanMove)
+				{
+					return InCallbacks.template Get<OutputIndex>()(MoveTemp(InArgs[i])...);
+				}
+				else
+				{
+					return InCallbacks.template Get<OutputIndex>()(InArgs[i]...);
+				}
+			};
+			
+			auto OutValue = InvokeCallback(0);
 			using OutType = decltype(OutValue);
 
 			TArray<OutType, TInlineAllocator<DefaultChunkSize>> OutValues;
@@ -358,7 +375,7 @@ namespace NAryOperation
 
 			for (int32 i = 1; i < Range; ++i)
 			{
-				OutValues[i] = InCallbacks.template Get<OutputIndex>()(InArgs[i]...);
+				OutValues[i] = InvokeCallback(i);
 			}
 
 			bool bSuccess = true;
@@ -390,7 +407,7 @@ namespace NAryOperation
 	template <typename... Callbacks, typename... Args>
 	bool Gather(PCGMetadataOps::FOperationData& InOperationData, int32 StartIndex, int32 Range, const Options& InOptions, const TTuple<Callbacks...>& InCallbacks, int InputIndex, Signature<> S, Args&& ...InArgs)
 	{
-		TRACE_CPUPROFILER_EVENT_SCOPE(PCG::Private::NAryOperation::Apply);
+		//TRACE_CPUPROFILER_EVENT_SCOPE(PCG::Private::NAryOperation::Apply);
 
 		constexpr int LastOutputIndex = (int)sizeof...(Callbacks) - 1;
 		return Apply<LastOutputIndex>(InOperationData, StartIndex, Range, InOptions, InCallbacks, std::forward<Args>(InArgs)...);
@@ -447,7 +464,7 @@ namespace NAryOperation
 	template <typename... InputTypes, typename... Callbacks>
 	bool Operation(PCGMetadataOps::FOperationData& InOperationData, int32 StartIndex, int32 Range, const Options& InOptions, const TTuple<Callbacks...>& InCallbacks)
 	{
-		TRACE_CPUPROFILER_EVENT_SCOPE(PCG::Private::NAryOperation::Operation);
+		//TRACE_CPUPROFILER_EVENT_SCOPE(PCG::Private::NAryOperation::Operation);
 		return Gather(InOperationData, StartIndex, Range, InOptions, InCallbacks, 0, Signature<InputTypes...>());
 	}
 }
diff --git a/Engine/Plugins/PCG/Source/PCG/Public/Elements/Metadata/PCGMetadataStringOpElement.h b/Engine/Plugins/PCG/Source/PCG/Public/Elements/Metadata/PCGMetadataStringOpElement.h
--- a/Engine/Plugins/PCG/Source/PCG/Public/Elements/Metadata/PCGMetadataStringOpElement.h
+++ b/Engine/Plugins/PCG/Source/PCG/Public/Elements/Metadata/PCGMetadataStringOpElement.h
@@ -12,12 +12,21 @@ enum class EPCGMetadataStringOperation : uint16
 	Append UMETA(DisplayName="Append String"),
 	Replace UMETA(DisplayName="Replace String"),
 	Substring UMETA(DisplayName="Substring", Tooltip="True if string A contains substring B.", SearchHints="contains"),
-	Matches UMETA(DisplayName="Matches", Tooltip="True if string A matches substring B exactly.", SearchHints="equal"),
+	Matches UMETA(DisplayName="Matches Wildcard", Tooltip="True if string A matches substring B exactly. B can have wildcards.\n*?-type wildcard: '*' matches 0 or more arbitrary characters, '?' matches either an arbitrary character or nothing, all other characters match themselves.", SearchHints="equal"),
 	ToUpper UMETA(DisplayName="To Upper", Tooltip="Convert all characters to upper case."),
 	ToLower UMETA(DisplayName="To Lower", Tooltip="Convert all characters to lower case."),
 	TrimStart UMETA(DisplayName="Trim Start", Tooltip="Trim whitespace from the beginning of the string."),
 	TrimEnd UMETA(DisplayName="Trim End", Tooltip="Trim whitespace from the end of the string."),
-	TrimStartAndEnd UMETA(DisplayName="Trim Start and End", Tooltip="Trim whitespace from the beginning and end of the string.")
+	TrimStartAndEnd UMETA(DisplayName="Trim Start and End", Tooltip="Trim whitespace from the beginning and end of the string."),
+	Left UMETA(ToolTip="Returns the left most given number of characters"),
+	LeftChop UMETA(DisplayName="Left Chop", ToolTip="Returns the left most characters from the string chopping the given number of characters from the end"),
+	Right UMETA(ToolTip="Returns the right most given number of characters"),
+	RightChop UMETA(DisplayName="Right Chop", ToolTip="Returns the right most characters from the string chopping the given number of characters from the start"),
+	Mid UMETA(ToolTip="Returns the substring from Start position for Count characters."),
+	RemoveFromStart UMETA(DisplayName="Remove from Start", ToolTip="Removes the text from the start of the string if it exists."),
+	RemoveFromEnd UMETA(DisplayName="Remove from End", ToolTip="Removes the text from the end of the string if it exists."),
+	Find UMETA(ToolTip="Return the index of the first occurence of B in A. -1 if not found"),
+	FindLast UMETA(DisplayName = "Find Last", ToolTip="Return the index of the last occurence of B in A. -1 if not found")
 };
 
 UCLASS(MinimalAPI, BlueprintType, ClassGroup = (Procedural))
@@ -26,6 +35,10 @@ class UPCGMetadataStringOpSettings : public UPCGMetadataSettingsBase
 	GENERATED_BODY()
 
 public:
+	//~Begin UObject interface
+	virtual void PostLoad() override;
+	//~End UObject interface
+	
 	//~Begin UPCGSettings interface
 #if WITH_EDITOR
 	virtual FName GetDefaultNodeName() const override;
@@ -56,26 +69,45 @@ public:
 
 	//~Begin IPCGSettingsDefaultValueProvider interface
 #if WITH_EDITOR
-	virtual FString GetPinInitialDefaultValueString(FName PinLabel) const override { return PCG::Private::MetadataTraits<FString>::ZeroValueString(); }
+	virtual FString GetPinInitialDefaultValueString(FName PinLabel) const override;
 #endif // WITH_EDITOR
-	virtual EPCGMetadataTypes GetPinInitialDefaultValueType(FName PinLabel) const override { return EPCGMetadataTypes::String; }
+	virtual EPCGMetadataTypes GetPinInitialDefaultValueType(FName PinLabel) const override;
 	//~End IPCGSettingsDefaultValueProvider interface
 
 public:
 	UPROPERTY(BlueprintReadWrite, EditAnywhere, Category = Settings)
 	EPCGMetadataStringOperation Operation = EPCGMetadataStringOperation::Append;
 
-	UPROPERTY(BlueprintReadWrite, EditAnywhere, Category = Settings, meta = (PCG_Overridable, EditCondition = "Operation == EPCGMetadataStringOperation::Substring || Operation == EPCGMetadataStringOperation::Matches", EditConditionHides))
+	UPROPERTY(BlueprintReadWrite, EditAnywhere, Category = Settings, meta = (PCG_Overridable, EditCondition = "HasSearchCase()", EditConditionHides))
 	TEnumAsByte<ESearchCase::Type> SearchCase = ESearchCase::CaseSensitive;
 
 	UPROPERTY(BlueprintReadWrite, EditAnywhere, Category = Input, meta = (PCG_Overridable))
 	FPCGAttributePropertyInputSelector InputSource1;
 
-	UPROPERTY(BlueprintReadWrite, EditAnywhere, Category = Input, meta = (PCG_Overridable))
+	UPROPERTY(BlueprintReadWrite, EditAnywhere, Category = Input, meta = (PCG_Overridable, EditCondition = "HasTwoOperands() || HasThreeOperands()", EditConditionHides))
 	FPCGAttributePropertyInputSelector InputSource2;
 
-	UPROPERTY(BlueprintReadWrite, EditAnywhere, Category = Input, meta = (PCG_Overridable, EditCondition = "Operation==EPCGMetadataStringOperation::Replace", EditConditionHides))
+	UPROPERTY(BlueprintReadWrite, EditAnywhere, Category = Input, meta = (PCG_Overridable, EditCondition = "HasThreeOperands()", EditConditionHides))
 	FPCGAttributePropertyInputSelector InputSource3;
+
+	/** For deprecation, for Matches and Substring, allow to output the result in String instead of bool. */
+	UPROPERTY(BlueprintReadWrite, EditAnywhere, Category = Output, AdvancedDisplay, meta=(EditCondition = "OutputsBool()", EditConditionHides))
+	bool OutputTypeAsString = false;
+
+private:
+#if WITH_EDITOR
+	UFUNCTION()
+	bool HasSearchCase() const;
+#endif // WITH_EDITOR
+
+	UFUNCTION()
+	bool HasTwoOperands() const;
+
+	UFUNCTION()
+	bool HasThreeOperands() const;
+
+	UFUNCTION()
+	bool OutputsBool() const;
 };
 
 class FPCGMetadataStringOpElement : public FPCGMetadataElementBase
diff --git a/Engine/Source/Runtime/Core/Public/UObject/FortniteMainBranchObjectVersions.inl b/Engine/Source/Runtime/Core/Public/UObject/FortniteMainBranchObjectVersions.inl
--- a/Engine/Source/Runtime/Core/Public/UObject/FortniteMainBranchObjectVersions.inl
+++ b/Engine/Source/Runtime/Core/Public/UObject/FortniteMainBranchObjectVersions.inl
@@ -709,5 +709,8 @@ FORTNITE_MAIN_VERSION(UAFVariableReferenceGUID, 231)
 // Unified source color settings between img media & media io
 FORTNITE_MAIN_VERSION(MediaColorSettingsUnification, 232)
 
+// Changed behavior of the String Operation metadata node in PCG for Substring and Matches
+FORTNITE_MAIN_VERSION(PCGMetadataOpStringTypeChange, 233)
+
 // -----<new versions can be added above this line>-------------------------------------------------
 #undef FORTNITE_MAIN_VERSION
