[PCG] Label server items in debug object tree so that what appear to be duplicated entries can be understood.

#rnx
#Synced-CL 48343954

diff --git a/Engine/Plugins/PCG/Source/PCGEditor/Private/Widgets/SPCGEditorGraphDebugObjectTree.cpp b/Engine/Plugins/PCG/Source/PCGEditor/Private/Widgets/SPCGEditorGraphDebugObjectTree.cpp
--- a/Engine/Plugins/PCG/Source/PCGEditor/Private/Widgets/SPCGEditorGraphDebugObjectTree.cpp
+++ b/Engine/Plugins/PCG/Source/PCGEditor/Private/Widgets/SPCGEditorGraphDebugObjectTree.cpp
@@ -243,7 +243,14 @@ void FPCGEditorGraphDebugObjectItem::SortChildren(bool bIsAscending, bool bIsRec
 
 FString FPCGEditorGraphDebugObjectItem_Actor::GetLabel(bool bForSorting) const
 {
-	return Actor.IsValid() ? Actor->GetActorNameOrLabel() : FString();
+	if (Actor.IsValid())
+	{
+		return PCGHelpers::IsServer(Actor->GetWorld()) ? (Actor->GetActorNameOrLabel() + FString(TEXT(" (Server)"))) : Actor->GetActorNameOrLabel();
+	}
+	else
+	{
+		return FString();
+	}
 }
 
 const FSlateBrush* FPCGEditorGraphDebugObjectItem_Actor::GetIcon() const
