[PCG] Added support to move generated data to target class on isolate.

#Synced-CL 48394633

diff --git a/Engine/Plugins/PCG/Source/PCG/Internal/PCGIsolatedActor.h b/Engine/Plugins/PCG/Source/PCG/Internal/PCGIsolatedActor.h
new file mode 100644
--- /dev/null
+++ b/Engine/Plugins/PCG/Source/PCG/Internal/PCGIsolatedActor.h
@@ -0,0 +1,25 @@
+// Copyright Epic Games, Inc. All Rights Reserved.
+
+#pragma once
+
+#include "PCGData.h"
+
+#include "GameFramework/Actor.h"
+
+#include "PCGIsolatedActor.generated.h"
+
+/** 
+* Simnple actor wrapper class to store PCG generated data.
+*/
+UCLASS(MinimalAPI, NotPlaceable)
+class APCGIsolatedActor : public AActor
+{
+	GENERATED_BODY()
+
+public:
+	UPROPERTY(VisibleInstanceOnly, Category = PCG)
+	FPCGDataCollection Data;
+
+	UPROPERTY(VisibleInstanceOnly, Category = PCG)
+	TSoftObjectPtr<AActor> Origin = nullptr;
+};
\ No newline at end of file
diff --git a/Engine/Plugins/PCG/Source/PCG/Private/PCGComponent.cpp b/Engine/Plugins/PCG/Source/PCG/Private/PCGComponent.cpp
--- a/Engine/Plugins/PCG/Source/PCG/Private/PCGComponent.cpp
+++ b/Engine/Plugins/PCG/Source/PCG/Private/PCGComponent.cpp
@@ -1439,6 +1439,26 @@ bool UPCGComponent::MoveResourcesToNewActor(const FPCGMoveResourceParams& Params
 		}
 	}
 
+	// Move generated output if the params indicate to do so.
+	if (Params.TargetDataCollection)
+	{
+		for(int DataIndex = GeneratedGraphOutput.TaggedData.Num() - 1; DataIndex >= 0; --DataIndex)
+		{
+			FPCGTaggedData OutputData = GeneratedGraphOutput.TaggedData[DataIndex];
+
+			if (Params.Accepts(OutputData))
+			{
+				check(OutputData.Data);
+				// Remove it from the generated graph output (incl. cache)
+				OutputData.Data->ClearUsage(EPCGDataUsage::ComponentOutputData);
+				GeneratedGraphOutput.TaggedData.RemoveAt(DataIndex);
+				Params.TargetDataCollection->TaggedData.Add(OutputData);
+
+				bHasMovedResources = true;
+			}
+		}
+	}
+
 	if (CreatedChildActor)
 	{
 		// No resources moved destroy stamp actor
diff --git a/Engine/Plugins/PCG/Source/PCG/Private/PCGManagedResource.cpp b/Engine/Plugins/PCG/Source/PCG/Private/PCGManagedResource.cpp
--- a/Engine/Plugins/PCG/Source/PCG/Private/PCGManagedResource.cpp
+++ b/Engine/Plugins/PCG/Source/PCG/Private/PCGManagedResource.cpp
@@ -57,6 +57,22 @@ bool FPCGMoveResourceParams::Accepts(const UObject* InObject) const
 	}
 }
 
+bool FPCGMoveResourceParams::Accepts(const FPCGTaggedData& InData) const
+{
+	if (!InData.Data)
+	{
+		return false;
+	}
+	else if (RequiredTags.IsEmpty())
+	{
+		return true;
+	}
+	else
+	{
+		return Algo::AllOf(RequiredTags, [&InData](FName Tag) { return InData.Tags.Contains(Tag.ToString()); });
+	}
+}
+
 void UPCGManagedResource::PostApplyToComponent()
 {
 	// Nothing - apply to component should already properly remap most of everything we need to do
diff --git a/Engine/Plugins/PCG/Source/PCG/Public/PCGManagedResource.h b/Engine/Plugins/PCG/Source/PCG/Public/PCGManagedResource.h
--- a/Engine/Plugins/PCG/Source/PCG/Public/PCGManagedResource.h
+++ b/Engine/Plugins/PCG/Source/PCG/Public/PCGManagedResource.h
@@ -22,13 +22,21 @@ class UActorComponent;
 class UInstancedSkinnedMeshComponent;
 class UInstancedStaticMeshComponent;
 class USplineMeshComponent;
+struct FPCGDataCollection;
+struct FPCGTaggedData;
 
 struct FPCGMoveResourceParams
 {
+	/** Class of the object that needs to be created as part of moving the resources. */
 	UClass* TemplateTargetClass = nullptr;
+	/** Target object to move resources to. Can be null, in which case we'll create one from the TemplateTargetClass. */
 	UObject* Target = nullptr;
+	/** Additional validation object owner. */
 	const UObject* ExpectedPreviousOwner = nullptr;
+	/** Conjunction of tags required to move a resource to the target. */
 	TArray<FName> RequiredTags;
+	/** If there is generated data on the source objects, it will be moved using the same requirements, into the target collection. */
+	FPCGDataCollection* TargetDataCollection = nullptr;
 
 	template<typename ObjectType>
 	TSubclassOf<ObjectType> GetTemplateClass() const
@@ -56,6 +64,7 @@ struct FPCGMoveResourceParams
 	}
 
 	bool Accepts(const UObject* InObject) const;
+	bool Accepts(const FPCGTaggedData& InData) const;
 };
 
 /** 
diff --git a/Engine/Plugins/PCG/Source/PCGEditor/Private/EditorMode/Tools/Query/PCGQueryTool.cpp b/Engine/Plugins/PCG/Source/PCGEditor/Private/EditorMode/Tools/Query/PCGQueryTool.cpp
--- a/Engine/Plugins/PCG/Source/PCGEditor/Private/EditorMode/Tools/Query/PCGQueryTool.cpp
+++ b/Engine/Plugins/PCG/Source/PCGEditor/Private/EditorMode/Tools/Query/PCGQueryTool.cpp
@@ -5,7 +5,9 @@
 #include "PCGEditorModule.h"
 #include "EditorMode/PCGEdModeStyle.h"
 
+#include "PCGIsolatedActor.h"
 #include "PCGManagedResource.h"
+#include "Helpers/PCGHelpers.h"
 
 #include "EditorViewportClient.h"
 #include "Components/InstancedStaticMeshComponent.h"
@@ -34,6 +36,57 @@ namespace PCGIsolateTool
 	{
 		const FName ToolTag = TEXT("IsolateTool");
 	}
+
+	AActor* Isolate(AActor* Actor, const TArray<FName>& Tags)
+	{
+		UPCGComponent* PCGComponent = Actor ? Actor->GetComponentByClass<UPCGComponent>() : nullptr;
+
+		if (!Actor || !PCGComponent)
+		{
+			return nullptr;
+		}
+
+		// Will store data temporarily (output from the ClearPCGLink) until we move it to the created actor
+		FPCGDataCollection TempCollection;
+
+		FPCGMoveResourceParams Params
+		{
+			.TemplateTargetClass = APCGIsolatedActor::StaticClass(),
+			.RequiredTags = Tags,
+			.TargetDataCollection = &TempCollection,
+		};
+
+		AActor* IsolatedActor = PCGComponent->ClearPCGLink(Params);
+
+		if (IsolatedActor)
+		{
+			// @todo_pcg: Push tags onto actor?
+			// Remove all PCG Generated tags, so we don't ignore those elsewhere
+			IsolatedActor->ForEachComponent(/*bIncludeFromChildActors=*/true, [](UActorComponent* Component)
+			{
+				Component->ComponentTags.Remove(PCGHelpers::DefaultPCGTag);
+			});
+
+			if (APCGIsolatedActor* IsolatedPCGActor = Cast<APCGIsolatedActor>(IsolatedActor))
+			{
+				IsolatedPCGActor->Origin = Actor;
+
+				// Move collection to isolated actor
+				IsolatedPCGActor->Data = MoveTemp(TempCollection);
+
+				// Re-outer data
+				for (FPCGTaggedData& TaggedData : IsolatedPCGActor->Data.TaggedData)
+				{
+					TaggedData.Data->VisitDataNetwork([IsolatedPCGActor](const UPCGData* InData)
+					{
+						const_cast<UPCGData*>(InData)->Rename(nullptr, IsolatedPCGActor, REN_DontCreateRedirectors);
+					});
+				}
+			}
+		}
+
+		return IsolatedActor;
+	}
 }
 
 FName UPCGQueryTool::StaticGetToolTag()
@@ -228,6 +281,11 @@ UInteractiveTool* UPCGQueryToolBuilder::BuildTool(const FToolBuilderState& Scene
 	}
 }
 
+bool UPCGQueryToolBuilder::CanBuildTool(const FToolBuilderState& SceneState) const
+{
+	return !!ToolClass;
+}
+
 void UPCGQueryToolBuilder::SetToolClass(TSubclassOf<UPCGQueryTool> InClass)
 {
 	ToolClass = InClass;
@@ -237,21 +295,7 @@ void UPCGInteractiveToolSettings_Isolate::Apply(UInteractiveTool* OwningTool)
 {
 	if (AActor* Actor = SelectedActor.Get())
 	{
-		if (UPCGComponent* PCGComponent = Actor->GetComponentByClass<UPCGComponent>())
-		{
-			FPCGMoveResourceParams Params
-			{
-				.RequiredTags = SelectedTags
-			};
-
-			AActor* IsolatedActor = PCGComponent->ClearPCGLink(Params);
-
-			if (IsolatedActor)
-			{
-				// @todo_pcg: Push tags onto actor?
-				// @todo_pcg: Clear PCG generated component?
-			}
-		}
+		PCGIsolateTool::Isolate(Actor, SelectedTags);
 	}
 }
 
diff --git a/Engine/Plugins/PCG/Source/PCGEditor/Private/EditorMode/Tools/Query/PCGQueryTool.h b/Engine/Plugins/PCG/Source/PCGEditor/Private/EditorMode/Tools/Query/PCGQueryTool.h
--- a/Engine/Plugins/PCG/Source/PCGEditor/Private/EditorMode/Tools/Query/PCGQueryTool.h
+++ b/Engine/Plugins/PCG/Source/PCGEditor/Private/EditorMode/Tools/Query/PCGQueryTool.h
@@ -11,10 +11,17 @@
 
 #include "PCGQueryTool.generated.h"
 
+class AActor;
 class HHitProxy;
 class UActorComponent;
 class USingleClickInputBehavior;
 
+namespace PCGIsolateTool
+{
+	/** Extracts (like Clear PCG Link) content from the target actor, assuming it has a PCG component, into a PCG Isolated Actor. */
+	AActor* Isolate(AActor* Actor, const TArray<FName>& RequiredTags);
+}
+
 /**
 * Simple editor tool to gather information from a picked object in the world. 
 * By default logs visited hierarchy tags (from touched component, to actor, ...) but can be reimplemented in a derived class.
@@ -62,14 +69,14 @@ class UPCGQueryToolBuilder : public UPCGInteractiveToolBuilder
 	GENERATED_BODY()
 
 public:
-	virtual bool CanBuildTool(const FToolBuilderState& SceneState) const override { return true; }
+	virtual bool CanBuildTool(const FToolBuilderState& SceneState) const override;
 	virtual UInteractiveTool* BuildTool(const FToolBuilderState& SceneState) const override;
 
 	void SetToolClass(TSubclassOf<UPCGQueryTool> InClass);
 
 private:
 	UPROPERTY()
-	TSubclassOf<UPCGQueryTool> ToolClass;
+	TSubclassOf<UPCGQueryTool> ToolClass = nullptr;
 };
 
 UCLASS()
