[PCG] Add warnings when UpdatePins have destructive behavior in PostLoad

* Should help find issues when a pin disappear for some reason (not overridable anymore, renamed, etc...)

#rnx

#Synced-CL 48394839

diff --git a/Engine/Plugins/PCG/Source/PCG/Private/PCGGraph.cpp b/Engine/Plugins/PCG/Source/PCG/Private/PCGGraph.cpp
--- a/Engine/Plugins/PCG/Source/PCG/Private/PCGGraph.cpp
+++ b/Engine/Plugins/PCG/Source/PCG/Private/PCGGraph.cpp
@@ -22,6 +22,7 @@
 #include "Helpers/PCGPropertyHelpers.h"
 
 #include "AssetRegistry/AssetData.h"
+#include "Logging/StructuredLog.h"
 
 #if WITH_EDITOR
 #include "CoreGlobals.h"
@@ -50,6 +51,11 @@ namespace PCGGraph
 		TEXT("pcg.GPU.ComputeGraphInstancePooling"),
 		false,
 		TEXT("Caches compute graph instances rather than creating each instance and its data providers from scratch each time."));
+
+	static TAutoConsoleVariable<bool> CVarEnableWarningOnPinDestructiveChangesInPostLoad(
+		TEXT("pcg.Graph.EnableWarningOnPinDestructiveChangesInPostLoad"),
+		false,
+		TEXT("Output warnings when destructive changes happens on PostLoad in UpdatePins."));
 }
 
 namespace PCGGraphUtils
@@ -481,6 +487,8 @@ void UPCGGraph::PostLoad()
 	Super::PostLoad();
 
 #if WITH_EDITOR
+	UE_LOG_CONTEXT("PCGGraph", FSoftObjectPath(this));
+	
 	// TODO: review this once the data has been updated, it's unwanted weight going forward
 	// Deprecation
 	InputNode->ConditionalPostLoad();
@@ -552,7 +560,8 @@ void UPCGGraph::PostLoad()
 	ForEachNode([](UPCGNode* InNode) { InNode->ApplyDeprecationBeforeUpdatePins(); return true; });
 
 	// Update pins on all nodes
-	ForEachNode([](UPCGNode* InNode) { InNode->UpdatePins(); return true; });
+	bool bQuiet = !PCGGraph::CVarEnableWarningOnPinDestructiveChangesInPostLoad.GetValueOnAnyThread();
+	ForEachNode([bQuiet](UPCGNode* InNode) { InNode->UpdatePins(bQuiet); return true; });
 
 	// Finally, apply deprecation that changes edges/rebinds
 	ForEachNode([](UPCGNode* InNode) { InNode->ApplyDeprecation(); return true; });
diff --git a/Engine/Plugins/PCG/Source/PCG/Private/PCGNode.cpp b/Engine/Plugins/PCG/Source/PCG/Private/PCGNode.cpp
--- a/Engine/Plugins/PCG/Source/PCG/Private/PCGNode.cpp
+++ b/Engine/Plugins/PCG/Source/PCG/Private/PCGNode.cpp
@@ -16,6 +16,8 @@
 #include "Editor.h"
 #endif
 
+#include "Logging/StructuredLog.h"
+
 #include UE_INLINE_GENERATED_CPP_BY_NAME(PCGNode)
 
 namespace PCGNodeHelpers
@@ -895,9 +897,9 @@ void UPCGNode::UpdateAfterSettingsChangeDuringCreation()
 	UpdatePins();
 }
 
-EPCGChangeType UPCGNode::UpdatePins()
+EPCGChangeType UPCGNode::UpdatePins(bool bQuiet)
 {
-	return UpdatePins([](UPCGNode* Node){ return NewObject<UPCGPin>(Node); });
+	return UpdatePins([](UPCGNode* Node){ return NewObject<UPCGPin>(Node); }, bQuiet);
 }
 
 void UPCGNode::CreateDefaultPins(TFunctionRef<UPCGPin* (UPCGNode*)> PinAllocator)
@@ -919,7 +921,7 @@ void UPCGNode::CreateDefaultPins(TFunctionRef<UPCGPin* (UPCGNode*)> PinAllocator
 	CreatePins(MutableView(OutputPins), Settings->DefaultOutputPinProperties());
 }
 
-EPCGChangeType UPCGNode::UpdatePins(TFunctionRef<UPCGPin*(UPCGNode*)> PinAllocator)
+EPCGChangeType UPCGNode::UpdatePins(TFunctionRef<UPCGPin*(UPCGNode*)> PinAllocator, bool bQuiet)
 {
 	TSet<UPCGNode*> TouchedNodes;
 	const UPCGSettings* Settings = GetSettings();
@@ -954,11 +956,16 @@ EPCGChangeType UPCGNode::UpdatePins(TFunctionRef<UPCGPin*(UPCGNode*)> PinAllocat
 		OutputPins.Reset();
 		return EPCGChangeType::Edge | EPCGChangeType::Node;
 	}
+
+	UE_LOG_CONTEXT("PCGNode", FSoftObjectPath(Settings));
+
+	const FName SettingsName = Settings->GetFName();
+	const FString GraphPath = GetGraph() ?  FSoftObjectPath(GetGraph()).ToString() : TEXT("Unknown");
 	
 	TArray<FPCGPinProperties> InboundPinProperties = Settings->AllInputPinProperties();
 	TArray<FPCGPinProperties> OutboundPinProperties = Settings->AllOutputPinProperties();
 
-	auto RemoveDuplicates = [this](TArray<FPCGPinProperties>& Properties)
+	auto RemoveDuplicates = [this, &GraphPath](TArray<FPCGPinProperties>& Properties)
 	{
 		for (int32 i = Properties.Num() - 2; i >= 0; --i)
 		{
@@ -966,10 +973,8 @@ EPCGChangeType UPCGNode::UpdatePins(TFunctionRef<UPCGPin*(UPCGNode*)> PinAllocat
 			{
 				if (Properties[i].Label == Properties[j].Label)
 				{
-					const UPCGGraph* PCGGraph = GetGraph();
-					const FString GraphName = PCGGraph ? PCGGraph->GetName() : FString(TEXT("Unknown"));
 					UE_LOG(LogPCG, Warning, TEXT("UpdatePins: Pin properties from the settings on node '%s' in graph '%s' contained a duplicate pin '%s', removing this pin properties."),
-						*GetName(), *GraphName, *Properties[i].Label.ToString());
+						*GetName(), *GraphPath, *Properties[i].Label.ToString());
 
 					// Remove but preserve order
 					Properties.RemoveAt(j);
@@ -981,7 +986,7 @@ EPCGChangeType UPCGNode::UpdatePins(TFunctionRef<UPCGPin*(UPCGNode*)> PinAllocat
 	RemoveDuplicates(InboundPinProperties);
 	RemoveDuplicates(OutboundPinProperties);
 
-	auto UpdatePins = [this, &PinAllocator, &TouchedNodes](TArray<UPCGPin*>& Pins, const TArray<FPCGPinProperties>& PinProperties)
+	auto UpdatePins = [this, &PinAllocator, &TouchedNodes, &SettingsName, bQuiet, &GraphPath](TArray<UPCGPin*>& Pins, const TArray<FPCGPinProperties>& PinProperties)
 	{
 		bool bAppliedEdgeChanges = false;
 		bool bChangedPins = false;
@@ -998,7 +1003,17 @@ EPCGChangeType UPCGNode::UpdatePins(TFunctionRef<UPCGPin*(UPCGNode*)> PinAllocat
 					Pin->Modify();
 					Pin->Properties = *MatchingProperties;
 
-					bAppliedEdgeChanges |= Pin->BreakAllIncompatibleEdges(&TouchedNodes);
+					FString Reason{};
+
+					if (Pin->BreakAllIncompatibleEdges(&TouchedNodes, !bQuiet ? &Reason : nullptr))
+					{
+						bAppliedEdgeChanges = true;
+						if (!bQuiet)
+						{
+							UE_LOG(LogPCG, Warning, TEXT("[%s - %s] UpdatePins: %s"), *GraphPath, *SettingsName.ToString(), *Reason);	
+						}
+					}
+
 					bChangedPins = true;
 				}
 #if WITH_EDITOR
@@ -1012,7 +1027,17 @@ EPCGChangeType UPCGNode::UpdatePins(TFunctionRef<UPCGPin*(UPCGNode*)> PinAllocat
 				// Make sure that if the pin type changed in the ApplyDeprecationBeforeUpdatePins to break incompatible edges.
 				if (Pin->bPropertiesHaveChangedOnDeprecation)
 				{
-					bAppliedEdgeChanges |= Pin->BreakAllIncompatibleEdges(&TouchedNodes);
+					FString Reason{};
+					
+					if (Pin->BreakAllIncompatibleEdges(&TouchedNodes, !bQuiet ? &Reason : nullptr))
+					{
+						bAppliedEdgeChanges = true;
+						if (!bQuiet)
+						{
+							UE_LOG(LogPCG, Warning, TEXT("[%s-%s] UpdatePins: %s"), *GraphPath, *SettingsName.ToString(), *Reason);	
+						}
+					}
+					
 					Pin->bPropertiesHaveChangedOnDeprecation = false;
 				}
 			}
@@ -1036,10 +1061,23 @@ EPCGChangeType UPCGNode::UpdatePins(TFunctionRef<UPCGPin*(UPCGNode*)> PinAllocat
 		const bool bUpdateUnmatchedPin = UnmatchedPins.Num() == 1 && UnmatchedProperties.Num() == 1;
 		if (bUpdateUnmatchedPin)
 		{
+			const FName PreviousName = UnmatchedPins[0]->Properties.Label;
+			const FName NewName = UnmatchedProperties[0].Label;
+
 			UnmatchedPins[0]->Modify();
 			UnmatchedPins[0]->Properties = UnmatchedProperties[0];
 
-			bAppliedEdgeChanges |= UnmatchedPins[0]->BreakAllIncompatibleEdges(&TouchedNodes);
+			FString Reason{};
+
+			if (UnmatchedPins[0]->BreakAllIncompatibleEdges(&TouchedNodes, !bQuiet ? &Reason : nullptr))
+			{
+				bAppliedEdgeChanges = true;
+				if (!bQuiet)
+				{
+					UE_LOG(LogPCG, Warning, TEXT("[%s-%s] UpdatePins: Pin %s was matched and renamed to %s\n%s"), *GraphPath, *SettingsName.ToString(), *PreviousName.ToString(), *NewName.ToString(), *Reason);	
+				}
+			}
+
 			bChangedPins = true;
 		}
 		else
@@ -1066,7 +1104,7 @@ EPCGChangeType UPCGNode::UpdatePins(TFunctionRef<UPCGPin*(UPCGNode*)> PinAllocat
 				bChangedPins = true;
 			}
 
-			auto RemovePins = [&Pins, &AllPinNames, &bAppliedEdgeChanges, &TouchedNodes](TArray<UPCGPin*>& PinsToRemove, bool bRemoveFromAllNames)
+			auto RemovePins = [&Pins, &AllPinNames, &bAppliedEdgeChanges, &TouchedNodes, bQuiet, &SettingsName, &GraphPath](TArray<UPCGPin*>& PinsToRemove, bool bRemoveFromAllNames, bool bWarn)
 			{
 				for (int32 RemovedPinIndex = PinsToRemove.Num() - 1; RemovedPinIndex >= 0; --RemovedPinIndex)
 				{
@@ -1078,14 +1116,42 @@ EPCGChangeType UPCGNode::UpdatePins(TFunctionRef<UPCGPin*(UPCGNode*)> PinAllocat
 							AllPinNames.Remove(Pins[PinIndex]->Properties.Label);
 						}
 
-						bAppliedEdgeChanges |= Pins[PinIndex]->BreakAllEdges(&TouchedNodes);
+						TSet<UPCGPin*> OtherPinsTouched{};
+
+						if (Pins[PinIndex]->BreakAllEdges(&TouchedNodes, !bQuiet ? &OtherPinsTouched : nullptr))
+						{
+							bAppliedEdgeChanges = true;
+
+							if (bWarn && !bQuiet)
+							{
+								TStringBuilder<512> WarnMessage;
+								WarnMessage += FString::Printf(TEXT("[%s-%s] UpdatePins: Pin %s does not exist anymore.\n"), *GraphPath, *SettingsName.ToString(), *Pins[PinIndex]->Properties.Label.ToString());
+								WarnMessage += TEXT("List of broken edges:\n");
+								for (const UPCGPin* OtherPin : OtherPinsTouched)
+								{
+									if (!OtherPin)
+									{
+										continue;
+									}
+
+									const FName OtherPinLabel = OtherPin->Properties.Label;
+									const UPCGSettings* OtherSettings = OtherPin->Node ? OtherPin->Node->GetSettings() : nullptr;
+									const FName OtherSettingsName = OtherSettings ? OtherSettings->GetFName() : "Unknown";
+
+									WarnMessage += FString::Printf(TEXT(" - %s:%s"), *OtherSettingsName.ToString(), *OtherPinLabel.ToString());
+								}
+
+								UE_LOG(LogPCG, Warning, TEXT("%s"), WarnMessage.ToString());
+							}
+						}
+
 						Pins.RemoveAt(PinIndex);
 					}
 				}
 			};
 
-			RemovePins(UnmatchedPins, /*bRemoveFromAllNames=*/ true);
-			RemovePins(DuplicatedNamePins, /*bRemoveFromAllNames=*/ false);
+			RemovePins(UnmatchedPins, /*bRemoveFromAllNames=*/true, /*bWarn=*/true);
+			RemovePins(DuplicatedNamePins, /*bRemoveFromAllNames=*/false, /*bWarn=*/false);
 
 			// Add new pins
 			for (const FPCGPinProperties& UnmatchedProperty : UnmatchedProperties)
diff --git a/Engine/Plugins/PCG/Source/PCG/Private/PCGPin.cpp b/Engine/Plugins/PCG/Source/PCG/Private/PCGPin.cpp
--- a/Engine/Plugins/PCG/Source/PCG/Private/PCGPin.cpp
+++ b/Engine/Plugins/PCG/Source/PCG/Private/PCGPin.cpp
@@ -237,7 +237,7 @@ bool UPCGPin::BreakEdgeTo(UPCGPin* OtherPin, TSet<UPCGNode*>* InTouchedNodes/*=
 	return false;
 }
 
-bool UPCGPin::BreakAllEdges(TSet<UPCGNode*>* InTouchedNodes/*= nullptr*/)
+bool UPCGPin::BreakAllEdges(TSet<UPCGNode*>* InTouchedNodes, TSet<UPCGPin*>* OptionalOtherPinsTouched)
 {
 	bool bChanged = false;
 
@@ -267,6 +267,11 @@ bool UPCGPin::BreakAllEdges(TSet<UPCGNode*>* InTouchedNodes/*= nullptr*/)
 			{
 				InTouchedNodes->Add(OtherPin->Node);
 			}
+
+			if (OptionalOtherPinsTouched)
+			{
+				OptionalOtherPinsTouched->Add(OtherPin);
+			}
 		}
 	}
 
@@ -275,22 +280,60 @@ bool UPCGPin::BreakAllEdges(TSet<UPCGNode*>* InTouchedNodes/*= nullptr*/)
 	return bChanged;
 }
 
-bool UPCGPin::BreakAllIncompatibleEdges(TSet<UPCGNode*>* InTouchedNodes/*= nullptr*/)
+bool UPCGPin::BreakAllIncompatibleEdges(TSet<UPCGNode*>* InTouchedNodes/*= nullptr*/, FString* OptionalReason)
 {
 	bool bChanged = false;
 	bool bHasAValidEdge = false;
 
+	auto AppendReason = [OptionalReason](FString InReason)
+	{
+		if (!OptionalReason)
+		{
+			return;
+		}
+
+		if (OptionalReason->IsEmpty())
+		{
+			*OptionalReason = MoveTemp(InReason);
+		}
+		else
+		{
+			*OptionalReason = FString::Printf(TEXT("%s\n%s"), **OptionalReason, *InReason);
+		}
+	};
+
+	const FName ThisPinName = Properties.Label;
+
 	for (int32 EdgeIndex = Edges.Num() - 1; EdgeIndex >= 0; --EdgeIndex)
 	{
 		UPCGEdge* Edge = Edges[EdgeIndex];
 		UPCGPin* OtherPin = Edge->GetOtherPin(this);
 
+		const UPCGNode* OtherNode = Cast<const UPCGNode>(OtherPin->GetOuter());
+		const FName OtherNodeName = OtherNode ? OtherNode->GetFName() : FName("Unknown");
+		const FName OtherPinName = OtherPin ? OtherPin->Properties.Label : FName("Unknown");
+
 		bool bRemoveEdge = (!AllowsMultipleConnections() && bHasAValidEdge);
+		if (bRemoveEdge && OptionalReason)
+		{
+			AppendReason(FString::Printf(TEXT("Breaking edge between %s on node %s and %s because the pin does not support multi connection."),
+				*OtherPinName.ToString(),
+				*OtherNodeName.ToString(),
+				*ThisPinName.ToString()));
+		}
 
 		// If the edge is not compatible because of a subtype mismatch, we do not want to break it for deprecation purposes.
 		const EPCGDataTypeCompatibilityResult Compatibility = GetCompatibilityWithOtherPin(OtherPin);
 		if (!PCGDataTypeCompatibilityResult::IsValid(Compatibility) && Compatibility != EPCGDataTypeCompatibilityResult::TypeCompatibleSubtypeNotCompatible)
 		{
+			if (OptionalReason)
+			{
+				AppendReason(FString::Printf(TEXT("Breaking edge between %s on node %s and %s because the pin types are not compatible."),
+					*OtherPinName.ToString(),
+					*OtherNodeName.ToString(),
+					*ThisPinName.ToString()));
+			}
+			
 			bRemoveEdge = true;
 		}
 
diff --git a/Engine/Plugins/PCG/Source/PCG/Public/PCGNode.h b/Engine/Plugins/PCG/Source/PCG/Public/PCGNode.h
--- a/Engine/Plugins/PCG/Source/PCG/Public/PCGNode.h
+++ b/Engine/Plugins/PCG/Source/PCG/Public/PCGNode.h
@@ -220,10 +220,10 @@ public:
 #endif // WITH_EDITOR
 	
 protected:
-	/** Updates pins based on node settings. Attempts to migrate pins via matching. Broadcasts node change events for affected nodes. */
-	UE_API EPCGChangeType UpdatePins();
-	/** Updates pins based on node settings PinAllocator creates new pin objects. Attempts to migrate pins via matching. Broadcasts node change events for affected nodes. */
-	UE_API EPCGChangeType UpdatePins(TFunctionRef<UPCGPin* (UPCGNode*)> PinAllocator);
+	/** Updates pins based on node settings. Attempts to migrate pins via matching. Broadcasts node change events for affected nodes. Can provide warning messages when something destructive happens. */
+	UE_API EPCGChangeType UpdatePins(bool bQuiet = true);
+	/** Updates pins based on node settings PinAllocator creates new pin objects. Attempts to migrate pins via matching. Broadcasts node change events for affected nodes. Can provide warning messages when something destructive happens. */
+	UE_API EPCGChangeType UpdatePins(TFunctionRef<UPCGPin* (UPCGNode*)> PinAllocator, bool bQuiet = true);
 
 	// When we create a new graph, we initialize the input/output nodes as default, with default pins.
 	// Those default pins are not serialized, therefore if we change the default pins, combined with the use
diff --git a/Engine/Plugins/PCG/Source/PCG/Public/PCGPin.h b/Engine/Plugins/PCG/Source/PCG/Public/PCGPin.h
--- a/Engine/Plugins/PCG/Source/PCG/Public/PCGPin.h
+++ b/Engine/Plugins/PCG/Source/PCG/Public/PCGPin.h
@@ -293,10 +293,10 @@ public:
 	UE_API bool BreakEdgeTo(UPCGPin* OtherPin, TSet<UPCGNode*>* InTouchedNodes = nullptr);
 
 	/** Breaks all connected edges. Optionally returns list of affected nodes (including the pin owner). */
-	UE_API bool BreakAllEdges(TSet<UPCGNode*>* InTouchedNodes = nullptr);
+	UE_API bool BreakAllEdges(TSet<UPCGNode*>* InTouchedNodes = nullptr, TSet<UPCGPin*>* OptionalOtherPinsTouched = nullptr);
 
 	/** Break all edges the connect pins of invalid type. Optionally returns list of affected nodes (including the pin owner). */
-	UE_API bool BreakAllIncompatibleEdges(TSet<UPCGNode*>* InTouchedNodes = nullptr);
+	UE_API bool BreakAllIncompatibleEdges(TSet<UPCGNode*>* InTouchedNodes = nullptr, FString* OptionalReason = nullptr);
 
 	UFUNCTION(BlueprintCallable, Category = Settings)
 	UE_API bool IsConnected() const;
