[PCG] Static analysis fix

#rnx

#Synced-CL 48394839

diff --git a/Engine/Plugins/PCG/Source/PCG/Private/PCGPin.cpp b/Engine/Plugins/PCG/Source/PCG/Private/PCGPin.cpp
--- a/Engine/Plugins/PCG/Source/PCG/Private/PCGPin.cpp
+++ b/Engine/Plugins/PCG/Source/PCG/Private/PCGPin.cpp
@@ -309,7 +309,7 @@ bool UPCGPin::BreakAllIncompatibleEdges(TSet<UPCGNode*>* InTouchedNodes/*= nullp
 		UPCGEdge* Edge = Edges[EdgeIndex];
 		UPCGPin* OtherPin = Edge->GetOtherPin(this);
 
-		const UPCGNode* OtherNode = Cast<const UPCGNode>(OtherPin->GetOuter());
+		const UPCGNode* OtherNode = OtherPin ? Cast<const UPCGNode>(OtherPin->GetOuter()) : nullptr;
 		const FName OtherNodeName = OtherNode ? OtherNode->GetFName() : FName("Unknown");
 		const FName OtherPinName = OtherPin ? OtherPin->Properties.Label : FName("Unknown");
 
