[PCG] Added isolated actor query capabilities

#Synced-CL 48394633

diff --git a/Engine/Plugins/PCG/Source/PCG/Private/Elements/PCGActorSelector.cpp b/Engine/Plugins/PCG/Source/PCG/Private/Elements/PCGActorSelector.cpp
--- a/Engine/Plugins/PCG/Source/PCG/Private/Elements/PCGActorSelector.cpp
+++ b/Engine/Plugins/PCG/Source/PCG/Private/Elements/PCGActorSelector.cpp
@@ -3,6 +3,7 @@
 #include "Elements/PCGActorSelector.h"
 
 #include "PCGComponent.h"
+#include "PCGIsolatedActor.h"
 #include "PCGModule.h"
 #include "Grid/PCGPartitionActor.h"
 #include "Helpers/PCGActorHelpers.h"
@@ -236,6 +237,44 @@ namespace PCGActorSelector
 			}
 		}
 
+		if (Settings.IncludeIsolatedActors != EPCGIsolatedActorIncludeFlag::Ignore)
+		{
+			int ActorIndex = 0;
+			while (ActorIndex < ActorsToCheck.Num())
+			{
+				TArray<APCGIsolatedActor*> IsolatedActors = UPCGActorHelpers::GetIsolatedActorsOriginatingFrom(ActorsToCheck[ActorIndex]);
+
+				if (Settings.IncludeIsolatedActors == EPCGIsolatedActorIncludeFlag::Include)
+				{
+					++ActorIndex;
+
+					if (IsolatedActors.Num() > 0)
+					{
+						AActor** IsolatedActorsPtr = reinterpret_cast<AActor**>(IsolatedActors.GetData());
+						ActorsToCheck.Insert(IsolatedActorsPtr, IsolatedActors.Num(), ActorIndex);
+						ActorIndex += IsolatedActors.Num();
+					}
+				}
+				else if (Settings.IncludeIsolatedActors == EPCGIsolatedActorIncludeFlag::Only)
+				{
+					if (IsolatedActors.IsEmpty())
+					{
+						ActorsToCheck.RemoveAt(ActorIndex);
+					}
+					else
+					{
+						ActorsToCheck[ActorIndex++] = IsolatedActors[0];
+						if (IsolatedActors.Num() > 1)
+						{
+							AActor** IsolatedActorsPtr = reinterpret_cast<AActor**>(IsolatedActors.GetData());
+							ActorsToCheck.Insert(IsolatedActorsPtr + 1, IsolatedActors.Num() - 1, ActorIndex);
+							ActorIndex += IsolatedActors.Num() - 1;
+						}
+					}
+				}
+			}
+		}
+
 		for (AActor* Actor : ActorsToCheck)
 		{
 			// FoundActor is set by the FilteringFunction (captured)
@@ -517,7 +556,7 @@ bool FPCGSelectionKey::IsMatching(const FPCGSelectionKey& InSelectionKey, const
 			for (AActor* Current : ActorsToCheck)
 			{
 				// TempActorComponents is reset in GetComponents
-				InActor->GetComponents(UPCGComponent::StaticClass(), TempActorComponents);
+				Current->GetComponents(UPCGComponent::StaticClass(), TempActorComponents);
 				ActorComponents.Append(TempActorComponents);
 			}
 		}
diff --git a/Engine/Plugins/PCG/Source/PCG/Private/Elements/PCGDataFromActor.cpp b/Engine/Plugins/PCG/Source/PCG/Private/Elements/PCGDataFromActor.cpp
--- a/Engine/Plugins/PCG/Source/PCG/Private/Elements/PCGDataFromActor.cpp
+++ b/Engine/Plugins/PCG/Source/PCG/Private/Elements/PCGDataFromActor.cpp
@@ -4,6 +4,7 @@
 
 #include "PCGComponent.h"
 #include "PCGCustomVersion.h"
+#include "PCGIsolatedActor.h"
 #include "PCGParamData.h"
 #include "PCGSubsystem.h"
 #include "PCGTrackingManager.h"
@@ -524,6 +525,18 @@ bool FPCGDataFromActorElement::ExecuteInternal(FPCGContext* InContext) const
 				
 				FPCGDynamicTrackingHelper::AddSingleDynamicTrackingKey(Context, std::move(Key), Settings->ActorSelector.bMustOverlapSelf);
 			}
+
+			// Any isolated actors should also be dynamically tracked.
+			if (Settings->ActorSelector.IncludeIsolatedActors != EPCGIsolatedActorIncludeFlag::Ignore)
+			{
+				for (AActor* FoundActor : Context->FoundActors)
+				{
+					if (APCGIsolatedActor* IsolatedActor = Cast<APCGIsolatedActor>(FoundActor))
+					{
+						FPCGDynamicTrackingHelper::AddSingleDynamicTrackingKey(Context, FPCGSelectionKey::CreateFromObjectPtr(IsolatedActor), /*bIsCulled=*/false);
+					}
+				}
+			}
 #endif // WITH_EDITOR
 		}
 
@@ -1105,6 +1118,17 @@ void FPCGDataFromActorElement::GetDependenciesCrc(const FPCGGetDependenciesCrcPa
 					}
 				}
 			}
+
+			// Include isolated actors when relevant.
+			// @todo_pcg : is this a good idea if this isn't on the GT?
+			if (Settings->ActorSelector.IncludeIsolatedActors != EPCGIsolatedActorIncludeFlag::Ignore && ComponentToCheck)
+			{
+				TArray<APCGIsolatedActor*> IsolatedActors = UPCGActorHelpers::GetIsolatedActorsOriginatingFrom(ComponentToCheck->GetOwner());
+				for (APCGIsolatedActor* IsolatedActor : IsolatedActors)
+				{
+					Crc.Combine(GetTypeHash(IsolatedActor->GetPathName()));
+				}
+			}
 		}
 
 		const bool bDependsOnComponentData = Settings->Mode == EPCGGetDataFromActorMode::GetDataFromPCGComponent || Settings->Mode == EPCGGetDataFromActorMode::GetDataFromPCGComponentOrParseComponents;
diff --git a/Engine/Plugins/PCG/Source/PCG/Private/Helpers/PCGActorHelpers.cpp b/Engine/Plugins/PCG/Source/PCG/Private/Helpers/PCGActorHelpers.cpp
--- a/Engine/Plugins/PCG/Source/PCG/Private/Helpers/PCGActorHelpers.cpp
+++ b/Engine/Plugins/PCG/Source/PCG/Private/Helpers/PCGActorHelpers.cpp
@@ -3,6 +3,7 @@
 #include "Helpers/PCGActorHelpers.h"
 
 #include "PCGComponent.h"
+#include "PCGIsolatedActor.h"
 #include "PCGManagedResource.h"
 #include "PCGModule.h"
 #include "Elements/PCGSplineMeshParams.h"
@@ -866,3 +867,26 @@ int UPCGActorHelpers::ComputeHashFromActorTagsAndReference(const AActor* InActor
 
 	return Result;
 }
+
+TArray<APCGIsolatedActor*> UPCGActorHelpers::GetIsolatedActorsOriginatingFrom(const AActor* InActor)
+{
+	TArray<APCGIsolatedActor*> IsolatedActors;
+
+	if (InActor)
+	{
+		TSoftObjectPtr<const AActor> ActorPtr(InActor);
+
+		ForEachActorInWorld(InActor->GetWorld(), APCGIsolatedActor::StaticClass(), [&IsolatedActors, &ActorPtr](AActor* IsolatedActorCandidate)
+		{
+			APCGIsolatedActor* IsolatedActor = Cast<APCGIsolatedActor>(IsolatedActorCandidate);
+			if (IsolatedActor && IsolatedActor->Origin == ActorPtr)
+			{
+				IsolatedActors.Add(IsolatedActor);
+			}
+
+			return true;
+		});
+	}
+
+	return IsolatedActors;
+}
\ No newline at end of file
diff --git a/Engine/Plugins/PCG/Source/PCG/Public/Elements/PCGActorSelector.h b/Engine/Plugins/PCG/Source/PCG/Public/Elements/PCGActorSelector.h
--- a/Engine/Plugins/PCG/Source/PCG/Public/Elements/PCGActorSelector.h
+++ b/Engine/Plugins/PCG/Source/PCG/Public/Elements/PCGActorSelector.h
@@ -56,6 +56,17 @@ enum class EPCGActorFilter : uint8
 	FromInput
 };
 
+UENUM()
+enum class EPCGIsolatedActorIncludeFlag : uint8
+{
+	/** Ignores isolated actors that stem from the selection. */
+	Ignore,
+	/** Includes isolated actors that were extracted from the selection. */
+	Include,
+	/** Returns only isolated actors extracted from the original selection. */
+	Only,
+};
+
 /**
  * Base struct to extend FPCGSelectionKey to allow more fine grained and custom selection
  */
@@ -216,6 +227,9 @@ struct FPCGActorSelectorSettings
 	UPROPERTY(BlueprintReadWrite, EditAnywhere, Category = "Actor Selector Settings", meta = (EditCondition = "bShowIncludeChildren && ActorFilter!=EPCGActorFilter::AllWorldActors", EditConditionHides))
 	bool bIncludeChildren = false;
 
+	UPROPERTY(BlueprintReadWrite, EditAnywhere, Category = "Actor Selector Settings")
+	EPCGIsolatedActorIncludeFlag IncludeIsolatedActors = EPCGIsolatedActorIncludeFlag::Ignore;
+
 	/** Enables/disables fine-grained actor filtering options. */
 	UPROPERTY(BlueprintReadWrite, EditAnywhere, Category = "Actor Selector Settings", meta = (PCG_NotOverridable, EditCondition = "ActorFilter!=EPCGActorFilter::AllWorldActors && bIncludeChildren", EditConditionHides))
 	bool bDisableFilter = false;
diff --git a/Engine/Plugins/PCG/Source/PCG/Public/Helpers/PCGActorHelpers.h b/Engine/Plugins/PCG/Source/PCG/Public/Helpers/PCGActorHelpers.h
--- a/Engine/Plugins/PCG/Source/PCG/Public/Helpers/PCGActorHelpers.h
+++ b/Engine/Plugins/PCG/Source/PCG/Public/Helpers/PCGActorHelpers.h
@@ -20,6 +20,7 @@
 #define UE_API PCG_API
 
 class AActor;
+class APCGIsolatedActor;
 struct FPCGContext;
 class UActorComponent;
 class UDataLayerInstance;
@@ -257,6 +258,9 @@ PRAGMA_ENABLE_DEPRECATION_WARNINGS
 #if WITH_EDITOR
 	static UPackage* CreatePreviewPackage(ULevel* InLevel, const FString& InActorName);
 #endif
+
+	/** Returns all isolated actors that originate from the given actor. */
+	static TArray<APCGIsolatedActor*> GetIsolatedActorsOriginatingFrom(const AActor* InActor);
 };
 
 #undef UE_API
