[PCG] Enable debug on named reroutes

* Since named reroutes are culled at the compilation we have to not cull them in debug
* Which means we have to trigger a recompilation (unfortunate)
* Added a function to set debug flag to propagate the event accordingly.
* Fixed forward of inputs in the reroute element, missing the invisible pin.

#rnx
#Synced-CL 48394839

diff --git a/Engine/Plugins/PCG/Source/PCG/Private/Elements/PCGReroute.cpp b/Engine/Plugins/PCG/Source/PCG/Private/Elements/PCGReroute.cpp
--- a/Engine/Plugins/PCG/Source/PCG/Private/Elements/PCGReroute.cpp
+++ b/Engine/Plugins/PCG/Source/PCG/Private/Elements/PCGReroute.cpp
@@ -203,6 +203,31 @@ bool UPCGRerouteSettings::ConvertNode(const FPCGPreconfiguredInfo& ConversionInf
 	return false;
 }
 
+UPCGNamedRerouteBaseSettings::UPCGNamedRerouteBaseSettings()
+{
+#if WITH_EDITORONLY_DATA
+	// While reroutes don't support debugging, it makes sense for named reroutes to support it.
+	bDisplayDebuggingProperties = true;
+	DebugSettings.bDisplayProperties = true;
+#endif
+}
+
+#if WITH_EDITOR
+EPCGChangeType UPCGNamedRerouteBaseSettings::GetChangeTypeForProperty(const FName& InPropertyName) const
+{
+	EPCGChangeType ChangeType = Super::GetChangeTypeForProperty(InPropertyName);
+	
+	// Reroute nodes are culled at the compilation, making them not debuggable.
+	// If we want to debug we need to force the element to exist, and this means recompiling the graph.
+	if (InPropertyName == GET_MEMBER_NAME_CHECKED(UPCGSettingsInterface, bDebug))
+	{
+		ChangeType |= EPCGChangeType::Structural;
+	}
+
+	return ChangeType;
+}
+#endif // WITH_EDITOR
+
 TArray<FPCGPinProperties> UPCGNamedRerouteDeclarationSettings::OutputPinProperties() const
 {
 	TArray<FPCGPinProperties> PinProperties;
@@ -279,14 +304,33 @@ bool FPCGRerouteElement::ExecuteInternal(FPCGContext* Context) const
 {
 	check(Context);
 
-	// Reroute elements are culled during graph compilation unless they have no inbound edge.
-	// In such as case, this is a good place to log an error for user to deal with.
-	PCGE_LOG(Error, GraphAndLog, NSLOCTEXT("PCGRerouteSettings", "DetachedReroute", "Reroute is not linked to anything. Reconnect to recreate to fix the error."));
+	const UPCGSettings* Settings = Context->GetInputSettings<UPCGSettings>();
+	check(Settings);
 	
-	Context->OutputData = Context->InputData;
-	for (FPCGTaggedData& Output : Context->OutputData.TaggedData)
+	// Reroute elements are culled during graph compilation unless they have no inbound edge or in debug.
+	// If they have no inbound edge, this is a good place to log an error for user to deal with.
+#if WITH_EDITOR
+	if (!Settings->bDebug)
+#endif
+	{
+		PCGE_LOG(Error, GraphAndLog, NSLOCTEXT("PCGRerouteSettings", "DetachedReroute", "Reroute is not linked to anything. Reconnect to recreate to fix the error."));
+	}
+
+	// Need to forward the output to the output pin, but also to the invisible pin in case of named reroute declaration node.
+	TArray<FName, TInlineAllocator<2>> PinsToForwardTo = {PCGPinConstants::DefaultOutputLabel};
+	if (Settings->IsA<UPCGNamedRerouteDeclarationSettings>())
 	{
-		Output.Pin = PCGPinConstants::DefaultOutputLabel;
+		PinsToForwardTo.Add(PCGNamedRerouteConstants::InvisiblePinLabel);
+	}
+
+	Context->OutputData.TaggedData.Reserve(Context->OutputData.TaggedData.Num() + PinsToForwardTo.Num() * Context->InputData.TaggedData.Num());
+
+	for (const FName PinToForwardTo : PinsToForwardTo)
+	{
+		for (const FPCGTaggedData& Input : Context->InputData.TaggedData)
+		{
+			Context->OutputData.TaggedData.Emplace_GetRef(Input).Pin = PinToForwardTo;
+		}
 	}
 
 	return true;
diff --git a/Engine/Plugins/PCG/Source/PCG/Private/Graph/PCGGraphCompiler.cpp b/Engine/Plugins/PCG/Source/PCG/Private/Graph/PCGGraphCompiler.cpp
--- a/Engine/Plugins/PCG/Source/PCG/Private/Graph/PCGGraphCompiler.cpp
+++ b/Engine/Plugins/PCG/Source/PCG/Private/Graph/PCGGraphCompiler.cpp
@@ -1366,7 +1366,20 @@ void FPCGGraphCompiler::CompileTopGraph(UPCGGraph* InGraph, uint32 GenerationGri
 		TRACE_CPUPROFILER_EVENT_SCOPE(TaskStaticCulling);
 
 		// Remove reroute nodes before execution grid setup, as grid linkages need final nodes to connect from/to.
-		CullTasks(CompiledTasks, /*bAddPassthroughWires=*/true, [](const FPCGGraphTask& InTask) { return InTask.Node && Cast<UPCGRerouteSettings>(InTask.Node->GetSettings()) && InTask.Node->HasInboundEdges(); });
+		// Note that reroutes nodes won't be culled if they have no inbound edges or are marked for debug.
+		auto RerouteCullTasks = [](const FPCGGraphTask& InTask) -> bool
+		{
+			if (const UPCGRerouteSettings* Settings = InTask.Node ? Cast<UPCGRerouteSettings>(InTask.Node->GetSettings()) : nullptr)
+			{
+				return InTask.Node->HasInboundEdges() && !Settings->bDebug;
+			}
+			else
+			{
+				return false;
+			}
+		};
+
+		CullTasks(CompiledTasks, /*bAddPassthroughWires=*/true, RerouteCullTasks);
 
 		// Cull inactive branches downstream of branch nodes with static selection values.
 		CullTasksStaticInactive(CompiledTasks);
diff --git a/Engine/Plugins/PCG/Source/PCG/Private/PCGSettings.cpp b/Engine/Plugins/PCG/Source/PCG/Private/PCGSettings.cpp
--- a/Engine/Plugins/PCG/Source/PCG/Private/PCGSettings.cpp
+++ b/Engine/Plugins/PCG/Source/PCG/Private/PCGSettings.cpp
@@ -215,6 +215,27 @@ void UPCGSettingsInterface::SetEnabled(bool bInEnabled)
 	}
 }
 
+bool UPCGSettingsInterface::SetDebugged(bool bInDebugged)
+{
+	if (bDebug == bInDebugged)
+	{
+		return false;
+	}
+
+	Modify(/*bAlwaysMarkDirty=*/false);
+	bDebug = bInDebugged;
+
+#if WITH_EDITOR
+	if (UPCGSettings* Settings = GetSettings())
+	{
+		const EPCGChangeType ChangeType = Settings->GetChangeTypeForProperty(GET_MEMBER_NAME_CHECKED(UPCGSettingsInterface, bDebug));
+		OnSettingsChangedDelegate.Broadcast(Settings, ChangeType);
+	}
+#endif
+
+	return true;
+}
+
 uint32 UPCGSettings::GetTypeNameHash() const
 {
 	return GetTypeHash(GetClass()->GetFName());
diff --git a/Engine/Plugins/PCG/Source/PCG/Public/Elements/PCGReroute.h b/Engine/Plugins/PCG/Source/PCG/Public/Elements/PCGReroute.h
--- a/Engine/Plugins/PCG/Source/PCG/Public/Elements/PCGReroute.h
+++ b/Engine/Plugins/PCG/Source/PCG/Public/Elements/PCGReroute.h
@@ -58,8 +58,11 @@ class UPCGNamedRerouteBaseSettings : public UPCGRerouteSettings
 	GENERATED_BODY()
 
 public:
+	UPCGNamedRerouteBaseSettings();
+
 #if WITH_EDITOR
 	virtual bool CanUserEditTitle() const override { return true; }
+	virtual EPCGChangeType GetChangeTypeForProperty(const FName& InPropertyName) const override;
 #endif
 
 };
diff --git a/Engine/Plugins/PCG/Source/PCG/Public/PCGSettings.h b/Engine/Plugins/PCG/Source/PCG/Public/PCGSettings.h
--- a/Engine/Plugins/PCG/Source/PCG/Public/PCGSettings.h
+++ b/Engine/Plugins/PCG/Source/PCG/Public/PCGSettings.h
@@ -179,7 +179,14 @@ public:
 	/** Dedicated method to change enable state because some nodes have more complex behavior on enable/disable (such as subgraphs) */
 	UE_API void SetEnabled(bool bInEnabled);
 
+	/**
+	 * Dedicated method to change debug state because some nodes have more complex behavior on enable/disable (such as reroutes).
+	 * Returns true if the flag actually changed.
+	 */
+	UE_API bool SetDebugged(bool bInDebugged);
+
 	/** Whether this element can be disabled. */
+	UFUNCTION()
 	virtual bool CanBeDisabled() const { return true; }
 
 	/** Whether this element supports Debug and Inspect features. */
@@ -192,7 +199,7 @@ public:
 	bool bIsInspecting = false;
 #endif
 
-	UPROPERTY(EditAnywhere, BlueprintReadWrite, Category = Debug, meta = (EditCondition = bDisplayDebuggingProperties, EditConditionHides, HideEditConditionToggle))
+	UPROPERTY(EditAnywhere, BlueprintReadWrite, Category = Debug, meta = (EditCondition = "bDisplayDebuggingProperties && CanBeDisabled()", EditConditionHides, HideEditConditionToggle))
 	bool bEnabled = true;
 
 	UPROPERTY(Transient, EditAnywhere, BlueprintReadWrite, Category = Debug, meta = (EditCondition = bDisplayDebuggingProperties, EditConditionHides, HideEditConditionToggle))
diff --git a/Engine/Plugins/PCG/Source/PCGEditor/Private/PCGEditor.cpp b/Engine/Plugins/PCG/Source/PCGEditor/Private/PCGEditor.cpp
--- a/Engine/Plugins/PCG/Source/PCGEditor/Private/PCGEditor.cpp
+++ b/Engine/Plugins/PCG/Source/PCGEditor/Private/PCGEditor.cpp
@@ -2292,13 +2292,7 @@ void FPCGEditor::OnToggleDebug()
 				continue;
 			}
 
-			if (PCGSettingsInterface->bDebug != bNewCheckState)
-			{
-				PCGSettingsInterface->Modify(/*bAlwaysMarkDirty=*/false);
-				PCGSettingsInterface->bDebug = bNewCheckState;
-				PCGNode->OnNodeChangedDelegate.Broadcast(PCGNode, EPCGChangeType::Settings);
-				bChanged = true;
-			}
+			bChanged |= PCGSettingsInterface->SetDebugged(bNewCheckState);
 		}
 
 		if (!bChanged)
@@ -2382,13 +2376,7 @@ void FPCGEditor::OnDebugOnlySelected()
 			// Selected set to target state, non-selected should not be debugged.
 			const bool bShouldBeDebug = SelectedNodes.Contains(PCGEditorGraphNode) ? bTargetDebugState : false;
 
-			if (PCGSettingsInterface->bDebug != bShouldBeDebug)
-			{
-				PCGSettingsInterface->Modify(/*bAlwaysMarkDirty=*/false);
-				PCGSettingsInterface->bDebug = bShouldBeDebug;
-				PCGNode->OnNodeChangedDelegate.Broadcast(PCGNode, EPCGChangeType::Settings);
-				bChanged = true;
-			}
+			bChanged |= PCGSettingsInterface->SetDebugged(bShouldBeDebug);
 		}
 
 		if (!bChanged)
@@ -2415,13 +2403,7 @@ void FPCGEditor::OnDisableDebugOnAllNodes()
 				continue;
 			}
 
-			if (PCGSettingsInterface->bDebug)
-			{
-				PCGSettingsInterface->Modify(/*bAlwaysMarkDirty=*/false);
-				PCGSettingsInterface->bDebug = false;
-				PCGNode->OnNodeChangedDelegate.Broadcast(PCGNode, EPCGChangeType::Settings);
-				bChanged = true;
-			}
+			bChanged |= PCGSettingsInterface->SetDebugged(false);
 		}
 
 		if (!bChanged)
