[PCG] Add possibility for nodes creating artifact to add tags.
- Deprecate the use of CSV for tags when it was used.

#rnx
#Synced-CL 48392588

diff --git a/Engine/Plugins/PCG/Source/PCG/Private/Elements/PCGApplyOnActor.cpp b/Engine/Plugins/PCG/Source/PCG/Private/Elements/PCGApplyOnActor.cpp
--- a/Engine/Plugins/PCG/Source/PCG/Private/Elements/PCGApplyOnActor.cpp
+++ b/Engine/Plugins/PCG/Source/PCG/Private/Elements/PCGApplyOnActor.cpp
@@ -194,6 +194,17 @@ bool FPCGApplyOnObjectElement::ExecuteInternal(FPCGContext* Context) const
 			PCGObjectPropertyOverrideHelpers::ApplyOverrides(Params);
 		}
 
+		if (!TargetObjectsAndIndices.IsEmpty())
+		{
+			if (UActorComponent* ActorComponent = Cast<UActorComponent>(TargetObjectsAndIndices[0].Key))
+			{
+				for (const FName& Tag : Settings->TagsToAddOnActors)
+				{
+					ActorComponent->ComponentTags.AddUnique(Tag);
+				}
+			}
+		}
+
 		for (const TPair<UObject*, int32>& TargetObjectAndIndex : TargetObjectsAndIndices)
 		{
 			UObject* TargetObject = TargetObjectAndIndex.Key;
diff --git a/Engine/Plugins/PCG/Source/PCG/Private/Elements/PCGCreateSpline.cpp b/Engine/Plugins/PCG/Source/PCG/Private/Elements/PCGCreateSpline.cpp
--- a/Engine/Plugins/PCG/Source/PCG/Private/Elements/PCGCreateSpline.cpp
+++ b/Engine/Plugins/PCG/Source/PCG/Private/Elements/PCGCreateSpline.cpp
@@ -194,6 +194,10 @@ bool FPCGCreateSplineElement::ExecuteInternal(FPCGContext* Context) const
 				SplineComponent = NewObject<USplineComponent>(SplineActor);
 				SplineComponent->ComponentTags.Add(SourceComponent->GetFName());
 				SplineComponent->ComponentTags.Add(PCGHelpers::DefaultPCGTag);
+				for (const FName& Tag : Settings->TagsToAddOnComponents)
+				{
+					SplineComponent->ComponentTags.AddUnique(Tag);
+				}
 
 				SplineComponent->RegisterComponent();
 				SplineActor->AddInstanceComponent(SplineComponent);
diff --git a/Engine/Plugins/PCG/Source/PCG/Private/Elements/PCGCreateTargetActor.cpp b/Engine/Plugins/PCG/Source/PCG/Private/Elements/PCGCreateTargetActor.cpp
--- a/Engine/Plugins/PCG/Source/PCG/Private/Elements/PCGCreateTargetActor.cpp
+++ b/Engine/Plugins/PCG/Source/PCG/Private/Elements/PCGCreateTargetActor.cpp
@@ -53,6 +53,11 @@ namespace PCGCreateTargetActor
 		true,
 		TEXT("Controls whether PCG Create Target Actor node can reuse actors when re-executing (requires Create Target Actor node resave so they have a Stable Reuse Guid)"));
 
+	static TAutoConsoleVariable<bool> CVarEnableWarningOnDeprecatedCSVActorTags(
+		TEXT("pcg.Graph.EnableWarningOnDeprecatedCSVActorTags"),
+		true,
+		TEXT("Output warnings when using dreprecated CSV actor tags inputs."));
+
 	FPCGCrc GetAdditionalDependenciesCrc(const UPCGActorHelpers::FSpawnDefaultActorParams& InParams, AActor* TargetActor)
 	{
 		FArchiveCrc32 Ar;
@@ -484,8 +489,23 @@ bool FPCGCreateTargetActorElement::ExecuteInternal(FPCGContext* Context) const
 
 		GeneratedActor->Tags.Add(PCGHelpers::DefaultPCGActorTag);
 
-		const TArray<FString> AdditionalTags = PCGHelpers::GetStringArrayFromCommaSeparatedList(Settings->CommaSeparatedActorTags);
-		GeneratedActor->Tags.Append(AdditionalTags);
+		if (!Settings->CommaSeparatedActorTags.IsEmpty())
+		{
+			if (PCGCreateTargetActor::CVarEnableWarningOnDeprecatedCSVActorTags.GetValueOnAnyThread())
+			{
+				PCGE_LOG(Warning, GraphAndLog, LOCTEXT("Deprecated Property", "\"CommaSeparatedActorTags\" is now deprecated, please use \"TagsToAddOnActor\"."));
+			}
+			const TArray<FString> AdditionalTags = PCGHelpers::GetStringArrayFromCommaSeparatedList(Settings->CommaSeparatedActorTags);
+			for (const FString& Tag : AdditionalTags)
+			{
+				GeneratedActor->Tags.AddUnique(FName(Tag));
+			}
+		}
+
+		for (const FName& Tag : Settings->TagsToAddOnActor)
+		{
+			GeneratedActor->Tags.AddUnique(Tag);
+		}
 
 #if WITH_EDITOR
 		if (!Settings->ActorLabel.IsEmpty())
diff --git a/Engine/Plugins/PCG/Source/PCG/Private/Elements/PCGSpawnSpline.cpp b/Engine/Plugins/PCG/Source/PCG/Private/Elements/PCGSpawnSpline.cpp
--- a/Engine/Plugins/PCG/Source/PCG/Private/Elements/PCGSpawnSpline.cpp
+++ b/Engine/Plugins/PCG/Source/PCG/Private/Elements/PCGSpawnSpline.cpp
@@ -108,6 +108,10 @@ bool FPCGSpawnSplineElement::ExecuteInternal(FPCGContext* InContext) const
 		USplineComponent* SplineComponent = NewObject<USplineComponent>(TargetActor, *SplineComponentToSpawn);
 		SplineComponent->ComponentTags.Add(SourceComponent->GetFName());
 		SplineComponent->ComponentTags.Add(PCGHelpers::DefaultPCGTag);
+		for (const FName& Tag : Settings->TagsToAddOnComponents)
+		{
+			SplineComponent->ComponentTags.AddUnique(Tag);
+		}
 
 		// Also propagate the tags on the tagged data to the component.
 		Algo::Transform(InputData.Tags, SplineComponent->ComponentTags, [](const FString& Tag) { return FName(Tag); });
diff --git a/Engine/Plugins/PCG/Source/PCG/Private/Elements/PCGSpawnSplineMesh.cpp b/Engine/Plugins/PCG/Source/PCG/Private/Elements/PCGSpawnSplineMesh.cpp
--- a/Engine/Plugins/PCG/Source/PCG/Private/Elements/PCGSpawnSplineMesh.cpp
+++ b/Engine/Plugins/PCG/Source/PCG/Private/Elements/PCGSpawnSplineMesh.cpp
@@ -306,17 +306,24 @@ bool FPCGSpawnSplineMeshElement::ExecuteInternal(FPCGContext* InContext) const
 			IterState.SMCBuilderParams.SettingsCrc = Settings->GetSettingsCrc();
 			ensure(IterState.SMCBuilderParams.SettingsCrc.IsValid());
 
-			USplineMeshComponent* SplineMeshComponent = UPCGActorHelpers::GetOrCreateSplineMeshComponent(ExecState.TargetActor, SourceComponent, IterState.SMCBuilderParams, Context);
-
-			// Apply the overrides on the spawned component
-			if (IterState.ComponentOverrides.IsValid())
+			if (USplineMeshComponent* SplineMeshComponent = UPCGActorHelpers::GetOrCreateSplineMeshComponent(ExecState.TargetActor, SourceComponent, IterState.SMCBuilderParams, Context))
 			{
-				IterState.ComponentOverrides.UpdateTemplateObject(SplineMeshComponent);
-				IterState.ComponentOverrides.Apply(IterState.ElementIndex);
-			}
+				// Apply the overrides on the spawned component
+				if (IterState.ComponentOverrides.IsValid())
+				{
+					IterState.ComponentOverrides.UpdateTemplateObject(SplineMeshComponent);
+					IterState.ComponentOverrides.Apply(IterState.ElementIndex);
+				}
 
-			// TODO: Write out the geometry to a dynamic mesh type.
-			//SplineMeshComponent->BodySetup->TriMeshGeometries
+				// Apply tags on the Spline Mesh Component
+				for (const FName& Tag : Settings->TagsToAddOnComponents)
+				{
+					SplineMeshComponent->ComponentTags.AddUnique(Tag);
+				}
+
+				// TODO: Write out the geometry to a dynamic mesh type.
+				//SplineMeshComponent->BodySetup->TriMeshGeometries
+			}
 
 			IterState.ElementIndex++;
 
diff --git a/Engine/Plugins/PCG/Source/PCG/Public/Elements/PCGApplyOnActor.h b/Engine/Plugins/PCG/Source/PCG/Public/Elements/PCGApplyOnActor.h
--- a/Engine/Plugins/PCG/Source/PCG/Public/Elements/PCGApplyOnActor.h
+++ b/Engine/Plugins/PCG/Source/PCG/Public/Elements/PCGApplyOnActor.h
@@ -47,6 +47,9 @@ public:
 	TSoftObjectPtr<AActor> TargetActor;
 #endif
 
+	UPROPERTY(BlueprintReadWrite, EditAnywhere, Category = Settings, meta = (PCG_Overridable, EditCondition = "Mode!= EPCGCreateSplineMode::CreateDataOnly", EditConditionHides))
+	TArray<FName> TagsToAddOnActors;
+
 	/** Override the default property values on the target actor. Applied before post-process functions. */
 	UPROPERTY(EditAnywhere, BlueprintReadWrite, Category = Settings)
 	TArray<FPCGObjectPropertyOverrideDescription> PropertyOverrideDescriptions;
diff --git a/Engine/Plugins/PCG/Source/PCG/Public/Elements/PCGCreateSpline.h b/Engine/Plugins/PCG/Source/PCG/Public/Elements/PCGCreateSpline.h
--- a/Engine/Plugins/PCG/Source/PCG/Public/Elements/PCGCreateSpline.h
+++ b/Engine/Plugins/PCG/Source/PCG/Public/Elements/PCGCreateSpline.h
@@ -63,6 +63,9 @@ public:
 	UPROPERTY(BlueprintReadWrite, EditAnywhere, Category = Settings, meta = (PCG_Overridable, EditCondition = "!bLinear && bApplyCustomTangents"))
 	FName LeaveTangentAttribute;
 
+	UPROPERTY(BlueprintReadWrite, EditAnywhere, Category = Settings, meta = (PCG_Overridable, EditCondition = "Mode!= EPCGCreateSplineMode::CreateDataOnly", EditConditionHides))
+	TArray<FName> TagsToAddOnComponents;
+
 	UPROPERTY(meta = (PCG_Overridable))
 	TSoftObjectPtr<AActor> TargetActor;
 
diff --git a/Engine/Plugins/PCG/Source/PCG/Public/Elements/PCGCreateTargetActor.h b/Engine/Plugins/PCG/Source/PCG/Public/Elements/PCGCreateTargetActor.h
--- a/Engine/Plugins/PCG/Source/PCG/Public/Elements/PCGCreateTargetActor.h
+++ b/Engine/Plugins/PCG/Source/PCG/Public/Elements/PCGCreateTargetActor.h
@@ -76,9 +76,12 @@ public:
 	FTransform ActorPivot;
 
 	/** Comma-separated list of tags to add to the newly created target actor. */
-	UPROPERTY(EditAnywhere, BlueprintReadWrite, Category = Settings, meta = (PCG_Overridable))
+	UPROPERTY(EditAnywhere, BlueprintReadWrite, Category = Settings, DisplayName = "Comma Separated Actor Tags (Deprecated)", meta = (PCG_Overridable))
 	FString CommaSeparatedActorTags;
 
+	UPROPERTY(BlueprintReadWrite, EditAnywhere, Category = Settings, meta = (PCG_Overridable))
+	TArray<FName> TagsToAddOnActor;
+
 	/** Override the default property values on the created target actor. Applied before post-process functions. */
 	UPROPERTY(EditAnywhere, BlueprintReadWrite, Category = Settings)
 	TArray<FPCGObjectPropertyOverrideDescription> PropertyOverrideDescriptions;
diff --git a/Engine/Plugins/PCG/Source/PCG/Public/Elements/PCGSpawnSpline.h b/Engine/Plugins/PCG/Source/PCG/Public/Elements/PCGSpawnSpline.h
--- a/Engine/Plugins/PCG/Source/PCG/Public/Elements/PCGSpawnSpline.h
+++ b/Engine/Plugins/PCG/Source/PCG/Public/Elements/PCGSpawnSpline.h
@@ -47,6 +47,9 @@ public:
 	UPROPERTY(meta = (PCG_Overridable))
 	TSoftObjectPtr<AActor> TargetActor;
 
+	UPROPERTY(BlueprintReadWrite, EditAnywhere, Category = Settings, meta = (PCG_Overridable))
+	TArray<FName> TagsToAddOnComponents;
+
 	/** Specify a list of functions to be called on the target actor after spline creation. Functions need to be parameter-less and with "CallInEditor" flag enabled. */
 	UPROPERTY(BlueprintReadWrite, EditAnywhere, Category = Settings)
 	TArray<FName> PostProcessFunctionNames;
diff --git a/Engine/Plugins/PCG/Source/PCG/Public/Elements/PCGSpawnSplineMesh.h b/Engine/Plugins/PCG/Source/PCG/Public/Elements/PCGSpawnSplineMesh.h
--- a/Engine/Plugins/PCG/Source/PCG/Public/Elements/PCGSpawnSplineMesh.h
+++ b/Engine/Plugins/PCG/Source/PCG/Public/Elements/PCGSpawnSplineMesh.h
@@ -56,6 +56,9 @@ public:
 	UPROPERTY(meta = (PCG_Overridable))
 	TSoftObjectPtr<AActor> TargetActor;
 
+	UPROPERTY(BlueprintReadWrite, EditAnywhere, Category = Settings, meta = (PCG_Overridable))
+	TArray<FName> TagsToAddOnComponents;
+
 	/** Specify a list of functions to be called on the target actor after spline mesh creation. Functions need to be parameter-less and with "CallInEditor" flag enabled. */
 	UPROPERTY(BlueprintReadWrite, EditAnywhere, Category = Settings)
 	TArray<FName> PostProcessFunctionNames;
