[PCG] Manual merge for isolate command

#rnx

#Synced-CL 48405364

diff --git a/Engine/Plugins/PCG/Source/PCGEditor/Private/EditorMode/PCGEdMode.cpp b/Engine/Plugins/PCG/Source/PCGEditor/Private/EditorMode/PCGEdMode.cpp
--- a/Engine/Plugins/PCG/Source/PCGEditor/Private/EditorMode/PCGEdMode.cpp
+++ b/Engine/Plugins/PCG/Source/PCGEditor/Private/EditorMode/PCGEdMode.cpp
@@ -2,6 +2,7 @@
 
 #include "EditorMode/PCGEdMode.h"
 
+#include "PCGEditorModule.h"
 #include "EditorMode/PCGEdModeCommands.h"
 #include "EditorMode/PCGEdModeSettings.h"
 #include "EditorMode/PCGEdModeStyle.h"
@@ -88,7 +89,45 @@ namespace PCGEditorMode
 				{
 					ToolSettings->SetToolGraph(Graph);
 				}
+				if (Args.Num() > 1)
+				{
+					ToolSettings->SetActorLabel(FName(Args[1]));
+				}
+			}
+		}));
+
+	static FAutoConsoleCommand CommandIsolate(
+		TEXT("pcg.tool.Isolate"),
+		TEXT("Isolates (like clear PCG link) content from an actor with a PCG component into a separate actor."),
+		FConsoleCommandWithArgsDelegate::CreateLambda([](const TArray<FString>& Args)
+		{
+			if(Args.IsEmpty())
+			{
+				UE_LOG(LogPCGEditor, Error, TEXT("Invalid arguments to the pcg.tool.Isolate command."));
+				return;
+			}
+
+			// [0] - selected actor
+			FSoftObjectPath ActorPath(Args[0]);
+			UObject* ActorObject = ActorPath.TryLoad();
+
+			AActor* Actor = Cast<AActor>(ActorObject);
+
+			if (!Actor)
+			{
+				UE_LOG(LogPCGEditor, Error, TEXT("Unable to retrieve the actor '{0}' in the pcg.tool.Isolate command."), *Args[0]);
+				return;
+			}
+
+			// [1+] - tags conjunction
+			TArray<FName> Tags;
+
+			for (int ArgIndex = 1; ArgIndex < Args.Num(); ++ArgIndex)
+			{
+				Tags.Add(FName(Args[ArgIndex]));
 			}
+			
+			PCGIsolateTool::Isolate(Actor, Tags);
 		}));
 }
 
