[PCG] Always show the real name of a named reroutes

* Some node titles, like with underscores, were replaced by spaces for display.
* Two nodes, one with a space and one with an underscore would be visually identical
* This also makes duplication working weirdly, as the duplication would produce the same name visually (but the underscore was replaced by a space)

#rnx
#Synced-CL 48452988

diff --git a/Engine/Plugins/PCG/Source/PCGEditor/Private/Nodes/PCGEditorGraphNodeReroute.cpp b/Engine/Plugins/PCG/Source/PCGEditor/Private/Nodes/PCGEditorGraphNodeReroute.cpp
--- a/Engine/Plugins/PCG/Source/PCGEditor/Private/Nodes/PCGEditorGraphNodeReroute.cpp
+++ b/Engine/Plugins/PCG/Source/PCGEditor/Private/Nodes/PCGEditorGraphNodeReroute.cpp
@@ -120,8 +120,8 @@ UEdGraphPin* UPCGEditorGraphNodeReroute::GetOutputPin() const
 
 FText UPCGEditorGraphNodeNamedRerouteBase::GetNodeTitle(ENodeTitleType::Type TitleType) const
 {
-	// Never show full title, only list view
-	return Super::GetNodeTitle(ENodeTitleType::ListView);
+	// Never reformat the node name, as we don't want to have nodes that seems the same but ar not (example, "As_A" and "As A" would appear the same)
+	return PCGNode ? FText::FromName(PCGNode->NodeTitle) : FText{};
 }
 
 TSharedPtr<SGraphNode> UPCGEditorGraphNodeNamedRerouteBase::CreateVisualWidget()
