[PCG] QoL improvements for BP function in Base Point Data

* Add "Number" as keyword for search for "Get/Set Num Points"
* Change category from SpatialData to PointData

#rnx
#Synced-CL 48452988

diff --git a/Engine/Plugins/PCG/Source/PCG/Public/Data/PCGBasePointData.h b/Engine/Plugins/PCG/Source/PCG/Public/Data/PCGBasePointData.h
--- a/Engine/Plugins/PCG/Source/PCG/Public/Data/PCGBasePointData.h
+++ b/Engine/Plugins/PCG/Source/PCG/Public/Data/PCGBasePointData.h
@@ -104,16 +104,16 @@ public:
 
 	virtual bool IsValidRef(const PCGPointOctree::FPointRef& InPointRef) const PURE_VIRTUAL(UPCGBasePointData::IsValidRef, return false;);
 
-	UFUNCTION(BlueprintCallable, Category = SpatialData)
+	UFUNCTION(BlueprintCallable, Category = PointData, meta = (Keywords="Number"))
 	virtual int32 GetNumPoints() const PURE_VIRTUAL(UPCGBasePointData::GetNumPoints, return 0;);
 
-	UFUNCTION(BlueprintCallable, Category = SpatialData)
+	UFUNCTION(BlueprintCallable, Category = PointData)
 	bool IsEmpty() const { return GetNumPoints() == 0; }
 
-	UFUNCTION(BlueprintCallable, Category = SpatialData)
+	UFUNCTION(BlueprintCallable, Category = PointData, meta = (Keywords="Number"))
 	virtual void SetNumPoints(int32 InNumPoints, bool bInitializeValues = true) PURE_VIRTUAL(UPCGBasePointData::SetNumPoints, return;);
 
-	UFUNCTION(BlueprintCallable, Category = SpatialData, meta=(DisplayName="Allocate Properties"))
+	UFUNCTION(BlueprintCallable, Category = PointData, meta=(DisplayName="Allocate Properties"))
 	void BP_AllocateProperties(UPARAM(meta = (BitMask, BitmaskEnum = "/Script/PCG.EPCGPointNativeProperties")) int32 Properties) { AllocateProperties((EPCGPointNativeProperties)Properties); }
 
 	virtual void AllocateProperties(EPCGPointNativeProperties Properties) {}
@@ -189,82 +189,82 @@ public:
 		}
 	}
 
-	UFUNCTION(BlueprintCallable, Category = SpatialData, meta = (DisplayName = "Set All Transforms", ToolTip = "Set all transforms to same value"))
+	UFUNCTION(BlueprintCallable, Category = PointData, meta = (DisplayName = "Set All Transforms", ToolTip = "Set all transforms to same value"))
 	UE_API virtual void SetTransform(const FTransform& InTransform);
 	
-	UFUNCTION(BlueprintCallable, Category = SpatialData, meta = (DisplayName = "Set All Densities", ToolTip = "Set all densities to same value"))
+	UFUNCTION(BlueprintCallable, Category = PointData, meta = (DisplayName = "Set All Densities", ToolTip = "Set all densities to same value"))
 	UE_API virtual void SetDensity(float InDensity);
 	
-	UFUNCTION(BlueprintCallable, Category = SpatialData, meta = (DisplayName = "Set All Bounds Min", ToolTip = "Set all bounds min to same value"))
+	UFUNCTION(BlueprintCallable, Category = PointData, meta = (DisplayName = "Set All Bounds Min", ToolTip = "Set all bounds min to same value"))
 	UE_API virtual void SetBoundsMin(const FVector& InBoundsMin);
 	
-	UFUNCTION(BlueprintCallable, Category = SpatialData, meta = (DisplayName = "Set All Bounds Max", ToolTip = "Set all bounds max to same value"))
+	UFUNCTION(BlueprintCallable, Category = PointData, meta = (DisplayName = "Set All Bounds Max", ToolTip = "Set all bounds max to same value"))
 	UE_API virtual void SetBoundsMax(const FVector& InBoundsMax);
 	
-	UFUNCTION(BlueprintCallable, Category = SpatialData, meta = (DisplayName = "Set All Colors", ToolTip = "Set all colors to same value"))
+	UFUNCTION(BlueprintCallable, Category = PointData, meta = (DisplayName = "Set All Colors", ToolTip = "Set all colors to same value"))
 	UE_API virtual void SetColor(const FVector4& InColor);
 	
-	UFUNCTION(BlueprintCallable, Category = SpatialData, meta = (DisplayName = "Set All Steepness", ToolTip = "Set all steepness to same value"))
+	UFUNCTION(BlueprintCallable, Category = PointData, meta = (DisplayName = "Set All Steepness", ToolTip = "Set all steepness to same value"))
 	UE_API virtual void SetSteepness(float InSteepness);
 	
-	UFUNCTION(BlueprintCallable, Category = SpatialData, meta = (DisplayName = "Set All Seeds", ToolTip = "Set all seeds to same value"))
+	UFUNCTION(BlueprintCallable, Category = PointData, meta = (DisplayName = "Set All Seeds", ToolTip = "Set all seeds to same value"))
 	UE_API virtual void SetSeed(int32 InSeed);
 	
-	UFUNCTION(BlueprintCallable, Category = SpatialData, meta = (DisplayName = "Set All Metadata Entries", ToolTip = "Set all metadata entries to same value"))
+	UFUNCTION(BlueprintCallable, Category = PointData, meta = (DisplayName = "Set All Metadata Entries", ToolTip = "Set all metadata entries to same value"))
 	UE_API virtual void SetMetadataEntry(int64 InMetadataEntry);
 
-	UFUNCTION(BlueprintCallable, Category = SpatialData, meta = (DisplayName = "Set All Extents", ToolTip = "Set all extents to same value"))
+	UFUNCTION(BlueprintCallable, Category = PointData, meta = (DisplayName = "Set All Extents", ToolTip = "Set all extents to same value"))
 	UE_API virtual void SetExtents(const FVector& InExtents);
 	
-	UFUNCTION(BlueprintCallable, Category = SpatialData, meta = (DisplayName = "Set All Local Centers", ToolTip = "Set all local centers to same value"))
+	UFUNCTION(BlueprintCallable, Category = PointData, meta = (DisplayName = "Set All Local Centers", ToolTip = "Set all local centers to same value"))
 	UE_API virtual void SetLocalCenter(const FVector& InLocalCenter);
 
-	UFUNCTION(BlueprintCallable, Category = SpatialData)
+	UFUNCTION(BlueprintCallable, Category = PointData)
 	UE_API virtual const FTransform& GetTransform(int32 InPointIndex) const;
 	
-	UFUNCTION(BlueprintCallable, Category = SpatialData)
+	UFUNCTION(BlueprintCallable, Category = PointData)
 	UE_API virtual const float& GetDensity(int32 InPointIndex) const;
 	
-	UFUNCTION(BlueprintCallable, Category = SpatialData)
+	UFUNCTION(BlueprintCallable, Category = PointData)
 	UE_API virtual const FVector& GetBoundsMin(int32 InPointIndex) const;
 	
-	UFUNCTION(BlueprintCallable, Category = SpatialData)
+	UFUNCTION(BlueprintCallable, Category = PointData)
 	UE_API virtual const FVector& GetBoundsMax(int32 InPointIndex) const;
 	
-	UFUNCTION(BlueprintCallable, Category = SpatialData)
+	UFUNCTION(BlueprintCallable, Category = PointData)
 	UE_API virtual const FVector4& GetColor(int32 InPointIndex) const;
 	
-	UFUNCTION(BlueprintCallable, Category = SpatialData)
+	UFUNCTION(BlueprintCallable, Category = PointData)
 	UE_API virtual const float& GetSteepness(int32 InPointIndex) const;
 
-	UFUNCTION(BlueprintCallable, Category = SpatialData)
+	UFUNCTION(BlueprintCallable, Category = PointData)
 	UE_API virtual const int32& GetSeed(int32 InPointIndex) const;
 
-	UFUNCTION(BlueprintCallable, Category = SpatialData)
+	UFUNCTION(BlueprintCallable, Category = PointData)
 	UE_API virtual const int64& GetMetadataEntry(int32 InPointIndex) const;
 
-	UFUNCTION(BlueprintCallable, Category = SpatialData)
+	UFUNCTION(BlueprintCallable, Category = PointData)
 	UE_API virtual FBoxSphereBounds GetDensityBounds(int32 InPointIndex) const;
 
-	UFUNCTION(BlueprintCallable, Category = SpatialData)
+	UFUNCTION(BlueprintCallable, Category = PointData)
 	UE_API virtual FBox GetLocalDensityBounds(int32 InPointIndex) const;
 
-	UFUNCTION(BlueprintCallable, Category = SpatialData)
+	UFUNCTION(BlueprintCallable, Category = PointData)
 	UE_API virtual FBox GetLocalBounds(int32 InPointIndex) const;
 
-	UFUNCTION(BlueprintCallable, Category = SpatialData)
+	UFUNCTION(BlueprintCallable, Category = PointData)
 	UE_API virtual FVector GetLocalCenter(int32 InPointIndex) const;
 
-	UFUNCTION(BlueprintCallable, Category = SpatialData)
+	UFUNCTION(BlueprintCallable, Category = PointData)
 	UE_API virtual FVector GetExtents(int32 InPointIndex) const;
 
-	UFUNCTION(BlueprintCallable, Category = SpatialData)
+	UFUNCTION(BlueprintCallable, Category = PointData)
 	UE_API virtual FVector GetScaledExtents(int32 InPointIndex) const;
 
-	UFUNCTION(BlueprintCallable, Category = SpatialData)
+	UFUNCTION(BlueprintCallable, Category = PointData)
 	UE_API virtual FVector GetLocalSize(int32 InPointIndex) const;
 
-	UFUNCTION(BlueprintCallable, Category = SpatialData)
+	UFUNCTION(BlueprintCallable, Category = PointData)
 	UE_API virtual FVector GetScaledLocalSize(int32 InPointIndex) const;
 		
 	/** Get the dirty status of the Octree. Note that the Point Octree can be rebuilt from another thread, so this info can be invalidated at anytime. */
@@ -272,7 +272,7 @@ public:
 	UE_API virtual const PCGPointOctree::FPointOctree& GetPointOctree() const;
 	UE_API virtual FBox GetBounds() const override;
 		
-	UFUNCTION(BlueprintCallable, Category = SpatialData, meta = (DisplayName="Set Points From"))
+	UFUNCTION(BlueprintCallable, Category = PointData, meta = (DisplayName="Set Points From"))
 	UE_API void BP_SetPointsFrom(const UPCGBasePointData* InData, const TArray<int32>& InDataIndices);
 	
 	UE_API void SetPointsFrom(const UPCGBasePointData* InData, const TArrayView<const int32>& InDataIndices);
@@ -337,58 +337,58 @@ public:
 		return PropertiesToAllocate;
 	}
 
-	UFUNCTION(BlueprintCallable, Category = SpatialData)
+	UFUNCTION(BlueprintCallable, Category = PointData)
 	static TArray<FTransform> GetTransformValuesFromRange(const FPCGPointInputRange& InputRange);
 
-	UFUNCTION(BlueprintCallable, Category = SpatialData)
+	UFUNCTION(BlueprintCallable, Category = PointData)
 	static TArray<float> GetDensityValuesFromRange(const FPCGPointInputRange& InputRange);
 
-	UFUNCTION(BlueprintCallable, Category = SpatialData)
+	UFUNCTION(BlueprintCallable, Category = PointData)
 	static TArray<FVector> GetBoundsMinValuesFromRange(const FPCGPointInputRange& InputRange);
 
-	UFUNCTION(BlueprintCallable, Category = SpatialData)
+	UFUNCTION(BlueprintCallable, Category = PointData)
 	static TArray<FVector> GetBoundsMaxValuesFromRange(const FPCGPointInputRange& InputRange);
 
-	UFUNCTION(BlueprintCallable, Category = SpatialData)
+	UFUNCTION(BlueprintCallable, Category = PointData)
 	static TArray<FVector4> GetColorValuesFromRange(const FPCGPointInputRange& InputRange);
 
-	UFUNCTION(BlueprintCallable, Category = SpatialData)
+	UFUNCTION(BlueprintCallable, Category = PointData)
 	static TArray<float> GetSteepnessValuesFromRange(const FPCGPointInputRange& InputRange);
 
-	UFUNCTION(BlueprintCallable, Category = SpatialData)
+	UFUNCTION(BlueprintCallable, Category = PointData)
 	static TArray<int32> GetSeedValuesFromRange(const FPCGPointInputRange& InputRange);
 
-	UFUNCTION(BlueprintCallable, Category = SpatialData)
+	UFUNCTION(BlueprintCallable, Category = PointData)
 	static TArray<int64> GetMetadataEntryValuesFromRange(const FPCGPointInputRange& InputRange);
 
-	UFUNCTION(BlueprintCallable, Category = SpatialData)
+	UFUNCTION(BlueprintCallable, Category = PointData)
 	static void SetTransformValuesOnRange(UPARAM(ref) FPCGPointOutputRange& OutputRange, const TArray<FTransform>& Values);
 		
-	UFUNCTION(BlueprintCallable, Category = SpatialData)
+	UFUNCTION(BlueprintCallable, Category = PointData)
 	static void SetDensityValuesOnRange(UPARAM(ref) FPCGPointOutputRange& OutputRange, const TArray<float>& Values);
 
-	UFUNCTION(BlueprintCallable, Category = SpatialData)
+	UFUNCTION(BlueprintCallable, Category = PointData)
 	static void SetBoundsMinValuesOnRange(UPARAM(ref) FPCGPointOutputRange& OutputRange, const TArray<FVector>& Values);
 
-	UFUNCTION(BlueprintCallable, Category = SpatialData)
+	UFUNCTION(BlueprintCallable, Category = PointData)
 	static void SetBoundsMaxValuesOnRange(UPARAM(ref) FPCGPointOutputRange& OutputRange, const TArray<FVector>& Values);
 
-	UFUNCTION(BlueprintCallable, Category = SpatialData)
+	UFUNCTION(BlueprintCallable, Category = PointData)
 	static void SetColorValuesOnRange(UPARAM(ref) FPCGPointOutputRange& OutputRange, const TArray<FVector4>& Values);
 
-	UFUNCTION(BlueprintCallable, Category = SpatialData)
+	UFUNCTION(BlueprintCallable, Category = PointData)
 	static void SetSteepnessValuesOnRange(UPARAM(ref) FPCGPointOutputRange& OutputRange, const TArray<float>& Values);
 
-	UFUNCTION(BlueprintCallable, Category = SpatialData)
+	UFUNCTION(BlueprintCallable, Category = PointData)
 	static void SetSeedValuesOnRange(UPARAM(ref) FPCGPointOutputRange& OutputRange, const TArray<int32>& Values);
 
-	UFUNCTION(BlueprintCallable, Category = SpatialData)
+	UFUNCTION(BlueprintCallable, Category = PointData)
 	static void SetMetadataEntryValuesOnRange(UPARAM(ref) FPCGPointOutputRange& OutputRange, const TArray<int64>& Values);
 
-	UFUNCTION(BlueprintCallable, Category = SpatialData)
+	UFUNCTION(BlueprintCallable, Category = PointData)
 	static FPCGPoint GetPointFromRange(const FPCGPointInputRange& InputRange, int32 Index);
 
-	UFUNCTION(BlueprintCallable, Category = SpatialData)
+	UFUNCTION(BlueprintCallable, Category = PointData)
 	static void SetPointOnRange(UPARAM(ref) FPCGPointOutputRange& OutputRange, int32 Index, const FPCGPoint& Point);
 
 protected:
