[PCG] De-duplicate generation sources at beginning of RGS flow.

Add the filtered gen sources to the debug overlay which is big help for inspecting/debugging.

[RN] [PCG] Add unique generation source list to runtime generation debug overlay
#Synced-CL 0

diff --git a/Engine/Plugins/PCG/Source/PCG/Private/RuntimeGen/GenSources/PCGGenSourceBase.cpp b/Engine/Plugins/PCG/Source/PCG/Private/RuntimeGen/GenSources/PCGGenSourceBase.cpp
new file mode 100644
--- /dev/null
+++ b/Engine/Plugins/PCG/Source/PCG/Private/RuntimeGen/GenSources/PCGGenSourceBase.cpp
@@ -0,0 +1,85 @@
+// Copyright Epic Games, Inc. All Rights Reserved.
+
+#include "RuntimeGen/GenSources/PCGGenSourceBase.h"
+
+#include "RuntimeGen/PCGRuntimeGenContext.h"
+
+#include UE_INLINE_GENERATED_CPP_BY_NAME(PCGGenSourceBase)
+
+static bool AreVectorsEquivalent(const TOptional<FVector>& InA, const TOptional<FVector>& InB)
+{
+	if (InA.IsSet() != InB.IsSet())
+	{
+		return false;
+	}
+
+	// Both unset, or both values match (approximately).
+	return !InA.IsSet() || InA.GetValue().Equals(InB.GetValue());
+}
+
+static bool AreConvexVolumesEquivalent(const TOptional<FConvexVolume>& InA, const TOptional<FConvexVolume>& InB)
+{
+	if (InA.IsSet() != InB.IsSet())
+	{
+		return false;
+	}
+
+	if (!InA.IsSet())
+	{
+		// Both unset, so equal.
+		return true;
+	}
+	
+	if (InA->Planes.Num() != InB->Planes.Num())
+	{
+		return false;
+	}
+
+	// Verify all planes approximately equal.
+	// Note: This is a weak comparison and will fail if order of planes do not match. Could be extended if we identify a beneficial case.
+	for (int Index = 0; Index < InA->Planes.Num(); ++Index)
+	{
+		if (!InA->Planes[Index].Equals(InB->Planes[Index]))
+		{
+			return false;
+		}
+	}
+
+	return true;
+}
+
+bool IPCGGenSourceBase::IsEquivalent(IPCGGenSourceBase* InOther, const FPCGRuntimeGenContext& InContext) const
+{
+	if (this == InOther)
+	{
+		return true;
+	}
+
+	if (!AreVectorsEquivalent(GetPosition(), InOther->GetPosition())
+		|| !AreVectorsEquivalent(GetDirection(), InOther->GetDirection()) // Always include for now, used for priority calculation
+		|| IsLocal() != InOther->IsLocal())
+	{
+		return false;
+	}
+
+	if (InContext.bAnySourcesUseFrustumCulling)
+	{
+		if (InContext.bAnySourcesUse2DGrids)
+		{
+			if (!AreConvexVolumesEquivalent(GetViewFrustum(/*bIs2DGrid=*/true), InOther->GetViewFrustum(/*bIs2DGrid=*/true)))
+			{
+				return false;
+			}
+		}
+
+		if (InContext.bAnySourcesUse3DGrids)
+		{
+			if (!AreConvexVolumesEquivalent(GetViewFrustum(/*bIs2DGrid=*/false), InOther->GetViewFrustum(/*bIs2DGrid=*/false)))
+			{
+				return false;
+			}
+		}
+	}
+
+	return true;
+}
diff --git a/Engine/Plugins/PCG/Source/PCG/Private/RuntimeGen/PCGGenSourceManager.cpp b/Engine/Plugins/PCG/Source/PCG/Private/RuntimeGen/PCGGenSourceManager.cpp
--- a/Engine/Plugins/PCG/Source/PCG/Private/RuntimeGen/PCGGenSourceManager.cpp
+++ b/Engine/Plugins/PCG/Source/PCG/Private/RuntimeGen/PCGGenSourceManager.cpp
@@ -5,12 +5,14 @@
 #include "PCGModule.h"
 #include "PCGWorldActor.h"
 #include "Helpers/PCGHelpers.h"
+#include "RuntimeGen/PCGRuntimeGenContext.h"
 #include "RuntimeGen/GenSources/PCGGenSourceBase.h"
 #include "RuntimeGen/GenSources/PCGGenSourceComponent.h"
 #include "RuntimeGen/GenSources/PCGGenSourceEditorCamera.h"
 #include "RuntimeGen/GenSources/PCGGenSourcePlayer.h"
 #include "RuntimeGen/GenSources/PCGGenSourceWPStreamingSource.h"
 
+#include "Algo/NoneOf.h"
 #include "Engine/GameViewportClient.h"
 #include "Engine/World.h"
 #include "GameFramework/GameModeBase.h"
@@ -58,7 +60,6 @@ TSet<IPCGGenSourceBase*> FPCGGenSourceManager::GetAllGenSources(const APCGWorldA
 	if (bDirty)
 	{
 		UpdatePerTickGenSources(InPCGWorldActor);
-		bDirty = false;
 	}
 
 	TSet<IPCGGenSourceBase*> AllGenSources;
@@ -66,12 +67,28 @@ TSet<IPCGGenSourceBase*> FPCGGenSourceManager::GetAllGenSources(const APCGWorldA
 
 	for (TScriptInterface<IPCGGenSourceBase> GenSource : RegisteredGenSources)
 	{
-		AllGenSources.Add(GenSource.GetInterface());
+		if (GenSource)
+		{
+			if (bDirty)
+			{
+				GenSource->Tick();
+			}
+
+			AllGenSources.Add(GenSource.GetInterface());
+		}
 	}
 
 	for (const TPair<FName, TScriptInterface<IPCGGenSourceBase>>& NamedGenSource : RegisteredNamedGenSources)
 	{
-		AllGenSources.Add(NamedGenSource.Value.GetInterface());
+		if (NamedGenSource.Value)
+		{
+			if (bDirty)
+			{
+				NamedGenSource.Value->Tick();
+			}
+
+			AllGenSources.Add(NamedGenSource.Value.GetInterface());
+		}
 	}
 
 	// Support deprecated behavior
@@ -81,13 +98,36 @@ TSet<IPCGGenSourceBase*> FPCGGenSourceManager::GetAllGenSources(const APCGWorldA
 	// Acquire a generation source for the active editor viewport if one exists.
 	if (EditorCameraGenSource && EditorCameraGenSource->EditorViewportClient)
 	{
+		if (bDirty)
+		{
+			EditorCameraGenSource->Tick();
+		}
+
 		AllGenSources.Add(EditorCameraGenSource);
 	}
 #endif
 
+	// Flag that all gen sources have been ticked.
+	bDirty = false;
+
 	return AllGenSources;
 }
 
+void FPCGGenSourceManager::GetUniqueGenSources(const FPCGRuntimeGenContext& InContext, TSet<IPCGGenSourceBase*>& OutUniqueGenSources)
+{
+	TRACE_CPUPROFILER_EVENT_SCOPE(FPCGGenSourceManager::GetUniqueGenSources);
+
+	TSet<IPCGGenSourceBase*> AllGenSources = GetAllGenSources(InContext.PCGWorldActor);
+
+	for (IPCGGenSourceBase* GenSource : AllGenSources)
+	{
+		if (GenSource->GetPosition().IsSet() && Algo::NoneOf(OutUniqueGenSources, [&InContext, GenSource](const IPCGGenSourceBase* InAddedGenSource) { return InAddedGenSource && InAddedGenSource->IsEquivalent(GenSource, InContext); }))
+		{
+			OutUniqueGenSources.Add(GenSource);
+		}
+	}
+}
+
 bool FPCGGenSourceManager::RegisterGenSource(IPCGGenSourceBase* InGenSource, FName InGenSourceName)
 {
 	if (UPCGGenSourceComponent* GenSourceComponent = Cast<UPCGGenSourceComponent>(InGenSource))
@@ -198,7 +238,8 @@ void FPCGGenSourceManager::UpdatePerTickGenSources(const APCGWorldActor* InPCGWo
 	// Remove obsolete sources
 	for (TMap<FName, TScriptInterface<IPCGGenSourceBase>>::TIterator It = RegisteredNamedGenSources.CreateIterator(); It; ++It)
 	{
-		if (UPCGGenSourcePlayer* SourcePlayer = Cast<UPCGGenSourcePlayer>(It.Value().GetInterface()); !CurrentControllers.Contains(SourcePlayer->GetPlayerController().Get()))
+		const UPCGGenSourcePlayer* SourcePlayer = Cast<UPCGGenSourcePlayer>(It.Value().GetInterface());
+		if (SourcePlayer && !CurrentControllers.Contains(SourcePlayer->GetPlayerController().Get()))
 		{
 			It.RemoveCurrent();
 		}
@@ -230,24 +271,15 @@ void FPCGGenSourceManager::UpdatePerTickGenSources(const APCGWorldActor* InPCGWo
 void FPCGGenSourceManager::GatherWorldPartitionGenSources(TSet<IPCGGenSourceBase*>& OutGenSources) const
 {
 PRAGMA_DISABLE_DEPRECATION_WARNINGS
-	// Acquire a generation source for each active streaming source that doesn't have a generation source already (based on name collision test).
 	for (TObjectPtr<UPCGGenSourceWPStreamingSource> WPGenSource : WorldPartitionGenSources)
 	{
-		bool bAlreadyTracked = false;
-
-		if (const TScriptInterface<IPCGGenSourceBase>* FoundGenSource = (WPGenSource && WPGenSource->StreamingSource) ? RegisteredNamedGenSources.Find(WPGenSource->StreamingSource->Name) : nullptr)
+		if (WPGenSource)
 		{
-		// We've seen during cinematics there are two generation sources with the same name but with different locations in the world, so match based on position too.
-			TOptional<FVector> WPGenSourcePosition = WPGenSource->GetPosition();
-			TOptional<FVector> FoundGenSourcePosition = (*FoundGenSource)->GetPosition();
-			if (WPGenSourcePosition && FoundGenSourcePosition && (*FoundGenSourcePosition - *WPGenSourcePosition).SquaredLength() < UE_KINDA_SMALL_NUMBER)
+			if (bDirty)
 			{
-				bAlreadyTracked = true;
+				WPGenSource->Tick();
 			}
-		}
 
-		if (!bAlreadyTracked)
-		{
 			OutGenSources.Add(WPGenSource);
 		}
 	}
diff --git a/Engine/Plugins/PCG/Source/PCG/Private/RuntimeGen/PCGRuntimeGenScheduler.cpp b/Engine/Plugins/PCG/Source/PCG/Private/RuntimeGen/PCGRuntimeGenScheduler.cpp
--- a/Engine/Plugins/PCG/Source/PCG/Private/RuntimeGen/PCGRuntimeGenScheduler.cpp
+++ b/Engine/Plugins/PCG/Source/PCG/Private/RuntimeGen/PCGRuntimeGenScheduler.cpp
@@ -141,32 +141,54 @@ namespace PCGRuntimeGenSchedulerHelpers
 			TickStartTime = FPlatformTime::Seconds();
 		}
 
-		void EndTick(const FPCGRuntimeGenScheduler& InRuntimeGenScheduler)
+		void EndTick(const FPCGRuntimeGenScheduler& InRuntimeGenScheduler, const TSet<IPCGGenSourceBase*>& InGenSources)
 		{
 			if (CVarEnableDebugOverlay.GetValueOnGameThread())
 			{
 				const double ElapsedMs = (FPlatformTime::Seconds() - TickStartTime) * 1000.0;
 
+				FStringBuilderBase GenSourceStringBuilder;
+				for (const IPCGGenSourceBase* GenSource : InGenSources)
+				{
+					if (GenSource)
+					{
+						const TOptional<FVector> Position = GenSource->GetPosition();
+						const TOptional<FVector> Direction = GenSource->GetDirection();
+
+						GenSourceStringBuilder.Appendf(TEXT("        [%s] Position: %s Direction: %s\n"),
+							*GenSource->GetDebugName(),
+							Position ? *Position->ToString() : TEXT("UNSET"),
+							Direction ? *Direction->ToString() : TEXT("UNSET"));
+					}
+					else
+					{
+						GenSourceStringBuilder.Append(TEXT("        NULL SOURCE\n"));
+					}
+				}
+
 				GEngine->AddOnScreenDebugMessage(
 					/*Key=*/-1,
 					/*TimeToDisplay=*/0.0f,
 					FColor::Yellow,
 					FString::Printf(
-						TEXT("PCG Runtime Generation\n")
+						TEXT("PCG Runtime Generation%s\n")
 						TEXT("    Tick time: %.3fms\n")
 						TEXT("        Generate Calls: %d\n")
 						TEXT("        Cleanup Calls: %d\n")
 						TEXT("        Grid Cell Scans: %d\n")
 						TEXT("    Num Generating Components: %d / %d\n")
 						TEXT("    PA Pool: %d / %d available\n")
-						TEXT("    VT Preloads: %d\n"),
+						TEXT("    VT Preloads: %d\n")
+						TEXT("    Unique generation sources:\n%s"),
+						PCGHelpers::IsServer(InRuntimeGenScheduler.World) ? TEXT(" (Server)") : TEXT(""),
 						ElapsedMs,
 						GenerateCallCounter,
 						CleanupCallCounter,
 						GridCellScanCounter,
 						NumGeneratingComponents, FMath::Max(1, CVarNumGeneratingComponentsAtSameTime.GetValueOnGameThread()),
 						InRuntimeGenScheduler.PartitionActorPool.Num(), InRuntimeGenScheduler.PartitionActorPoolSize,
-						VTPreloadCounter)
+						VTPreloadCounter,
+						GenSourceStringBuilder.ToString())
 				);
 			}
 		}
@@ -203,6 +225,8 @@ FPCGRuntimeGenScheduler::FPCGRuntimeGenScheduler(UWorld* InWorld, FPCGTrackingMa
 	BasePoolSizeLastFrame = PCGRuntimeGenSchedulerHelpers::CVarRuntimeGenerationBasePoolSize.GetValueOnAnyThread();
 	FramesUntilGeneration = PCGRuntimeGenSchedulerHelpers::CVarFramesBeforeFirstGenerate.GetValueOnGameThread();
 
+	RefreshContext();
+
 	FLevelStreamingDelegates::OnLevelStreamingStateChanged.AddRaw(this, &FPCGRuntimeGenScheduler::OnLevelStreamingStateChanged);
 
 #if WITH_EDITOR
@@ -266,20 +290,14 @@ void FPCGRuntimeGenScheduler::Tick(APCGWorldActor* InPCGWorldActor, double InEnd
 
 	TickCVars(InPCGWorldActor);
 
-	GenSources.Empty(GenSources.Num());
+	Context.PCGWorldActor = InPCGWorldActor;
 
-	if (bAnyRuntimeGenExecutionSourcesExist)
+	// Update and grab unique generation sources.
+	GenSources.Empty(GenSources.Num());
+	if (Context.bAnyRuntimeGenSourcesExist)
 	{
 		GenSourceManager->Tick();
-		GenSources = GenSourceManager->GetAllGenSources(InPCGWorldActor);
-
-		for (IPCGGenSourceBase* GenSource : GenSources)
-		{
-			if (ensure(GenSource))
-			{
-				GenSource->Tick();
-			}
-		}
+		GenSourceManager->GetUniqueGenSources(Context, GenSources);
 	}
 
 	if (!CachedPrimingInfos.IsEmpty() && !GenSources.IsEmpty() && PCGRuntimeGenSchedulerHelpers::CVarEnableVirtualTexturePriming.GetValueOnGameThread())
@@ -343,7 +361,7 @@ void FPCGRuntimeGenScheduler::Tick(APCGWorldActor* InPCGWorldActor, double InEnd
 	}
 
 #if PCG_RGS_ONSCREENDEBUGMESSAGES
-	PCGRuntimeGenSchedulerHelpers::Stats.EndTick(*this);
+	PCGRuntimeGenSchedulerHelpers::Stats.EndTick(*this, GenSources);
 #endif
 }
 
@@ -377,30 +395,13 @@ bool FPCGRuntimeGenScheduler::ShouldTick()
 	}
 #endif
 
-	if (bAnyRuntimeGenExecutionSourcesExistDirty)
+	if (bContextDirty)
 	{
-		const bool bDidAnyRuntimeGenComponentsExist = bAnyRuntimeGenExecutionSourcesExist;
-		bAnyRuntimeGenExecutionSourcesExist = TrackingManager->AnyRuntimeGenExecutionSourcesExist();
-		bAnyRuntimeGenExecutionSourcesExistDirty = false;
-
-		if (PCGRuntimeGenSchedulerHelpers::CVarRuntimeGenerationEnableDebugging.GetValueOnAnyThread())
-		{
-			if (bDidAnyRuntimeGenComponentsExist != bAnyRuntimeGenExecutionSourcesExist)
-			{
-				if (bAnyRuntimeGenExecutionSourcesExist)
-				{
-					UE_LOG(LogPCG, Warning, TEXT("[RUNTIMEGEN] THERE ARE NOW RUNTIME EXECUTION SOURCES IN THE LEVEL. SCHEDULER WILL BEGIN TICKING."));
-				}
-				else
-				{
-					UE_LOG(LogPCG, Warning, TEXT("[RUNTIMEGEN] THERE ARE NO MORE RUNTIME EXECUTION SOURCES. SCHEDULER WILL ONLY TICK TO CLEANUP."));
-				}
-			}
-		}
+		RefreshContext();
 	}
 
 	// We can stop ticking if there are no runtime gen components alive and there are no generated components that need cleaning up.
-	if (!bAnyRuntimeGenExecutionSourcesExist && GeneratedComponents.IsEmpty() && GeneratedComponentsToRemove.IsEmpty())
+	if (!Context.bAnyRuntimeGenSourcesExist && GeneratedComponents.IsEmpty() && GeneratedComponentsToRemove.IsEmpty())
 	{
 		return false;
 	}
@@ -1418,7 +1419,7 @@ void FPCGRuntimeGenScheduler::OnOriginalComponentRegistered(UPCGComponent* InOri
 	}
 
 	// When an original/non-partitioned component is registered, we need to dirty the state.
-	bAnyRuntimeGenExecutionSourcesExistDirty = true;
+	bContextDirty = true;
 
 #if WITH_EDITOR
 	RegisterGraphParameterChangedEvent(InOriginalComponent);
@@ -1437,7 +1438,7 @@ void FPCGRuntimeGenScheduler::OnOriginalComponentUnregistered(UPCGComponent* InO
 	check(TrackingManager);
 
 	// When an original/non-partitioned component is unregistered, we need to dirty the state.
-	bAnyRuntimeGenExecutionSourcesExistDirty = true;
+	bContextDirty = true;
 
 #if WITH_EDITOR
 	UnregisterGraphParameterChangedEvent(InOriginalComponent);
@@ -1992,7 +1993,7 @@ void FPCGRuntimeGenScheduler::RefreshComponent(IPCGGraphExecutionSource* InCompo
 	if (!bIsLocalComponent)
 	{
 		// When an original/non-partitioned component is refreshed, we need to dirty the state.
-		bAnyRuntimeGenExecutionSourcesExistDirty = true;
+		bContextDirty = true;
 	}
 }
 
@@ -2150,6 +2151,59 @@ void FPCGRuntimeGenScheduler::CacheVirtualTexturePrimingInfos(UPCGComponent* InO
 	}
 }
 
+void FPCGRuntimeGenScheduler::RefreshContext()
+{
+	const bool bAnyRuntimeGenSourcesBefore = Context.bAnyRuntimeGenSourcesExist;
+
+	Context.bAnySourcesUse2DGrids = false;
+	Context.bAnySourcesUse3DGrids = false;
+	Context.bAnySourcesUseFrustumCulling = false;
+	Context.bAnyRuntimeGenSourcesExist = false;
+
+	auto ProcessExecutionSources = [&ContextInner=Context](const TSet<IPCGGraphExecutionSource*>& InExecutionSources)
+	{
+		for (IPCGGraphExecutionSource* ExecutionSource : InExecutionSources)
+		{
+			if (ExecutionSource && ExecutionSource->GetExecutionState().IsManagedByRuntimeGenSystem())
+			{
+				ContextInner.bAnyRuntimeGenSourcesExist = true;
+
+				const bool bUses2DGrid = ExecutionSource->GetExecutionState().Use2DGrid();
+				ContextInner.bAnySourcesUse2DGrids |= bUses2DGrid;
+				ContextInner.bAnySourcesUse3DGrids |= !bUses2DGrid;
+
+				if (const UPCGComponent* Component = Cast<UPCGComponent>(ExecutionSource))
+				{
+					if (UPCGSchedulingPolicyBase* Policy = Component->GetRuntimeGenSchedulingPolicy())
+					{
+						ContextInner.bAnySourcesUseFrustumCulling |= Policy->CullsBasedOnDirection();
+					}
+				}
+			}
+		}
+	};
+
+	ProcessExecutionSources(TrackingManager->GetAllRegisteredNonPartitionedExecutionSources());
+	ProcessExecutionSources(TrackingManager->GetAllRegisteredPartitionedExecutionSources());
+
+	if (PCGRuntimeGenSchedulerHelpers::CVarRuntimeGenerationEnableDebugging.GetValueOnAnyThread())
+	{
+		if (bAnyRuntimeGenSourcesBefore != Context.bAnyRuntimeGenSourcesExist)
+		{
+			if (Context.bAnyRuntimeGenSourcesExist)
+			{
+				UE_LOG(LogPCG, Warning, TEXT("[RUNTIMEGEN] THERE ARE NOW RUNTIME EXECUTION SOURCES IN THE LEVEL. SCHEDULER WILL BEGIN TICKING."));
+			}
+			else
+			{
+				UE_LOG(LogPCG, Warning, TEXT("[RUNTIMEGEN] THERE ARE NO MORE RUNTIME EXECUTION SOURCES. SCHEDULER WILL ONLY TICK TO CLEANUP."));
+			}
+		}
+	}
+
+	bContextDirty = false;
+}
+
 #if WITH_EDITOR
 void FPCGRuntimeGenScheduler::RegisterGraphParameterChangedEvent(UPCGComponent* InOriginalComponent)
 {
diff --git a/Engine/Plugins/PCG/Source/PCG/Public/RuntimeGen/GenSources/PCGGenSourceBase.h b/Engine/Plugins/PCG/Source/PCG/Public/RuntimeGen/GenSources/PCGGenSourceBase.h
--- a/Engine/Plugins/PCG/Source/PCG/Public/RuntimeGen/GenSources/PCGGenSourceBase.h
+++ b/Engine/Plugins/PCG/Source/PCG/Public/RuntimeGen/GenSources/PCGGenSourceBase.h
@@ -7,7 +7,7 @@
 
 #include "PCGGenSourceBase.generated.h"
 
-#define UE_API PCG_API
+struct FPCGRuntimeGenContext;
 
 UINTERFACE(MinimalAPI, BlueprintType)
 class UPCGGenSourceBase : public UInterface
@@ -38,6 +38,9 @@ public:
 
 	/** Returns true if this is a local gen source. (client-server specific) */
 	virtual bool IsLocal() const { return true; }
-};
 
-#undef UE_API
+	/** Whether a single runtime gen grid scan can be performed for both sources. */
+	PCG_API virtual bool IsEquivalent(IPCGGenSourceBase* InOther, const FPCGRuntimeGenContext& InContext) const;
+
+	virtual FString GetDebugName() const { return TEXT("AnonymousGenSource"); }
+};
diff --git a/Engine/Plugins/PCG/Source/PCG/Public/RuntimeGen/GenSources/PCGGenSourceComponent.h b/Engine/Plugins/PCG/Source/PCG/Public/RuntimeGen/GenSources/PCGGenSourceComponent.h
--- a/Engine/Plugins/PCG/Source/PCG/Public/RuntimeGen/GenSources/PCGGenSourceComponent.h
+++ b/Engine/Plugins/PCG/Source/PCG/Public/RuntimeGen/GenSources/PCGGenSourceComponent.h
@@ -8,8 +8,6 @@
 
 #include "PCGGenSourceComponent.generated.h"
 
-#define UE_API PCG_API
-
 namespace EEndPlayReason { enum Type : int; }
 
 class FPCGGenSourceManager;
@@ -27,24 +25,20 @@ class UPCGGenSourceComponent : public UActorComponent, public IPCGGenSourceBase
 
 public:
 	//~Begin UActorComponent Interface
-	UE_API virtual void BeginPlay() override;
-	UE_API virtual void EndPlay(const EEndPlayReason::Type EndPlayReason) override;
-	UE_API virtual void OnRegister() override;
-	UE_API virtual void OnUnregister() override;
-	UE_API virtual void OnComponentCreated() override;
-	UE_API virtual void OnComponentDestroyed(bool bDestroyingHierarchy) override;
+	PCG_API virtual void BeginPlay() override;
+	PCG_API virtual void EndPlay(const EEndPlayReason::Type EndPlayReason) override;
+	PCG_API virtual void OnRegister() override;
+	PCG_API virtual void OnUnregister() override;
+	PCG_API virtual void OnComponentCreated() override;
+	PCG_API virtual void OnComponentDestroyed(bool bDestroyingHierarchy) override;
 	//~End UActorComponent Interface
 
 	//~Begin IPCGGenSourceInterface
-	/** Returns the world space position of this gen source. */
-	UE_API virtual TOptional<FVector> GetPosition() const override;
-
-	/** Returns the normalized forward vector of this gen source. */
-	UE_API virtual TOptional<FVector> GetDirection() const override;
+	PCG_API virtual TOptional<FVector> GetPosition() const override;
+	PCG_API virtual TOptional<FVector> GetDirection() const override;
+	virtual FString GetDebugName() const override { return GetName(); }
 	//~End IPCGGenSourceInterface
 
 protected:
-	UE_API FPCGGenSourceManager* GetGenSourceManager() const;
+	PCG_API FPCGGenSourceManager* GetGenSourceManager() const;
 };
-
-#undef UE_API
diff --git a/Engine/Plugins/PCG/Source/PCG/Public/RuntimeGen/GenSources/PCGGenSourceEditorCamera.h b/Engine/Plugins/PCG/Source/PCG/Public/RuntimeGen/GenSources/PCGGenSourceEditorCamera.h
--- a/Engine/Plugins/PCG/Source/PCG/Public/RuntimeGen/GenSources/PCGGenSourceEditorCamera.h
+++ b/Engine/Plugins/PCG/Source/PCG/Public/RuntimeGen/GenSources/PCGGenSourceEditorCamera.h
@@ -6,8 +6,6 @@
 
 #include "PCGGenSourceEditorCamera.generated.h"
 
-#define UE_API PCG_API
-
 #if WITH_EDITOR
 class FEditorViewportClient;
 #endif
@@ -22,17 +20,13 @@ class UPCGGenSourceEditorCamera : public UObject, public IPCGGenSourceBase
 	GENERATED_BODY()
 
 public:
-	/** Update the generation source so that it can cache data that is queried often (e.g. view frustum). Should be called every tick on any active generation sources. */
-	UE_API virtual void Tick() override;
-
-	/** Returns the world space position of this gen source. */
-	UE_API virtual TOptional<FVector> GetPosition() const override;
-
-	/** Returns the normalized forward vector of this gen source. */
-	UE_API virtual TOptional<FVector> GetDirection() const override;
-
-	/** Returns the view frustum of this gen source. */
-	UE_API virtual TOptional<FConvexVolume> GetViewFrustum(bool bIs2DGrid) const override;
+	//~ Begin IPCGGenSourceInterface
+	PCG_API virtual void Tick() override;
+	PCG_API virtual TOptional<FVector> GetPosition() const override;
+	PCG_API virtual TOptional<FVector> GetDirection() const override;
+	PCG_API virtual TOptional<FConvexVolume> GetViewFrustum(bool bIs2DGrid) const override;
+	virtual FString GetDebugName() const override { return GetName(); }
+	//~ End IPCGGenSourceInterface
 
 public:
 #if WITH_EDITORONLY_DATA
@@ -40,5 +34,3 @@ public:
 	TOptional<FConvexVolume> ViewFrustum;
 #endif
 };
-
-#undef UE_API
diff --git a/Engine/Plugins/PCG/Source/PCG/Public/RuntimeGen/GenSources/PCGGenSourcePlayer.h b/Engine/Plugins/PCG/Source/PCG/Public/RuntimeGen/GenSources/PCGGenSourcePlayer.h
--- a/Engine/Plugins/PCG/Source/PCG/Public/RuntimeGen/GenSources/PCGGenSourcePlayer.h
+++ b/Engine/Plugins/PCG/Source/PCG/Public/RuntimeGen/GenSources/PCGGenSourcePlayer.h
@@ -6,8 +6,6 @@
 
 #include "PCGGenSourcePlayer.generated.h"
 
-#define UE_API PCG_API
-
 class APlayerController;
 
 /**
@@ -19,23 +17,17 @@ class UPCGGenSourcePlayer : public UObject, public IPCGGenSourceBase
 	GENERATED_BODY()
 
 public:
-	/** Update the generation source so that it can cache data that is queried often (e.g. view frustum). Should be called every tick on any active generation sources. */
-	UE_API virtual void Tick() override;
-
-	/** Returns the world space position of this gen source. */
-	UE_API virtual TOptional<FVector> GetPosition() const override;
-
-	/** Returns the normalized forward vector of this gen source. */
-	UE_API virtual TOptional<FVector> GetDirection() const override;
-
-	/** Returns the view frustum of this gen source. */
+	//~ Begin IPCGGenSourceInterface
+	PCG_API virtual void Tick() override;
+	PCG_API virtual TOptional<FVector> GetPosition() const override;
+	PCG_API virtual TOptional<FVector> GetDirection() const override;
 	virtual TOptional<FConvexVolume> GetViewFrustum(bool bIs2DGrid) const override { return ViewFrustum; }
-
-	/** Returns true if this is a local gen source. (network) */
 	virtual bool IsLocal() const override;
+	virtual FString GetDebugName() const override { return GetName(); }
+	//~ End IPCGGenSourceInterface
 
 	TWeakObjectPtr<const APlayerController> GetPlayerController() const { return PlayerController; }
-	UE_API void SetPlayerController(const APlayerController* InPlayerController);
+	PCG_API void SetPlayerController(const APlayerController* InPlayerController);
 
 	bool IsValid() const { return PlayerController.IsValid(); }
 
@@ -46,5 +38,3 @@ protected:
 	TOptional<FVector> CachedPosition;
 	TOptional<FVector> CachedDirection;
 };
-
-#undef UE_API
diff --git a/Engine/Plugins/PCG/Source/PCG/Public/RuntimeGen/GenSources/PCGGenSourceWPStreamingSource.h b/Engine/Plugins/PCG/Source/PCG/Public/RuntimeGen/GenSources/PCGGenSourceWPStreamingSource.h
--- a/Engine/Plugins/PCG/Source/PCG/Public/RuntimeGen/GenSources/PCGGenSourceWPStreamingSource.h
+++ b/Engine/Plugins/PCG/Source/PCG/Public/RuntimeGen/GenSources/PCGGenSourceWPStreamingSource.h
@@ -6,8 +6,6 @@
 
 #include "PCGGenSourceWPStreamingSource.generated.h"
 
-#define UE_API PCG_API
-
 struct FWorldPartitionStreamingSource;
 
 UCLASS(MinimalAPI, BlueprintType, ClassGroup = (Procedural))
@@ -16,14 +14,12 @@ class UE_DEPRECATED(5.7, "No longer supported.") UPCGGenSourceWPStreamingSource
 	GENERATED_BODY()
 
 public:
-	/** Returns the world space position of this gen source. */
-	UE_API virtual TOptional<FVector> GetPosition() const override;
-
-	/** Returns the normalized forward vector of this gen source. */
-	UE_API virtual TOptional<FVector> GetDirection() const override;
+	//~ Begin IPCGGenSourceInterface
+	PCG_API virtual TOptional<FVector> GetPosition() const override;
+	PCG_API virtual TOptional<FVector> GetDirection() const override;
+	virtual FString GetDebugName() const override { return GetName(); }
+	//~ End IPCGGenSourceInterface
 
 public:
 	const FWorldPartitionStreamingSource* StreamingSource;
 };
-
-#undef UE_API
diff --git a/Engine/Plugins/PCG/Source/PCG/Public/RuntimeGen/PCGGenSourceManager.h b/Engine/Plugins/PCG/Source/PCG/Public/RuntimeGen/PCGGenSourceManager.h
--- a/Engine/Plugins/PCG/Source/PCG/Public/RuntimeGen/PCGGenSourceManager.h
+++ b/Engine/Plugins/PCG/Source/PCG/Public/RuntimeGen/PCGGenSourceManager.h
@@ -4,15 +4,15 @@
 
 #include "RuntimeGen/GenSources/PCGGenSourceBase.h"
 
+class AController;
+class AGameModeBase;
 class APCGWorldActor;
+class APlayerController;
 class UPCGGenSourceComponent;
 class UPCGGenSourceEditorCamera;
 class UPCGGenSourceWPStreamingSource;
-
-class AController;
-class AGameModeBase;
-class APlayerController;
 class UWorld;
+struct FPCGRuntimeGenContext;
 
 /**
  * The runtime Generation Source Manager tracks generation sources in the world for use by the Runtime Generation Scheduler.
@@ -29,6 +29,9 @@ public:
 	/** Creates the set of all generation sources tracked by the manager. Call 'Tick()' to keep tracked generation sources up to date. */
 	TSet<IPCGGenSourceBase*> PCG_API GetAllGenSources(const APCGWorldActor* InPCGWorldActor);
 
+	/** Creates a set of generation sources that are unique for a given runtime gen context (de-duplicating sources which will produce equivalent grid cell scans). */
+	void GetUniqueGenSources(const FPCGRuntimeGenContext& InContext, TSet<IPCGGenSourceBase*>& OutUniqueGenSources);
+
 	/** Adds a UPCGGenSource to be tracked by the GenSourceManager. */
 	bool PCG_API RegisterGenSource(IPCGGenSourceBase* InGenSource, FName InGenSourceName = NAME_None);
 
diff --git a/Engine/Plugins/PCG/Source/PCG/Public/RuntimeGen/PCGRuntimeGenContext.h b/Engine/Plugins/PCG/Source/PCG/Public/RuntimeGen/PCGRuntimeGenContext.h
new file mode 100644
--- /dev/null
+++ b/Engine/Plugins/PCG/Source/PCG/Public/RuntimeGen/PCGRuntimeGenContext.h
@@ -0,0 +1,16 @@
+// Copyright Epic Games, Inc. All Rights Reserved.
+
+#pragma once
+
+class APCGWorldActor;
+
+/** Runtime generation scheduler context information. */
+struct FPCGRuntimeGenContext
+{
+	const APCGWorldActor* PCGWorldActor = nullptr;
+
+	bool bAnyRuntimeGenSourcesExist = false;
+	bool bAnySourcesUseFrustumCulling = false;
+	bool bAnySourcesUse2DGrids = false;
+	bool bAnySourcesUse3DGrids = false;
+};
diff --git a/Engine/Plugins/PCG/Source/PCG/Public/RuntimeGen/PCGRuntimeGenScheduler.h b/Engine/Plugins/PCG/Source/PCG/Public/RuntimeGen/PCGRuntimeGenScheduler.h
--- a/Engine/Plugins/PCG/Source/PCG/Public/RuntimeGen/PCGRuntimeGenScheduler.h
+++ b/Engine/Plugins/PCG/Source/PCG/Public/RuntimeGen/PCGRuntimeGenScheduler.h
@@ -5,6 +5,7 @@
 #include "PCGCommon.h"
 #include "Grid/PCGGridDescriptor.h"
 #include "RuntimeGen/PCGRuntimeGenChangeDetection.h"
+#include "RuntimeGen/PCGRuntimeGenContext.h"
 
 #include "UObject/WeakObjectPtr.h"
 
@@ -214,6 +215,9 @@ private:
 	/** Cache VT priming infos found in graph params of the given component. */
 	void CacheVirtualTexturePrimingInfos(UPCGComponent* InOriginalComponent);
 
+	/** Update RGS context information. */
+	void RefreshContext();
+
 #if WITH_EDITOR
 	void RegisterGraphParameterChangedEvent(UPCGComponent* InOriginalComponent);
 	void UnregisterGraphParameterChangedEvent(UPCGComponent* InOriginalComponent);
@@ -259,10 +263,6 @@ private:
 	/** Requests to flush all actors are deferred so they can be handled at a known time during tick. */
 	bool bActorFlushRequested = false;
 
-	/** Track the existence of runtime gen execution sources to avoid unnecessary computation when there is no work to do. */
-	bool bAnyRuntimeGenExecutionSourcesExist = false;
-	bool bAnyRuntimeGenExecutionSourcesExistDirty = false;
-
 	/** Setting up a PA calls APCGPartitionActor::AddGraphInstance which later calls RefreshComponent, which can create
 	* an infinite refresh loop. To break this loop we write the OC pointer to this variable, and if Refresh gets called for
 	* this OC we early out. Basically we don't respond to refresh calls for a component we are midway through setting up.
@@ -281,6 +281,12 @@ private:
 	/** Used to detect when world state (generation sources, cvars etc) have changed enough to warrant rescanning generation cells. */
 	PCGRuntimeGenChangeDetection::FDetector ChangeDetector;
 
+	/** Meta-information about current context. */
+	FPCGRuntimeGenContext Context;
+
+	/** Whether Context needs to be refreshed. */
+	bool bContextDirty = true;
+
 #if WITH_EDITOR
 	TMap<TObjectKey<UPCGComponent>, FDelegateHandle> GraphParamChangedEventHandles;
 #endif
