[PCG] Update pins warning toggled on by default + few fixes related to blueprint pins

#rnx
#Synced-CL 48550537

diff --git a/Engine/Plugins/PCG/Source/PCG/Private/PCGGraph.cpp b/Engine/Plugins/PCG/Source/PCG/Private/PCGGraph.cpp
--- a/Engine/Plugins/PCG/Source/PCG/Private/PCGGraph.cpp
+++ b/Engine/Plugins/PCG/Source/PCG/Private/PCGGraph.cpp
@@ -54,7 +54,7 @@ namespace PCGGraph
 
 	static TAutoConsoleVariable<bool> CVarEnableWarningOnPinDestructiveChangesInPostLoad(
 		TEXT("pcg.Graph.EnableWarningOnPinDestructiveChangesInPostLoad"),
-		false,
+		true,
 		TEXT("Output warnings when destructive changes happens on PostLoad in UpdatePins."));
 }
 
diff --git a/Engine/Plugins/PCG/Source/PCG/Private/PCGNode.cpp b/Engine/Plugins/PCG/Source/PCG/Private/PCGNode.cpp
--- a/Engine/Plugins/PCG/Source/PCG/Private/PCGNode.cpp
+++ b/Engine/Plugins/PCG/Source/PCG/Private/PCGNode.cpp
@@ -959,13 +959,14 @@ EPCGChangeType UPCGNode::UpdatePins(TFunctionRef<UPCGPin*(UPCGNode*)> PinAllocat
 
 	UE_LOG_CONTEXT("PCGNode", FSoftObjectPath(Settings));
 
-	const FName SettingsName = Settings->GetFName();
-	const FString GraphPath = GetGraph() ?  FSoftObjectPath(GetGraph()).ToString() : TEXT("Unknown");
+	const FText ThisNodeTitle = !bQuiet ? GetNodeTitle(EPCGNodeTitleType::ListView) : FText{};
+	const FString NodeName = !bQuiet ? FString::Printf(TEXT("%s (%s)"), *ThisNodeTitle.ToString(), *GetName()) : FString{};
+	const FString GraphPath = !bQuiet && GetGraph() ?  FSoftObjectPath(GetGraph()).ToString() : TEXT("Unknown");
 	
 	TArray<FPCGPinProperties> InboundPinProperties = Settings->AllInputPinProperties();
 	TArray<FPCGPinProperties> OutboundPinProperties = Settings->AllOutputPinProperties();
 
-	auto RemoveDuplicates = [this, &GraphPath](TArray<FPCGPinProperties>& Properties)
+	auto RemoveDuplicates = [this, &GraphPath, &NodeName](TArray<FPCGPinProperties>& Properties)
 	{
 		for (int32 i = Properties.Num() - 2; i >= 0; --i)
 		{
@@ -974,7 +975,7 @@ EPCGChangeType UPCGNode::UpdatePins(TFunctionRef<UPCGPin*(UPCGNode*)> PinAllocat
 				if (Properties[i].Label == Properties[j].Label)
 				{
 					UE_LOG(LogPCG, Warning, TEXT("UpdatePins: Pin properties from the settings on node '%s' in graph '%s' contained a duplicate pin '%s', removing this pin properties."),
-						*GetName(), *GraphPath, *Properties[i].Label.ToString());
+						*NodeName, *GraphPath, *Properties[i].Label.ToString());
 
 					// Remove but preserve order
 					Properties.RemoveAt(j);
@@ -986,7 +987,7 @@ EPCGChangeType UPCGNode::UpdatePins(TFunctionRef<UPCGPin*(UPCGNode*)> PinAllocat
 	RemoveDuplicates(InboundPinProperties);
 	RemoveDuplicates(OutboundPinProperties);
 
-	auto UpdatePins = [this, &PinAllocator, &TouchedNodes, &SettingsName, bQuiet, &GraphPath](TArray<UPCGPin*>& Pins, const TArray<FPCGPinProperties>& PinProperties)
+	auto UpdatePins = [this, &PinAllocator, &TouchedNodes, &NodeName, bQuiet, &GraphPath](TArray<UPCGPin*>& Pins, const TArray<FPCGPinProperties>& PinProperties)
 	{
 		bool bAppliedEdgeChanges = false;
 		bool bChangedPins = false;
@@ -1010,7 +1011,7 @@ EPCGChangeType UPCGNode::UpdatePins(TFunctionRef<UPCGPin*(UPCGNode*)> PinAllocat
 						bAppliedEdgeChanges = true;
 						if (!bQuiet)
 						{
-							UE_LOG(LogPCG, Warning, TEXT("[%s - %s] UpdatePins: %s"), *GraphPath, *SettingsName.ToString(), *Reason);	
+							UE_LOG(LogPCG, Warning, TEXT("[%s:%s] UpdatePins: %s"), *GraphPath, *NodeName, *Reason);	
 						}
 					}
 
@@ -1034,7 +1035,7 @@ EPCGChangeType UPCGNode::UpdatePins(TFunctionRef<UPCGPin*(UPCGNode*)> PinAllocat
 						bAppliedEdgeChanges = true;
 						if (!bQuiet)
 						{
-							UE_LOG(LogPCG, Warning, TEXT("[%s-%s] UpdatePins: %s"), *GraphPath, *SettingsName.ToString(), *Reason);	
+							UE_LOG(LogPCG, Warning, TEXT("[%s:%s] UpdatePins: %s"), *GraphPath, *NodeName, *Reason);	
 						}
 					}
 					
@@ -1074,7 +1075,7 @@ EPCGChangeType UPCGNode::UpdatePins(TFunctionRef<UPCGPin*(UPCGNode*)> PinAllocat
 				bAppliedEdgeChanges = true;
 				if (!bQuiet)
 				{
-					UE_LOG(LogPCG, Warning, TEXT("[%s-%s] UpdatePins: Pin %s was matched and renamed to %s\n%s"), *GraphPath, *SettingsName.ToString(), *PreviousName.ToString(), *NewName.ToString(), *Reason);	
+					UE_LOG(LogPCG, Warning, TEXT("[%s:%s] UpdatePins: Pin %s was matched and renamed to %s\n%s"), *GraphPath, *NodeName, *PreviousName.ToString(), *NewName.ToString(), *Reason);	
 				}
 			}
 
@@ -1104,7 +1105,7 @@ EPCGChangeType UPCGNode::UpdatePins(TFunctionRef<UPCGPin*(UPCGNode*)> PinAllocat
 				bChangedPins = true;
 			}
 
-			auto RemovePins = [&Pins, &AllPinNames, &bAppliedEdgeChanges, &TouchedNodes, bQuiet, &SettingsName, &GraphPath](TArray<UPCGPin*>& PinsToRemove, bool bRemoveFromAllNames, bool bWarn)
+			auto RemovePins = [&Pins, &AllPinNames, &bAppliedEdgeChanges, &TouchedNodes, bQuiet, &NodeName, &GraphPath](TArray<UPCGPin*>& PinsToRemove, bool bRemoveFromAllNames, bool bWarn)
 			{
 				for (int32 RemovedPinIndex = PinsToRemove.Num() - 1; RemovedPinIndex >= 0; --RemovedPinIndex)
 				{
@@ -1125,7 +1126,7 @@ EPCGChangeType UPCGNode::UpdatePins(TFunctionRef<UPCGPin*(UPCGNode*)> PinAllocat
 							if (bWarn && !bQuiet)
 							{
 								TStringBuilder<512> WarnMessage;
-								WarnMessage += FString::Printf(TEXT("[%s-%s] UpdatePins: Pin %s does not exist anymore.\n"), *GraphPath, *SettingsName.ToString(), *Pins[PinIndex]->Properties.Label.ToString());
+								WarnMessage += FString::Printf(TEXT("[%s:%s] UpdatePins: Pin %s does not exist anymore.\n"), *GraphPath, *NodeName, *Pins[PinIndex]->Properties.Label.ToString());
 								WarnMessage += TEXT("List of broken edges:\n");
 								for (const UPCGPin* OtherPin : OtherPinsTouched)
 								{
@@ -1135,10 +1136,10 @@ EPCGChangeType UPCGNode::UpdatePins(TFunctionRef<UPCGPin*(UPCGNode*)> PinAllocat
 									}
 
 									const FName OtherPinLabel = OtherPin->Properties.Label;
-									const UPCGSettings* OtherSettings = OtherPin->Node ? OtherPin->Node->GetSettings() : nullptr;
-									const FName OtherSettingsName = OtherSettings ? OtherSettings->GetFName() : "Unknown";
+									const FText OtherNodeTitle = OtherPin->Node ? OtherPin->Node->GetNodeTitle(EPCGNodeTitleType::ListView) : FText{};
+									const FString OtherNodeName = OtherPin->Node ? FString::Printf(TEXT("%s (%s)"), *OtherNodeTitle.ToString(), *OtherPin->Node->GetName()) : TEXT("Unknown");
 
-									WarnMessage += FString::Printf(TEXT(" - %s:%s"), *OtherSettingsName.ToString(), *OtherPinLabel.ToString());
+									WarnMessage += FString::Printf(TEXT(" - %s:%s"), *OtherNodeName, *OtherPinLabel.ToString());
 								}
 
 								UE_LOG(LogPCG, Warning, TEXT("%s"), WarnMessage.ToString());
