[PCG] Demote error obtaining compute graph to warning as it occurs naturally when cache is flushed.

#rnx
#Synced-CL 0

diff --git a/Engine/Plugins/PCG/Source/PCG/Private/Compute/Elements/PCGComputeGraphElement.cpp b/Engine/Plugins/PCG/Source/PCG/Private/Compute/Elements/PCGComputeGraphElement.cpp
--- a/Engine/Plugins/PCG/Source/PCG/Private/Compute/Elements/PCGComputeGraphElement.cpp
+++ b/Engine/Plugins/PCG/Source/PCG/Private/Compute/Elements/PCGComputeGraphElement.cpp
@@ -186,11 +186,16 @@ bool FPCGComputeGraphElement::ExecuteInternal(FPCGContext* InContext) const
 
 			if (!Context->ComputeGraph)
 			{
-				UE_LOG(LogPCG, Error, TEXT("Failed to obtain compute graph for %s '%s', grid size %u, compute graph index %d. Compute graph execution aborted."),
-					DynamicSubgraph ? TEXT("dynamic subgraph") : TEXT("top graph"),
-					DynamicSubgraph ? *DynamicSubgraph->GetName() : *TopGraph->GetName(),
-					Context->GenerationGridSize,
-					ComputeGraphIndex);
+#if WITH_EDITOR
+				// Normally emitted when a graph compilation is flushed while generation in progress.
+				UE_LOG(LogPCG, Warning, TEXT("Failed to obtain compute graph for %s '%s', grid size %u, compute graph index %d. Expected if graph compilation cache flushed mid-execution. Compute graph execution aborted."),
+					DynamicSubgraph ? TEXT("dynamic subgraph") : TEXT("top graph"), DynamicSubgraph ? *DynamicSubgraph->GetName() : *TopGraph->GetName(), Context->GenerationGridSize, ComputeGraphIndex);
+#else
+				// Flush unlikely in standalone, so may be that cook failed to produce the compute graph.
+				UE_LOG(LogPCG, Error, TEXT("Failed to obtain compute graph for %s '%s', grid size %u, compute graph index %d. Either the graph compilation cache was flushed mid-execution, or the cook may have failed to produce the required compute graph. Compute graph execution aborted."),
+					DynamicSubgraph ? TEXT("dynamic subgraph") : TEXT("top graph"), DynamicSubgraph ? *DynamicSubgraph->GetName() : *TopGraph->GetName(), Context->GenerationGridSize, ComputeGraphIndex);
+#endif
+
 				return true;
 			}
 		}
