[PCG] Add additional title for GenerateLandscapeTexture preconfigured options

#rnx

#Synced-CL 48489707

diff --git a/Engine/Plugins/PCG/Source/PCG/Private/Elements/PCGGenerateLandscapeTextures.cpp b/Engine/Plugins/PCG/Source/PCG/Private/Elements/PCGGenerateLandscapeTextures.cpp
--- a/Engine/Plugins/PCG/Source/PCG/Private/Elements/PCGGenerateLandscapeTextures.cpp
+++ b/Engine/Plugins/PCG/Source/PCG/Private/Elements/PCGGenerateLandscapeTextures.cpp
@@ -81,6 +81,11 @@ TArray<FPCGPreConfiguredSettingsInfo> UPCGGenerateLandscapeTexturesSettings::Get
 }
 #endif
 
+FString UPCGGenerateLandscapeTexturesSettings::GetAdditionalTitleInformation() const
+{
+	return (bGenerateHeightMap ? LOCTEXT("GenerateGrassMapsNodeSubtitle_HeightAndGrass", "Height and Grass Maps") : LOCTEXT("GenerateGrassMapsNodeSubtitle_GrassOnly", "Grass Maps")).ToString();
+}
+
 void UPCGGenerateLandscapeTexturesSettings::ApplyPreconfiguredSettings(const FPCGPreConfiguredSettingsInfo& PreconfigureInfo)
 {
 	if (PreconfigureInfo.PreconfiguredIndex == 1)
diff --git a/Engine/Plugins/PCG/Source/PCG/Public/Elements/PCGGenerateLandscapeTextures.h b/Engine/Plugins/PCG/Source/PCG/Public/Elements/PCGGenerateLandscapeTextures.h
--- a/Engine/Plugins/PCG/Source/PCG/Public/Elements/PCGGenerateLandscapeTextures.h
+++ b/Engine/Plugins/PCG/Source/PCG/Public/Elements/PCGGenerateLandscapeTextures.h
@@ -35,6 +35,7 @@ public:
 	virtual bool OnlyExposePreconfiguredSettings() const override { return true; }
 	virtual bool GroupPreconfiguredSettings() const override { return false; }
 #endif
+	virtual FString GetAdditionalTitleInformation() const override;
 	virtual void ApplyPreconfiguredSettings(const FPCGPreConfiguredSettingsInfo& PreconfigureInfo) override;
 
 protected:
