[PCG] Remove data level attributes in static mesh spawner overrides. Fixed a small oversight in the write data index that wouldn't write when the attribute previously existed.

#rnx
#Synced-CL 48451591

diff --git a/Engine/Plugins/PCG/Source/PCG/Private/Elements/PCGWriteDataIndex.cpp b/Engine/Plugins/PCG/Source/PCG/Private/Elements/PCGWriteDataIndex.cpp
--- a/Engine/Plugins/PCG/Source/PCG/Private/Elements/PCGWriteDataIndex.cpp
+++ b/Engine/Plugins/PCG/Source/PCG/Private/Elements/PCGWriteDataIndex.cpp
@@ -56,6 +56,13 @@ bool FPCGWriteDataIndexElement::ExecuteInternal(FPCGContext* Context) const
 		{
 			UPCGData* NewData = Data->DuplicateData(Context);
 			FPCGMetadataDomain* IndexAttributeDomain = NewData ? NewData->MutableMetadata()->GetMetadataDomainFromSelector(Settings->IndexAttribute) : nullptr;
+
+			// Ensure attribute doesn't exist because we need to change the default value
+			if(IndexAttributeDomain && IndexAttributeDomain->HasAttribute(IndexAttributeName))
+			{
+				IndexAttributeDomain->DeleteAttribute(IndexAttributeName);
+			}
+
 			FPCGMetadataAttribute<int32>* IndexAttribute = IndexAttributeDomain ? IndexAttributeDomain->FindOrCreateAttribute<int32>(IndexAttributeName, DataIndex, /*bAllowInterpolation=*/false) : nullptr;
 
 			if (!IndexAttribute)
diff --git a/Engine/Plugins/PCG/Source/PCG/Private/MeshSelectors/PCGMeshSelectorByAttribute.cpp b/Engine/Plugins/PCG/Source/PCG/Private/MeshSelectors/PCGMeshSelectorByAttribute.cpp
--- a/Engine/Plugins/PCG/Source/PCG/Private/MeshSelectors/PCGMeshSelectorByAttribute.cpp
+++ b/Engine/Plugins/PCG/Source/PCG/Private/MeshSelectors/PCGMeshSelectorByAttribute.cpp
@@ -195,6 +195,16 @@ bool UPCGMeshSelectorByAttribute::SelectMeshInstances(
 		// Validate all the selectors are actual FPCGSoftISMComponentDescriptor properties
 		TArray<FPCGAttributePropertySelector> ValidSelectorOverrides = Overrides.GetAllInputSelectors();
 
+		// Remove all overrides that target data-level specifically
+		// @todo_pcg: warning - there`s an assumption that the override attributes aren't global attributes, which we really need to change.
+		for (int SelectorIndex = ValidSelectorOverrides.Num() - 1; SelectorIndex >= 0; --SelectorIndex)
+		{
+			if(InPointData->GetMetadataDomainIDFromSelector(ValidSelectorOverrides[SelectorIndex]) == PCGMetadataDomainID::Data)
+			{
+				ValidSelectorOverrides.RemoveAt(SelectorIndex, EAllowShrinking::No);
+			}
+		}
+
 		// If there are valid overrides, partition the points on those attributes so that an instance can be created for each
 		if (!ValidSelectorOverrides.IsEmpty())
 		{
