[PCG] Update editor even when changes are made by script and editor is closed

#rnx

#Synced-CL 48570258

diff --git a/Engine/Plugins/PCG/Source/PCGEditor/Private/PCGEditor.cpp b/Engine/Plugins/PCG/Source/PCGEditor/Private/PCGEditor.cpp
--- a/Engine/Plugins/PCG/Source/PCGEditor/Private/PCGEditor.cpp
+++ b/Engine/Plugins/PCG/Source/PCGEditor/Private/PCGEditor.cpp
@@ -3275,13 +3275,6 @@ void FPCGEditor::OnGraphChanged(UPCGGraphInterface* InGraph, EPCGChangeType Chan
 			}
 		}
 	}
-
-	if (!!(ChangeType & EPCGChangeType::ExternalModification))
-	{
-	    //Changes caused by external tools: Editor graph should be reconstructed.
-		PCGEditorGraph->ReconstructGraph();
-	}
-
 }
 
 void FPCGEditor::OnNodeSourceCompiled(const UPCGNode* InNode, const FPCGCompilerDiagnostics& InDiagnostics)
diff --git a/Engine/Plugins/PCG/Source/PCGEditor/Private/PCGEditorGraph.cpp b/Engine/Plugins/PCG/Source/PCGEditor/Private/PCGEditorGraph.cpp
--- a/Engine/Plugins/PCG/Source/PCGEditor/Private/PCGEditorGraph.cpp
+++ b/Engine/Plugins/PCG/Source/PCGEditor/Private/PCGEditorGraph.cpp
@@ -62,6 +62,7 @@ void UPCGEditorGraph::InitFromNodeGraph(UPCGGraph* InPCGGraph)
 	PCGGraph = InPCGGraph;
 
 	PCGGraph->OnGraphParametersChangedDelegate.AddUObject(this, &UPCGEditorGraph::OnGraphUserParametersChanged);
+	PCGGraph->OnGraphChangedDelegate.AddUObject(this, &UPCGEditorGraph::OnGraphChanged);
 
 	ReconstructGraph();
 }
@@ -150,18 +151,19 @@ void UPCGEditorGraph::BeginDestroy()
 	Super::BeginDestroy();
 
 	OnClose();
-}
-
-void UPCGEditorGraph::OnClose()
-{
-	ReplicateExtraNodes();
 
 	if (PCGGraph)
 	{
+		PCGGraph->OnGraphChangedDelegate.RemoveAll(this);
 		PCGGraph->OnGraphParametersChangedDelegate.RemoveAll(this);
 	}
 }
 
+void UPCGEditorGraph::OnClose()
+{
+	ReplicateExtraNodes();
+}
+
 void UPCGEditorGraph::CreateLinks(UPCGEditorGraphNodeBase* GraphNode, bool bCreateInbound, bool bCreateOutbound)
 {
 	check(GraphNode);
@@ -396,6 +398,15 @@ void UPCGEditorGraph::CreateLinks(UPCGEditorGraphNodeBase* GraphNode, bool bCrea
 	}
 }
 
+void UPCGEditorGraph::OnGraphChanged(UPCGGraphInterface* InGraph, EPCGChangeType ChangeType)
+{
+	if (!!(ChangeType & EPCGChangeType::ExternalModification))
+	{
+		//Changes caused by external tools: Editor graph should be reconstructed.
+		ReconstructGraph();
+	}
+}
+
 void UPCGEditorGraph::OnGraphUserParametersChanged(UPCGGraphInterface* InGraph, EPCGGraphParameterEvent ChangeType, FName ChangedPropertyName)
 {
 	if ((ChangeType != EPCGGraphParameterEvent::RemovedUnused && ChangeType != EPCGGraphParameterEvent::RemovedUsed) || InGraph != PCGGraph)
diff --git a/Engine/Plugins/PCG/Source/PCGEditor/Private/PCGEditorGraph.h b/Engine/Plugins/PCG/Source/PCGEditor/Private/PCGEditorGraph.h
--- a/Engine/Plugins/PCG/Source/PCGEditor/Private/PCGEditorGraph.h
+++ b/Engine/Plugins/PCG/Source/PCGEditor/Private/PCGEditorGraph.h
@@ -62,6 +62,7 @@ public:
 protected:
 	void CreateLinks(UPCGEditorGraphNodeBase* InGraphNode, bool bCreateInbound, bool bCreateOutbound, const TMap<UPCGNode*, UPCGEditorGraphNodeBase*>& InPCGNodeToPCGEditorNodeMap);
 
+	void OnGraphChanged(UPCGGraphInterface* InGraph, EPCGChangeType ChangeType);
 	void OnGraphUserParametersChanged(UPCGGraphInterface* InGraph, EPCGGraphParameterEvent ChangeType, FName ChangedPropertyName);
 
 	// ~Begin IPropertyBagEdGraph interface
