[PCG] Add explicit bounds checks to buffer loads/stores in data collection data interface.

[RN] [PCG] Fix cases of undefined behaviour and GPU crashes by adding protection against out of bounds reads/writes from Custom HLSL nodes.
#Synced-CL Unknown

diff --git a/Engine/Plugins/PCG/Shaders/Private/PCGDataCollectionDataInterface.ush b/Engine/Plugins/PCG/Shaders/Private/PCGDataCollectionDataInterface.ush
--- a/Engine/Plugins/PCG/Shaders/Private/PCGDataCollectionDataInterface.ush
+++ b/Engine/Plugins/PCG/Shaders/Private/PCGDataCollectionDataInterface.ush
@@ -22,6 +22,8 @@
 // @todo_pcg: Support having a non-RW byte address buffer for input data.
 RWByteAddressBuffer {DataInterfaceName}_DataCollectionBuffer;
 
+uint {DataInterfaceName}_BufferSizeBytes;
+
 // Multiplier that can be configured on output pins.
 uint {DataInterfaceName}_ElementCountMultiplier;
 
@@ -45,6 +47,16 @@ struct FPCGAttributeHeader_{DataInterfaceName}
 	bool bAllocated;
 };
 
+bool IsAddressValid_{DataInterfaceName}(uint InAddressBytes)
+{
+	return InAddressBytes < {DataInterfaceName}_BufferSizeBytes;
+}
+
+bool IsAddressValid_{DataInterfaceName}(uint InAddressBytes, uint InElementSizeBytes)
+{
+	return (InAddressBytes + InElementSizeBytes) <= {DataInterfaceName}_BufferSizeBytes;
+}
+
 // #################### SHADER PARAM GETTERS ##########################
 
 // @todo_pcg: This parameter is broader than this data interface. Needs to be shared in some way (e.g. with TextureDataInterface).
@@ -63,12 +75,24 @@ uint GetNumData_{DataInterfaceName}()
 
 uint GetNumElements_{DataInterfaceName}(uint InDataIndex)
 {
+	if (InDataIndex >= GetNumData_{DataInterfaceName}())
+	{
+		return 0u;
+	}
+
 	const uint ReadAddress =
 		PCG_DATA_COLLECTION_HEADER_SIZE_BYTES +
 		InDataIndex * PCG_DATA_HEADER_SIZE_BYTES +
 		/*TypeId*/4 + /*Attribute Count*/4;
 
-	return {DataInterfaceName}_DataCollectionBuffer.Load(ReadAddress);
+	if (IsAddressValid_{DataInterfaceName}(ReadAddress))
+	{
+		return {DataInterfaceName}_DataCollectionBuffer.Load(ReadAddress);
+	}
+	else
+	{
+		return 0u;
+	}
 }
 
 bool GetThreadData_{DataInterfaceName}(uint InThreadIndex, out uint OutDataIndex, out uint OutElementIndex)
@@ -112,12 +136,23 @@ uint GetNumElements_{DataInterfaceName}()
 
 uint LoadBufferInternal_{DataInterfaceName}(uint InAddress)
 {
-	return {DataInterfaceName}_DataCollectionBuffer.Load(InAddress);
+	// Full branch, not ternary, to avoid executing the load if address invalid.
+	if (IsAddressValid_{DataInterfaceName}(InAddress))
+	{
+		return {DataInterfaceName}_DataCollectionBuffer.Load(InAddress);
+	}
+	else
+	{
+		return 0u;
+	}
 }
 
 void StoreBufferInternal_{DataInterfaceName}(uint InAddress, uint InValue)
 {
-	{DataInterfaceName}_DataCollectionBuffer.Store(InAddress, InValue);
+	if (IsAddressValid_{DataInterfaceName}(InAddress))
+	{
+		{DataInterfaceName}_DataCollectionBuffer.Store(InAddress, InValue);
+	}
 }
 
 int GetRemappedAttributeIdInternal_{DataInterfaceName}(int InGraphAttributeId)
@@ -136,6 +171,11 @@ int GetRemappedAttributeIdInternal_{DataInterfaceName}(int InGraphAttributeId)
 
 FPCGAttributeHeader_{DataInterfaceName} GetAttributeHeaderInternal_{DataInterfaceName}(uint InDataIndex, int InAttributeId)
 {
+	if (InDataIndex >= GetNumData_{DataInterfaceName}())
+	{
+		return (FPCGAttributeHeader_{DataInterfaceName})0;
+	}
+
 	const int AttributeId = GetRemappedAttributeIdInternal_{DataInterfaceName}(InAttributeId);
 
 	const uint ReadAddress =
@@ -144,6 +184,11 @@ FPCGAttributeHeader_{DataInterfaceName} GetAttributeHeaderInternal_{DataInterfac
 		/*TypeId*/4 + /*Attribute Count*/4 + /*Element Count*/4 +
 		AttributeId * PCG_ATTRIBUTE_HEADER_SIZE_BYTES;
 
+	if (!IsAddressValid_{DataInterfaceName}(ReadAddress))
+	{
+		return (FPCGAttributeHeader_{DataInterfaceName})0;
+	}
+
 	const uint2 PackedAttributeHeader = {DataInterfaceName}_DataCollectionBuffer.Load2(ReadAddress);
 
 	FPCGAttributeHeader_{DataInterfaceName} AttributeHeader;
@@ -173,13 +218,13 @@ uint GetFirstValidAttributeElementInternal_{DataInterfaceName}()
 	{
 		const uint ElementAddress = GetElementAddressInternal_{DataInterfaceName}(/*DataIndex=*/0, /*ElementIndex=*/0, AttributeId);
 
-		if (ElementAddress > 0)
+		if (ElementAddress > 0u)
 		{
 			return ElementAddress;
 		}
 	}
 
-	return 0;
+	return 0u;
 }
 
 uint GetDataNumAttributesInternal_{DataInterfaceName}(uint InDataIndex)
@@ -189,7 +234,14 @@ uint GetDataNumAttributesInternal_{DataInterfaceName}(uint InDataIndex)
 		InDataIndex * PCG_DATA_HEADER_SIZE_BYTES +
 		/*TypeId*/4;
 
-	return {DataInterfaceName}_DataCollectionBuffer.Load(ReadAddress);
+	if (IsAddressValid_{DataInterfaceName}(ReadAddress))
+	{
+		return {DataInterfaceName}_DataCollectionBuffer.Load(ReadAddress);
+	}
+	else
+	{
+		return 0u;
+	}
 }
 
 void SetAsExecutedInternal_{DataInterfaceName}()
@@ -205,7 +257,7 @@ uint GetUint_{DataInterfaceName}(uint InDataIndex, uint InElementIndex, int InAt
 {
 	const uint ElementAddress = GetElementAddressInternal_{DataInterfaceName}(InDataIndex, InElementIndex, InAttributeId);
 
-	if (ElementAddress > 0)
+	if (ElementAddress > 0 && IsAddressValid_{DataInterfaceName}(ElementAddress))
 	{
 		return {DataInterfaceName}_DataCollectionBuffer.Load(ElementAddress);
 	}
@@ -234,7 +286,8 @@ float2 GetFloat2_{DataInterfaceName}(uint InDataIndex, uint InElementIndex, int
 {
 	const uint ElementAddress = GetElementAddressInternal_{DataInterfaceName}(InDataIndex, InElementIndex, InAttributeId);
 
-	if (ElementAddress > 0)
+	// 0 is invalid element address, and pen-ultimate word must lie within buffer.
+	if (ElementAddress > 0 && IsAddressValid_{DataInterfaceName}(ElementAddress, /*InElementSizeBytes=*/8u))
 	{
 		return asfloat({DataInterfaceName}_DataCollectionBuffer.Load2(ElementAddress));
 	}
@@ -248,7 +301,8 @@ float3 GetFloat3_{DataInterfaceName}(uint InDataIndex, uint InElementIndex, int
 {
 	const uint ElementAddress = GetElementAddressInternal_{DataInterfaceName}(InDataIndex, InElementIndex, InAttributeId);
 
-	if (ElementAddress > 0)
+	// 0 is invalid element address, and pen-ultimate word must lie within buffer.
+	if (ElementAddress > 0 && IsAddressValid_{DataInterfaceName}(ElementAddress, /*InElementSizeBytes=*/12u))
 	{
 		return asfloat({DataInterfaceName}_DataCollectionBuffer.Load3(ElementAddress));
 	}
@@ -262,7 +316,8 @@ float4 GetFloat4_{DataInterfaceName}(uint InDataIndex, uint InElementIndex, int
 {
 	const uint ElementAddress = GetElementAddressInternal_{DataInterfaceName}(InDataIndex, InElementIndex, InAttributeId);
 
-	if (ElementAddress > 0)
+	// 0 is invalid element address, and pen-ultimate word must lie within buffer.
+	if (ElementAddress > 0 && IsAddressValid_{DataInterfaceName}(ElementAddress, /*InElementSizeBytes=*/16u))
 	{
 		return asfloat({DataInterfaceName}_DataCollectionBuffer.Load4(ElementAddress));
 	}
@@ -286,7 +341,8 @@ float4x4 GetTransform_{DataInterfaceName}(uint InDataIndex, uint InElementIndex,
 {
 	const uint ElementAddress = GetElementAddressInternal_{DataInterfaceName}(InDataIndex, InElementIndex, InAttributeId);
 
-	if (ElementAddress > 0)
+	// 0 is invalid element address, and pen-ultimate word must lie within buffer.
+	if (ElementAddress > 0 && IsAddressValid_{DataInterfaceName}(ElementAddress, /*InElementSizeBytes=*/64u))
 	{
 		return float4x4(
 			asfloat({DataInterfaceName}_DataCollectionBuffer.Load4(ElementAddress)),
@@ -326,7 +382,7 @@ void SetBool_{DataInterfaceName}(uint InDataIndex, uint InElementIndex, int InAt
 {
 	const uint ElementAddress = GetElementAddressInternal_{DataInterfaceName}(InDataIndex, InElementIndex, InAttributeId);
 
-	if (ElementAddress > 0)
+	if (ElementAddress > 0 && IsAddressValid_{DataInterfaceName}(ElementAddress))
 	{
 		{DataInterfaceName}_DataCollectionBuffer.Store(ElementAddress, InValue ? 1u : 0u);
 	}
@@ -336,7 +392,7 @@ void SetInt_{DataInterfaceName}(uint InDataIndex, uint InElementIndex, int InAtt
 {
 	const uint ElementAddress = GetElementAddressInternal_{DataInterfaceName}(InDataIndex, InElementIndex, InAttributeId);
 
-	if (ElementAddress > 0)
+	if (ElementAddress > 0 && IsAddressValid_{DataInterfaceName}(ElementAddress))
 	{
 		{DataInterfaceName}_DataCollectionBuffer.Store(ElementAddress, asuint(InValue));
 	}
@@ -346,7 +402,7 @@ void SetUint_{DataInterfaceName}(uint InDataIndex, uint InElementIndex, int InAt
 {
 	const uint ElementAddress = GetElementAddressInternal_{DataInterfaceName}(InDataIndex, InElementIndex, InAttributeId);
 
-	if (ElementAddress > 0)
+	if (ElementAddress > 0 && IsAddressValid_{DataInterfaceName}(ElementAddress))
 	{
 		{DataInterfaceName}_DataCollectionBuffer.Store(ElementAddress, InValue);
 	}
@@ -356,7 +412,7 @@ void SetFloat_{DataInterfaceName}(uint InDataIndex, uint InElementIndex, int InA
 {
 	const uint ElementAddress = GetElementAddressInternal_{DataInterfaceName}(InDataIndex, InElementIndex, InAttributeId);
 
-	if (ElementAddress > 0)
+	if (ElementAddress > 0 && IsAddressValid_{DataInterfaceName}(ElementAddress))
 	{
 		{DataInterfaceName}_DataCollectionBuffer.Store(ElementAddress, asuint(InValue));
 	}
@@ -366,7 +422,7 @@ void SetFloat2_{DataInterfaceName}(uint InDataIndex, uint InElementIndex, int In
 {
 	const uint ElementAddress = GetElementAddressInternal_{DataInterfaceName}(InDataIndex, InElementIndex, InAttributeId);
 
-	if (ElementAddress > 0)
+	if (ElementAddress > 0 && IsAddressValid_{DataInterfaceName}(ElementAddress, /*InElementSizeBytes=*/8u))
 	{
 		{DataInterfaceName}_DataCollectionBuffer.Store2(ElementAddress, asuint(InValue));
 	}
@@ -376,7 +432,7 @@ void SetFloat3_{DataInterfaceName}(uint InDataIndex, uint InElementIndex, int In
 {
 	const uint ElementAddress = GetElementAddressInternal_{DataInterfaceName}(InDataIndex, InElementIndex, InAttributeId);
 
-	if (ElementAddress > 0)
+	if (ElementAddress > 0 && IsAddressValid_{DataInterfaceName}(ElementAddress, /*InElementSizeBytes=*/12u))
 	{
 		{DataInterfaceName}_DataCollectionBuffer.Store3(ElementAddress, asuint(InValue));
 	}
@@ -386,7 +442,7 @@ void SetFloat4_{DataInterfaceName}(uint InDataIndex, uint InElementIndex, int In
 {
 	const uint ElementAddress = GetElementAddressInternal_{DataInterfaceName}(InDataIndex, InElementIndex, InAttributeId);
 
-	if (ElementAddress > 0)
+	if (ElementAddress > 0 && IsAddressValid_{DataInterfaceName}(ElementAddress, /*InElementSizeBytes=*/16u))
 	{
 		{DataInterfaceName}_DataCollectionBuffer.Store4(ElementAddress, asuint(InValue));
 	}
@@ -394,29 +450,19 @@ void SetFloat4_{DataInterfaceName}(uint InDataIndex, uint InElementIndex, int In
 
 void SetRotator_{DataInterfaceName}(uint InDataIndex, uint InElementIndex, int InAttributeId, float3 InValue)
 {
-	const uint ElementAddress = GetElementAddressInternal_{DataInterfaceName}(InDataIndex, InElementIndex, InAttributeId);
-
-	if (ElementAddress > 0)
-	{
-		{DataInterfaceName}_DataCollectionBuffer.Store3(ElementAddress, asuint(InValue));
-	}
+	SetFloat3_{DataInterfaceName}(InDataIndex, InElementIndex, InAttributeId, InValue);
 }
 
 void SetQuat_{DataInterfaceName}(uint InDataIndex, uint InElementIndex, int InAttributeId, float4 InValue)
 {
-	const uint ElementAddress = GetElementAddressInternal_{DataInterfaceName}(InDataIndex, InElementIndex, InAttributeId);
-
-	if (ElementAddress > 0)
-	{
-		{DataInterfaceName}_DataCollectionBuffer.Store4(ElementAddress, asuint(InValue));
-	}
+	SetFloat4_{DataInterfaceName}(InDataIndex, InElementIndex, InAttributeId, InValue);
 }
 
 void SetTransform_{DataInterfaceName}(uint InDataIndex, uint InElementIndex, int InAttributeId, float4x4 InValue)
 {
 	const uint ElementAddress = GetElementAddressInternal_{DataInterfaceName}(InDataIndex, InElementIndex, InAttributeId);
 
-	if (ElementAddress > 0)
+	if (ElementAddress > 0 && IsAddressValid_{DataInterfaceName}(ElementAddress, /*InElementSizeBytes=*/64u))
 	{
 		{DataInterfaceName}_DataCollectionBuffer.Store4(ElementAddress, asuint(InValue[0]));
 		{DataInterfaceName}_DataCollectionBuffer.Store4(ElementAddress + 4 * 4, asuint(InValue[1]));
@@ -443,7 +489,7 @@ uint GetFirstAttributeAsUint_{DataInterfaceName}()
 {
 	const uint AttributeAddress = GetFirstValidAttributeElementInternal_{DataInterfaceName}();
 
-	if (AttributeAddress > 0)
+	if (AttributeAddress > 0 && IsAddressValid_{DataInterfaceName}(AttributeAddress))
 	{
 		return {DataInterfaceName}_DataCollectionBuffer.Load(AttributeAddress);
 	}
@@ -472,7 +518,7 @@ float2 GetFirstAttributeAsFloat2_{DataInterfaceName}()
 {
 	const uint AttributeAddress = GetFirstValidAttributeElementInternal_{DataInterfaceName}();
 
-	if (AttributeAddress > 0)
+	if (AttributeAddress > 0 && IsAddressValid_{DataInterfaceName}(AttributeAddress, /*InElementSizeBytes=*/8u))
 	{
 		return asfloat({DataInterfaceName}_DataCollectionBuffer.Load2(AttributeAddress));
 	}
@@ -486,7 +532,7 @@ float3 GetFirstAttributeAsFloat3_{DataInterfaceName}()
 {
 	const uint AttributeAddress = GetFirstValidAttributeElementInternal_{DataInterfaceName}();
 
-	if (AttributeAddress > 0)
+	if (AttributeAddress > 0 && IsAddressValid_{DataInterfaceName}(AttributeAddress, /*InElementSizeBytes=*/12u))
 	{
 		return asfloat({DataInterfaceName}_DataCollectionBuffer.Load3(AttributeAddress));
 	}
@@ -500,7 +546,7 @@ float4 GetFirstAttributeAsFloat4_{DataInterfaceName}()
 {
 	const uint AttributeAddress = GetFirstValidAttributeElementInternal_{DataInterfaceName}();
 
-	if (AttributeAddress > 0)
+	if (AttributeAddress > 0 && IsAddressValid_{DataInterfaceName}(AttributeAddress, /*InElementSizeBytes=*/16u))
 	{
 		return asfloat({DataInterfaceName}_DataCollectionBuffer.Load4(AttributeAddress));
 	}
@@ -523,7 +569,7 @@ float4 GetFirstAttributeAsQuat_{DataInterfaceName}()
 float4x4 GetFirstAttributeAsTransform_{DataInterfaceName}()
 {
 	const uint AttributeAddress = GetFirstValidAttributeElementInternal_{DataInterfaceName}();
-	if (AttributeAddress == 0)
+	if (AttributeAddress == 0 || !IsAddressValid_{DataInterfaceName}(AttributeAddress, /*InElementSizeBytes=*/64u))
 	{
 		return (float4x4)0;
 	}
@@ -555,7 +601,7 @@ int AtomicAddInt_{DataInterfaceName}(uint InDataIndex, uint InElementIndex, int
 {
 	const uint ElementAddress = GetElementAddressInternal_{DataInterfaceName}(InDataIndex, InElementIndex, InAttributeId);
 
-	if (ElementAddress > 0)
+	if (ElementAddress > 0 && IsAddressValid_{DataInterfaceName}(ElementAddress))
 	{
 		uint OriginalValue;
 		{DataInterfaceName}_DataCollectionBuffer.InterlockedAdd(ElementAddress, (uint)InValueToAdd, OriginalValue);
diff --git a/Engine/Plugins/PCG/Source/PCG/Private/Compute/DataInterfaces/PCGDataCollectionDataInterface.cpp b/Engine/Plugins/PCG/Source/PCG/Private/Compute/DataInterfaces/PCGDataCollectionDataInterface.cpp
--- a/Engine/Plugins/PCG/Source/PCG/Private/Compute/DataInterfaces/PCGDataCollectionDataInterface.cpp
+++ b/Engine/Plugins/PCG/Source/PCG/Private/Compute/DataInterfaces/PCGDataCollectionDataInterface.cpp
@@ -467,6 +467,7 @@ void UPCGDataCollectionDataInterface::GetSupportedInputs(TArray<FShaderFunctionD
 
 BEGIN_SHADER_PARAMETER_STRUCT(FPCGDataCollectionDataInterfaceParameters, )
 	SHADER_PARAMETER_RDG_BUFFER_UAV(RWByteAddressBuffer, DataCollectionBuffer)
+	SHADER_PARAMETER(uint32, BufferSizeBytes)
 	SHADER_PARAMETER(uint32, ElementCountMultiplier)
 	SHADER_PARAMETER_RDG_BUFFER_SRV(StructuredBuffer<int32>, GraphToBufferAttributeId)
 	SHADER_PARAMETER(int32, FirstRemappedAttributeId)
@@ -725,6 +726,7 @@ void FPCGDataCollectionDataProviderProxy::GatherDispatchData(FDispatchData const
 	{
 		FParameters& Parameters = ParameterArray[InvocationIndex];
 		Parameters.DataCollectionBuffer = BufferUAV;
+		Parameters.BufferSizeBytes = static_cast<uint32>(SizeBytes);
 		Parameters.ElementCountMultiplier = ElementCountMultiplier;
 		Parameters.GraphToBufferAttributeId = AttributeIdRemapSRV;
 		Parameters.FirstRemappedAttributeId = FirstRemappedAttributeId;
