[PCG] Refactor change tracking
- Move execution source tracking into FPCGChangeTracker which can be resued across execution source impls.
- Add FPCGTrackingRegistry in FPCGModule which allows registering custom  trackers
- Moved most specific code in FPCGTrackingManager to custom trackers (FPCGActorTracker, FPCGLandscapeTracker, FPCGMiscTracker)
- Moved handling of changes to custom handlers (FPCGComponentChangeHandler)

#Synced-CL 48606104

diff --git a/Engine/Plugins/PCG/Source/PCG/Internal/Grid/PCGExecutionSourceOctree.h b/Engine/Plugins/PCG/Source/PCG/Internal/Grid/PCGExecutionSourceOctree.h
new file mode 100644
--- /dev/null
+++ b/Engine/Plugins/PCG/Source/PCG/Internal/Grid/PCGExecutionSourceOctree.h
@@ -0,0 +1,95 @@
+// Copyright Epic Games, Inc. All Rights Reserved.
+
+#pragma once
+
+#include "Math/BoxSphereBounds.h"
+#include "Math/GenericOctree.h"
+#include "Misc/ScopeRWLock.h"
+#include "Misc/TransactionallySafeRWLock.h"
+#include "UObject/ObjectPtr.h"
+#include "Templates/SharedPointer.h"
+#include "Templates/SharedPointerFwd.h"
+
+template <typename ElementType, typename OctreeSemantics> class TOctree2;
+
+class IPCGGraphExecutionSource;
+
+struct FPCGExecutionSourceOctreeID : public TSharedFromThis<FPCGExecutionSourceOctreeID, ESPMode::ThreadSafe>
+{
+	FOctreeElementId2 Id;
+};
+
+using FPCGExecutionSourceOctreeIDSharedRef = TSharedRef<struct FPCGExecutionSourceOctreeID, ESPMode::ThreadSafe>;
+
+struct FPCGExecutionSourceRef
+{
+	FPCGExecutionSourceRef(IPCGGraphExecutionSource* InExecutionSource, const FPCGExecutionSourceOctreeIDSharedRef& InIdShared);
+
+	void UpdateBounds();
+
+	FPCGExecutionSourceOctreeIDSharedRef IdShared;
+	IPCGGraphExecutionSource* ExecutionSource = nullptr;
+	FBoxSphereBounds Bounds;
+};
+
+struct FPCGExecutionSourceRefSemantics
+{
+	enum { MaxElementsPerLeaf = 16 };
+	enum { MinInclusiveElementsPerNode = 7 };
+	enum { MaxNodeDepth = 12 };
+
+	typedef TInlineAllocator<MaxElementsPerLeaf> ElementAllocator;
+
+	inline static const FBoxSphereBounds& GetBoundingBox(const FPCGExecutionSourceRef& InVolume)
+	{
+		return InVolume.Bounds;
+	}
+
+	inline static const bool AreElementsEqual(const FPCGExecutionSourceRef& A, const FPCGExecutionSourceRef& B)
+	{
+		return A.ExecutionSource == B.ExecutionSource;
+	}
+
+	inline static void ApplyOffset(FPCGExecutionSourceRef& InVolume, const FVector& Offset)
+	{
+		InVolume.Bounds.Origin += Offset;
+	}
+
+	inline static void SetElementId(const FPCGExecutionSourceRef& Element, FOctreeElementId2 OctreeElementID)
+	{
+		Element.IdShared->Id = OctreeElementID;
+	}
+};
+
+using FPCGExecutionSourceOctree = TOctree2<FPCGExecutionSourceRef, FPCGExecutionSourceRefSemantics> ;
+using FPCGExecutionSourceToIdMap = TMap<IPCGGraphExecutionSource*, FPCGExecutionSourceOctreeIDSharedRef>;
+
+class FPCGExecutionSourceOctreeAndMap
+{
+public:
+	FPCGExecutionSourceOctreeAndMap() = default;
+	FPCGExecutionSourceOctreeAndMap(const FVector& InOrigin, FVector::FReal InExtent);
+
+	void Reset(const FVector& InOrigin, FVector::FReal InExtent);
+
+	TSet<IPCGGraphExecutionSource*> GetAllExecutionSources() const;
+
+	template<typename IterateBoundsFunc>
+	inline void FindElementsWithBoundsTest(const FBoxCenterAndExtent& BoxBounds, const IterateBoundsFunc& Func) const
+	{
+		UE::TReadScopeLock ReadLock(Lock);
+		return Octree.FindElementsWithBoundsTest(BoxBounds, Func);
+	}
+
+	bool Contains(const IPCGGraphExecutionSource* InExecutionSource) const;
+	FBox GetBounds(const IPCGGraphExecutionSource* InExecutionSource) const;
+
+	void AddOrUpdateExecutionSource(IPCGGraphExecutionSource* InExecutionSource, FBox& OutBounds, bool& bOutExecutionSourceHasChanged, bool& bOutExecutionSourcetWasAdded);
+	bool RemapExecutionSource(const IPCGGraphExecutionSource* InOldExecutionSource, IPCGGraphExecutionSource* InNewExecutionSource, bool& bOutBoundsHasChanged);
+	bool RemoveExecutionSource(IPCGGraphExecutionSource* InExecutionSource);
+
+private:
+	FPCGExecutionSourceOctree Octree;
+	FPCGExecutionSourceToIdMap ExecutionSourceToIdMap;
+	mutable FTransactionallySafeRWLock Lock;
+};
\ No newline at end of file
diff --git a/Engine/Plugins/PCG/Source/PCG/Internal/Landscape/PCGLandscapeCVars.h b/Engine/Plugins/PCG/Source/PCG/Internal/Landscape/PCGLandscapeCVars.h
new file mode 100644
--- /dev/null
+++ b/Engine/Plugins/PCG/Source/PCG/Internal/Landscape/PCGLandscapeCVars.h
@@ -0,0 +1,12 @@
+// Copyright Epic Games, Inc. All Rights Reserved.
+
+#pragma once
+
+#include "HAL/IConsoleManager.h"
+
+namespace PCGLandscapeCVars
+{
+	PCG_API extern TAutoConsoleVariable<int> CVarLandscapeRefreshTimeDelay;
+	PCG_API extern TAutoConsoleVariable<bool> CVarLandscapeDisableRefreshTracking;
+	PCG_API extern TAutoConsoleVariable<bool> CVarLandscapeDisableRefreshTrackingInLandscapeEditingMode;
+}
\ No newline at end of file
diff --git a/Engine/Plugins/PCG/Source/PCG/Internal/PCGTrackingManager.h b/Engine/Plugins/PCG/Source/PCG/Internal/PCGTrackingManager.h
new file mode 100644
--- /dev/null
+++ b/Engine/Plugins/PCG/Source/PCG/Internal/PCGTrackingManager.h
@@ -0,0 +1,258 @@
+// Copyright Epic Games, Inc. All Rights Reserved.
+
+#pragma once
+
+#include "PCGCommon.h"
+#include "Elements/PCGActorSelector.h"
+#include "Grid/PCGGridDescriptor.h"
+#include "Grid/PCGExecutionSourceOctree.h"
+#include "RuntimeGen/PCGRuntimeGenScheduler.h"
+
+#include "Containers/Array.h"
+#include "Containers/Map.h"
+#include "Containers/Set.h"
+#include "Templates/Function.h"
+#include "UObject/ObjectKey.h"
+#include "UObject/ObjectPtr.h"
+#include "UObject/ObjectSaveContext.h"
+#include "UObject/SoftObjectPtr.h"
+#include "UObject/WeakObjectPtr.h"
+
+class AActor;
+class ALandscapeProxy;
+class APCGPartitionActor;
+class FLandscapeProxyComponentDataChangedParams;
+class ILevelInstanceInterface;
+class UObject;
+class UPCGComponent;
+class UPCGGraph;
+class UPCGSubsystem;
+class UWorld;
+class UWorldPartition;
+class FWorldPartitionActorDescInstance;
+class FPackageReloadedEvent;
+
+enum class EPackageReloadPhase : uint8;
+
+#if WITH_EDITOR
+namespace PCGTrackingManager
+{
+	static const FPCGTrackingChangeID ObjectChanged(0x8C0A272F, 0xF1854FD0, 0xB821301B, 0x99D38D8D);
+	static const FPCGTrackingChangeID ExecutionSourceGenerated(0xA73B5462, 0x3CB24DB2, 0x9E9CADFE, 0xEE346EB5);
+}
+#endif
+
+/**
+* This class currently serves two purposes:
+* - Keep a mapping between original and local execution sources.
+* - Do change tracking to refresh execution sources.
+*
+* Its meant to be part of the PCG Subsystem and owned by it. We offload some logic to this class to avoid to clutter the subsystem.
+*/
+class FPCGTrackingManager
+{
+public:
+	friend UPCGSubsystem;
+	friend FPCGRuntimeGenScheduler;
+
+	~FPCGTrackingManager() = default;
+
+	/** Initializes callbacks, etc, tied to the PCG subsystem */
+	void Initialize();
+
+	/** Deinitializes callbacks, etc, tied to the PCG subsystem */
+	void Deinitialize();
+	
+	/** Should be called by the subsystem to handle delayed operations. */
+	void Tick();
+
+	/** Returns the PCGSubystem's World */
+	PCG_API UWorld* GetWorld() const;
+
+#if WITH_EDITOR
+	PCG_API bool IsTracking() const;
+	PCG_API bool IsKeyTracked(const FPCGSelectionKey& InKey) const;
+	PCG_API void ForEachTrackedKey(TFunctionRef<bool(const FPCGSelectionKey&, const TSet<IPCGGraphExecutionSource*>&)> InCallback) const;
+	/** Gather all settings from a given execution source that track the key, and clear the cache for them. Returns true if we should dirty afterwards (aka at least one settings was cleared and/or landscape changed). */
+	PCG_API bool ClearCacheForKeys(TArrayView<FPCGSelectionKey> InKeys, const IPCGGraphExecutionSource* InExecutionSource, const bool bIntersect, const UObject* InOriginatingChange) const;
+
+	PCG_API void OnSelectionKeyChanged(const FPCGSelectionKey& InSelectionKey);
+	PCG_API void OnSelectionKeyChanged(const FPCGSelectionKey& InSelectionKey, const UObject* InOriginatingChangeObject, TArrayView<const FBox> InBounds);
+	PCG_API void OnSelectionKeyChanged(const FPCGSelectionKey& InSelectionKey, const UObject* InOriginatingChangeObject, const TSet<FName>& InRemovedTags, bool bInNoRefreshOwner, TArrayView<const FBox> InBounds);
+
+	PCG_API void NotifyObjectChanged(const UObject* InObject, const FPCGTrackingChangeID& InChangeType, const UObject* InOriginatingChangeObject = nullptr);
+
+	static void BuildPartitionActorRecords(APCGWorldActor* PCGWorldActor, UWorldPartition* WorldPartition, TMap<FPCGGridCellDescriptor, FGuid>& OutPartitionActorRecords, TSet<FGuid>& OutInvalidPartitionActors);
+
+	/** Returns true if originating execution source is null or if dependent execution source has no more dependencies */
+	PCG_API bool RemoveExecutionSourceDependency(const IPCGGraphExecutionSource* InOriginatingExecutionSource, const IPCGGraphExecutionSource* InDependentExecutionSource);
+#endif // WITH_EDITOR
+
+	/** Iterate other all the int coordinates given a box and call a callback. Thread safe */
+	PCG_API void ForAllIntersectingPartitionActors(const FBox& InBounds, TFunctionRef<void(APCGPartitionActor*)> InFunc) const;
+
+	/** Register a new execution source or update it. Returns true if it was added/updated. Thread safe */
+	bool RegisterOrUpdateExecutionSource(IPCGGraphExecutionSource* InExecutionSource, bool bDoPartitionMapping = true);
+
+	/** In case of blueprints, we need to remap the old execution source destroyed by the construction script to the new one. Returns true if a re-mapping was done. Thread safe */
+	bool RemapExecutionSource(const IPCGGraphExecutionSource* OldExecutionSource, IPCGGraphExecutionSource* NewExecutionSource, bool bDoPartitionMapping);
+
+	/** Unregister a execution source. Can force it, if we have a delayed unregister. Thread safe */
+	void UnregisterExecutionSource(IPCGGraphExecutionSource* InExecutionSource, bool bForce = false);
+
+	/** Register a new Partition actor, will be added to a map and will query all intersecting volume to bind to them if asked. Thread safe */
+	void RegisterPartitionActor(APCGPartitionActor* InActor);
+
+	/** Unregister a Partition actor, will be removed from the map and remove itself to all intersecting volumes. Thread safe */
+	void UnregisterPartitionActor(APCGPartitionActor* InActor);
+
+	/** Return a copy of all the registered partitioned execution sources. Thread safe */
+	TSet<IPCGGraphExecutionSource*> GetAllRegisteredPartitionedExecutionSources() const;
+
+	/** Return a copy of all the registered non-partitioned execution sources. Thread safe */
+	TSet<IPCGGraphExecutionSource*> GetAllRegisteredNonPartitionedExecutionSources() const;
+
+	/** Return a copy of all the registered execution sources. Thread safe */
+	TSet<IPCGGraphExecutionSource*> GetAllRegisteredExecutionSources() const;
+
+	/** Retrieves a local component using grid descriptor and grid coordinates, returns nullptr if no such component is found. */
+	UPCGComponent* GetLocalComponent(const FPCGGridDescriptor& GridDescriptor, const FIntVector& CellCoords, const UPCGComponent* InOriginalComponent) const;
+
+	APCGPartitionActor* GetPartitionActor(const FPCGGridDescriptor& GridDescriptor, const FIntVector& CellCoords) const;
+
+#if WITH_EDITOR
+	/** Returns true if there is record of a partition actor living in a certain grid cell, regardless of whether or not it is loaded. */
+	bool DoesPartitionActorRecordExist(const FPCGGridDescriptor& GridDescriptor, const FIntVector& GridCoords) const;
+#endif
+private:
+
+#if WITH_EDITOR
+	void RegisterDelegates();
+	void UnregisterDelegates();
+
+	/** If the partition grid size change, call this to empty the Partition actors map */
+	void ResetPartitionActorsMap();
+
+	void RegisterTracking(IPCGGraphExecutionSource* InExecutionSource);
+	void UpdateTracking(IPCGGraphExecutionSource* InExecutionSource, const TArray<FPCGSelectionKey>* OptionalChangedKeys = nullptr);
+#endif
+	bool RegisterOrUpdateExecutionSource(IPCGGraphExecutionSource* InExecutionSource, bool bDoPartitionMapping, bool bForce);
+
+	// This class is only meant to be used as part of the PCG Subsytem and owned by it.
+	// So we put constructors private.
+	// We also need this class to be default constructible, since the PCGSubsytem needs to be default constructible.
+	FPCGTrackingManager() = default;
+	explicit FPCGTrackingManager(UPCGSubsystem* PCGSubsystem);
+
+	bool RegisterOrUpdatePartitionedExecutionSource(IPCGGraphExecutionSource* InExecutionSource, bool bDoPartitionMapping = true);
+	bool RegisterOrUpdateNonPartitionedExecutionSource(IPCGGraphExecutionSource* InExecutionSource);
+
+	void UnregisterPartitionedExecutionSource(IPCGGraphExecutionSource* InExecutionSource);
+	void UnregisterNonPartitionedExecutionSource(IPCGGraphExecutionSource* InExecutionSource);
+
+	/* Call the InFunc function to all local component registered to the original component. Return the list of all the tasks scheduled. Thread safe*/
+	TArray<FPCGTaskId> DispatchToRegisteredLocalComponents(UPCGComponent* OriginalComponent, const TFunctionRef<FPCGTaskId(UPCGComponent*, const TArray<FPCGTaskId>&)>& InFunc) const;
+
+	/* Call the InFunc function to all local component from the set of partition actors. Return the list of all the tasks scheduled. */
+	TArray<FPCGTaskId> DispatchToLocalComponents(UPCGComponent* OriginalComponent, const TSet<TObjectPtr<APCGPartitionActor>>& PartitionActors, const TFunctionRef<FPCGTaskId(UPCGComponent*, const TArray<FPCGTaskId>&)>& InFunc) const;
+
+	/** Call the InFunc function to all partitioned execution sources which bounds intersect 'InBounds'. */
+	void ForAllIntersectingPartitionedExecutionSources(const FBoxCenterAndExtent& InBounds, TFunctionRef<void(IPCGGraphExecutionSource*)> InFunc) const;
+
+	/** Call the InFunc function for all original execution sources (regardless of partitioned or not). */
+	void ForAllOriginalExecutionSources(TFunctionRef<void(IPCGGraphExecutionSource*)> InFunc);
+
+	/** Gather all the execution sources within some bounds. */
+	TArray<IPCGGraphExecutionSource*> GetAllIntersectingExecutionSources(const FBoxCenterAndExtent& InBounds) const;
+
+	/** Update the current mapping between a PCG component and its PCG Partition actors */
+	void UpdateMappingPCGComponentPartitionActor(UPCGComponent* InComponent, bool bInChangeGraphInstances = true);
+
+	/** Returns the current mapping between a PCG component and its PCG Partition actors */
+	TSet<TObjectPtr<APCGPartitionActor>> GetPCGComponentPartitionActorMappings(UPCGComponent* InComponent) const;
+	
+	/** Delete the current mapping between a PCG component and its PCG Partition actors */
+	void DeleteMappingPCGComponentPartitionActor(UPCGComponent* InComponent);
+
+	/** Returns true if an executon source is registered */
+	bool IsExecutionSourceRegistered(const IPCGGraphExecutionSource* InExecutionSource) const;
+
+	/** Return true if there are any Original or Non-Partitioned execution sources set to GenerateAtRuntime. */
+	bool AnyRuntimeGenExecutionSourcesExist() const;
+
+#if WITH_EDITOR
+	void OnPackageReloaded(const EPackageReloadPhase InPackageReloadPhase, FPackageReloadedEvent* InPackageReloadedEvent);
+	void OnActorDescInstanceAdded(FWorldPartitionActorDescInstance* InActorDescInstance);
+	void OnActorDescInstanceRemoved(FWorldPartitionActorDescInstance* InActorDescInstance);
+	void OnObjectPropertyChanged(UObject* InObject, FPropertyChangedEvent& InEvent);
+	void OnPCGGraphStartsGenerating(IPCGGraphExecutionSource* InExecutionSource);
+	void OnPCGGraphGeneratedOrCleaned(IPCGGraphExecutionSource* InExecutionSource);
+	void OnPCGGraphCancelled(IPCGGraphExecutionSource* InExecutionSource);
+
+	/** Remap the tracking in case of blueprints. */
+	void RemapTracking(const IPCGGraphExecutionSource* InOldExecutionSource, IPCGGraphExecutionSource* InNewExecutionSource);
+
+	/** Unregister tracking when a execution source is removed */
+	void UnregisterTracking(IPCGGraphExecutionSource* InExecutionSource);
+
+	/** Unregister tracking when a execution source is removed or keys are unregistered */
+	void UnregisterTracking(IPCGGraphExecutionSource* InExecutionSource, const TSet<FPCGSelectionKey>* OptionalKeysToUntrack);
+
+	void OnObjectChanged(const UObject* InObject, const FPCGTrackingChangeID& InChangeType, const UObject* InOriginatingChangeObject = nullptr);
+
+	/** Build initial records for existing partition actors */
+	void BuildPartitionActorRecords();
+
+	/** Used by the PCG Builder to force dispatched local components to be dependency chained together so that they do not run at the same time for determinism */
+	void SetChainedDispatchToLocalComponents(bool bInChainedDispatch) { bChainedDispatchToLocalComponents = bInChainedDispatch; }
+#endif // WITH_EDITOR
+
+private:
+	// Cached subsystem
+	UPCGSubsystem* PCGSubsystem = nullptr;
+
+	// Mapping Component <-> PartitionActor
+	/** Octree tracking all partitioned pcg components */
+	FPCGExecutionSourceOctreeAndMap PartitionedOctree;
+
+	/** Mapping from grid size and grid coords to partition actor. We can only have 1 partition actor per grid cell. */
+	TMap<FPCGGridDescriptor, TMap<FIntVector, TObjectPtr<APCGPartitionActor>>> PartitionActorsMap;
+	mutable PCG::FSharedLock PartitionActorsMapLock;
+
+	/** Mapping between original components and its overlapping partition actors. */
+	TMap<const UPCGComponent*, TSet<TObjectPtr<APCGPartitionActor>>> ComponentToPartitionActorsMap;
+	mutable PCG::FSharedLock ComponentToPartitionActorsMapLock;
+
+	/** Will hold all the components that are not partitioned (and not local) and are tracking something. Will be use to dispatch actor tracking updates. */
+	FPCGExecutionSourceOctreeAndMap NonPartitionedOctree;
+
+#if WITH_EDITOR
+	TArray<TUniquePtr<IPCGChangeTracker>> ChangeTrackers;
+	TArray<TUniquePtr<IPCGChangeHandler>> ChangeHandlers;
+
+	/** Components to be registered/unregistered on the next frame. see RegisterOrUpdatePCGComponent/UnregisterPCGComponent for a better understanding on why it is needed. */
+	mutable PCG::FLock DelayedComponentsRegistrationLock;
+	TSet<UPCGComponent*> DelayedComponentsToUnregister;
+	TSet<UPCGComponent*> DelayedComponentsToRegister;
+	
+	// Tracking actors
+	/** Keep a mapping between tracked keys and the execution sources that track them, and the tracking needs to be culled.*/
+	TMap<FPCGSelectionKey, TSet<IPCGGraphExecutionSource*>> CulledTrackedKeysToExecutionSourcesMap;
+
+	/** Same mapping but for always tracked keys */
+	TMap<FPCGSelectionKey, TSet<IPCGGraphExecutionSource*>> AlwaysTrackedKeysToExecutionSourcesMap;
+
+	mutable PCG::FSharedLock TrackedComponentsLock;
+
+	/** Transient map that keep track of all execution source that depend on other execution sources component currently generating, to trigger the refresh only once, when all are done. */
+	TMap<FObjectKey, TArray<FObjectKey>> ExecutionSourcesToDependencyMap;
+
+	// Set of existing PCG Partition Actors (for World Partition worlds)
+	TMap<FPCGGridCellDescriptor, FGuid> PartitionActorRecords;
+	// Previously generated PCG Partition Actors to ignore (will prevent them from getting registered and mark them for deletion)
+	TSet<FGuid> InvalidPartitionActors;
+
+	// If dispatch needs to run the local comopnents one after the other through dependency chain for determinism
+	bool bChainedDispatchToLocalComponents = false;
+#endif // WITH_EDITOR
+};
\ No newline at end of file
diff --git a/Engine/Plugins/PCG/Source/PCG/Internal/Utils/PCGGraphExecutionLogging.h b/Engine/Plugins/PCG/Source/PCG/Internal/Utils/PCGGraphExecutionLogging.h
new file mode 100644
--- /dev/null
+++ b/Engine/Plugins/PCG/Source/PCG/Internal/Utils/PCGGraphExecutionLogging.h
@@ -0,0 +1,58 @@
+// Copyright Epic Games, Inc. All Rights Reserved.
+
+#pragma once
+
+#include "PCGCommon.h"
+
+#include "Containers/UnrealString.h"
+#include "HAL/Platform.h"
+
+class IPCGGraphExecutionSource;
+class UPCGComponent;
+class UPCGGraph;
+enum class EPCGHiGenGrid : uint32;
+enum class EPCGIgnoreChangeOriginReason : uint8;
+struct FPCGContext;
+struct FPCGGraphActiveTask;
+struct FPCGGraphTask;
+struct FPCGPinDependencyExpression;
+struct FPCGStack;
+
+namespace PCGGraphExecutionLogging
+{
+	bool LogEnabled();
+	bool CullingLogEnabled();
+
+	void LogGraphTask(FPCGTaskId TaskId, const FPCGGraphTask& Task, const TSet<FPCGTaskId>* SuccessorIds = nullptr);
+	void LogGraphTasks(const TMap<FPCGTaskId, FPCGGraphTask>& Tasks, const TMap<FPCGTaskId, TSet<FPCGTaskId>>* TaskSuccessors = nullptr);
+	void LogGraphTasks(const TArray<FPCGGraphTask>& Tasks);
+
+	void LogGraphSchedule(const IPCGGraphExecutionSource* InExecutionSource, const UPCGGraph* InScheduledGraph);
+	void LogGraphScheduleDependency(const IPCGGraphExecutionSource* InExecutionSource, const FPCGStack* InFromStack);
+	void LogGraphScheduleDependencyFailed(const IPCGGraphExecutionSource* InExecutionSource, const FPCGStack* InFromStack);
+
+	void LogGraphPostSchedule(const TMap<FPCGTaskId, FPCGGraphTask>& Tasks, const TMap<FPCGTaskId, TSet<FPCGTaskId>>& TaskSuccessors);
+	void LogPostProcessGraph(const IPCGGraphExecutionSource* InExecutionSource);
+
+	void LogExecutionSourceCancellation(const TSet<IPCGGraphExecutionSource*>& CancelledComponents);
+
+	PCG_API void LogChangeOriginIgnoredForComponent(const UObject* InObject, const IPCGGraphExecutionSource* InExecutionSource, EPCGIgnoreChangeOriginReason InReason);
+
+	void LogGraphExecuteFrameFinished();
+
+	void LogTaskExecute(const FPCGGraphTask& Task);
+	void LogTaskExecuteCachingDisabled(const FPCGGraphTask& Task);
+
+	void LogTaskCullingBegin(FPCGTaskId CompletedTaskId, uint64 InactiveOutputPinBitmask, const TArray<FPCGPinId>& PinIdsToDeactivate);
+	void LogTaskCullingBeginLoop(FPCGTaskId PinTaskId, uint64 PinIndex, const TArray<FPCGPinId>& PinIdsToDeactivate);
+	void LogTaskCullingUpdatedPinDeps(FPCGTaskId TaskId, const FPCGPinDependencyExpression& PinDependency, bool bDependencyExpressionBecameFalse);
+
+	void LogGridLinkageTaskExecuteStore(const FPCGContext* InContext, EPCGHiGenGrid InGenerationGrid, int32 InFromLOD, int32 InToLOD, const FString& InResourcePath, int32 InDataItemCount);
+	void LogGridLinkageTaskExecuteRetrieve(const FPCGContext* InContext, EPCGHiGenGrid InGenerationGrid, int32 InFromLOD, int32 InToLOD, const FString& InResourcePath);
+	void LogGridLinkageTaskExecuteRetrieveSuccess(const FPCGContext* InContext, const IPCGGraphExecutionSource* InExecutionSource, const FString& InResourcePath, int32 InDataItemCount);
+	void LogGridLinkageTaskExecuteRetrieveScheduleGraph(const FPCGContext* InContext, const IPCGGraphExecutionSource* InScheduledSource, const FString& InResourcePath);
+	void LogGridLinkageTaskExecuteRetrieveWaitOnScheduledGraph(const FPCGContext* InContext, const IPCGGraphExecutionSource* InWaitOnSource, const FString& InResourcePath);
+	void LogGridLinkageTaskExecuteRetrieveWakeUp(const FPCGContext* InContext, const IPCGGraphExecutionSource* InWokenBySource);
+	void LogGridLinkageTaskExecuteRetrieveNoLocalSource(const FPCGContext* InContext, const FString& InResourcePath);
+	void LogGridLinkageTaskExecuteRetrieveNoData(const FPCGContext* InContext, const IPCGGraphExecutionSource* InExecutionSource, const FString& InResourcePath);
+}
diff --git a/Engine/Plugins/PCG/Source/PCG/Private/ChangeTracking/PCGChangeTracker.cpp b/Engine/Plugins/PCG/Source/PCG/Private/ChangeTracking/PCGChangeTracker.cpp
new file mode 100644
--- /dev/null
+++ b/Engine/Plugins/PCG/Source/PCG/Private/ChangeTracking/PCGChangeTracker.cpp
@@ -0,0 +1,300 @@
+// Copyright Epic Games, Inc. All Rights Reserved.
+
+#include "ChangeTracking/PCGChangeTracker.h"
+
+#if WITH_EDITOR
+
+#include "PCGGraph.h"
+#include "PCGGraphExecutionStateInterface.h"
+#include "Subsystems/PCGSubsystem.h"
+
+void FPCGChangeTracker::ResetCurrentExecution()
+{
+	PCG::TScopeLock Lock(CurrentExecutionDynamicTrackingLock);
+	CurrentExecutionDynamicTracking.Empty();
+	CurrentExecutionDynamicTrackingSettings.Empty();
+}
+
+TArray<FPCGSelectionKey> FPCGChangeTracker::GatherTrackingKeys() const
+{
+	TArray<FPCGSelectionKey> Keys;
+	Keys.Reserve(StaticallyTrackedKeysToSettings.Num() + DynamicallyTrackedKeysToSettings.Num());
+	for (const auto& It : StaticallyTrackedKeysToSettings)
+	{
+		Keys.Add(It.Key);
+	}
+
+	for (const auto& It : DynamicallyTrackedKeysToSettings)
+	{
+		Keys.Add(It.Key);
+	}
+
+	return Keys;
+}
+
+bool FPCGChangeTracker::IsKeyTrackedAndCulled(const FPCGSelectionKey& Key, bool& bOutIsCulled) const
+{
+	bool bIsTracked = false;
+
+	bool bStaticallyCulled = true;
+	bool bDynamicallyCulled = true;
+
+	if (auto* It = StaticallyTrackedKeysToSettings.Find(Key))
+	{
+		bIsTracked = true;
+		bStaticallyCulled = PCGSettings::IsKeyCulled(*It);
+	}
+
+	if (auto* It = DynamicallyTrackedKeysToSettings.Find(Key))
+	{
+		bIsTracked = true;
+		bDynamicallyCulled = PCGSettings::IsKeyCulled(*It);
+	}
+
+	// If it is tracked statically and dynamically, we will cull only and only if both are culling.
+	// Otherwise, it means that at least one key requires to always track, so bOutIsCulled needs to be False.
+	bOutIsCulled = bIsTracked && bStaticallyCulled && bDynamicallyCulled;
+
+	return bIsTracked;
+}
+
+bool FPCGChangeTracker::UpdateTrackingCache(const UPCGGraph* InGraph, TArray<FPCGSelectionKey>* OptionalChangedKeys)
+{
+	FPCGSelectionKeyToSettingsMap NewTrackedKeysToSettings;
+
+	const int32 PreviousKeyCount = StaticallyTrackedKeysToSettings.Num();
+
+	bool bHasChanged = false;
+
+	if (InGraph)
+	{
+		bHasChanged = InGraph->UpdateTrackedSelectionKeys(StaticallyTrackedKeysToSettings, NewTrackedKeysToSettings, OptionalChangedKeys);
+	}
+	else
+	{
+		// Handle case were we no longer have a graph
+		bHasChanged = StaticallyTrackedKeysToSettings.Num() > 0;
+		if (OptionalChangedKeys)
+		{
+			OptionalChangedKeys->Reserve(StaticallyTrackedKeysToSettings.Num());
+			for (const TPair<FPCGSelectionKey, TArray<FPCGSettingsAndCulling>>& It : StaticallyTrackedKeysToSettings)
+			{
+				OptionalChangedKeys->Add(It.Key);
+			}
+		}
+	}
+
+	StaticallyTrackedKeysToSettings = MoveTemp(NewTrackedKeysToSettings);
+
+	return bHasChanged;
+}
+
+void FPCGChangeTracker::ForEachSettingTrackingKey(const FPCGSelectionKey& InKey, const TFunctionRef<void(const FPCGSelectionKey&, const FPCGSettingsAndCulling&)> InCallback) const
+{
+	TRACE_CPUPROFILER_EVENT_SCOPE(FPCGChangeTracker::ForEachSettingTrackingKey);
+
+	auto FindAndApplyInMap = [&InKey, &InCallback](const FPCGSelectionKeyToSettingsMap& InMap)
+	{
+		if (const TArray<FPCGSettingsAndCulling>* StaticallyTrackedSettings = InMap.Find(InKey))
+		{
+			for (const FPCGSettingsAndCulling& SettingsAndCulling : *StaticallyTrackedSettings)
+			{
+				InCallback(InKey, SettingsAndCulling);
+			}
+		}
+	};
+
+	FindAndApplyInMap(StaticallyTrackedKeysToSettings);
+	FindAndApplyInMap(DynamicallyTrackedKeysToSettings);
+}
+
+void FPCGChangeTracker::RegisterDynamicTracking(const UPCGSettings* InSettings, const TArrayView<TPair<FPCGSelectionKey, bool>>& InDynamicKeysAndCulling)
+{
+	if (!InSettings)
+	{
+		return;
+	}
+
+	UE::TScopeLock Lock(CurrentExecutionDynamicTrackingLock);
+	CurrentExecutionDynamicTrackingSettings.Add(InSettings);
+
+	for (const TPair<FPCGSelectionKey, bool>& It : InDynamicKeysAndCulling)
+	{
+		// Make sure to not register null assets
+		if (It.Key.Selection == EPCGActorSelection::ByPath && It.Key.ObjectPath.IsNull())
+		{
+			continue;
+		}
+
+		CurrentExecutionDynamicTracking.FindOrAdd(It.Key).Emplace(InSettings, It.Value);
+	}
+}
+
+void FPCGChangeTracker::RegisterDynamicTracking(const FPCGSelectionKeyToSettingsMap& InKeysToSettings)
+{
+	if (InKeysToSettings.IsEmpty())
+	{
+		return;
+	}
+
+	UE::TScopeLock Lock(CurrentExecutionDynamicTrackingLock);
+
+	CurrentExecutionDynamicTracking.Append(InKeysToSettings);
+
+	for (const auto& It : InKeysToSettings)
+	{
+		Algo::Transform(It.Value, CurrentExecutionDynamicTrackingSettings, [](const FPCGSettingsAndCulling& SettingsAndCulling) { return SettingsAndCulling.Key.Get(); });
+	}
+}
+
+void FPCGChangeTracker::UpdateDynamicTracking(IPCGGraphExecutionSource* InExecutionSource)
+{
+	TRACE_CPUPROFILER_EVENT_SCOPE(FPCGChangeTracker::UpdateDynamicTracking);
+	check(InExecutionSource);
+
+	UPCGSubsystem* Subsystem = Cast<UPCGSubsystem>(InExecutionSource->GetExecutionState().GetSubsystem());
+	if (!Subsystem)
+	{
+		return;
+	}
+
+	IPCGGraphExecutionSource* OriginalSource = InExecutionSource->GetExecutionState().GetOriginalSource();
+	if (!OriginalSource)
+	{
+		return;
+	}
+
+	const bool bIsLocal = InExecutionSource->GetExecutionState().IsLocalSource();
+	const bool bIsRuntime = InExecutionSource->GetExecutionState().IsManagedByRuntimeGenSystem();
+
+	/*if (bIsRuntime && PCGComponent::CVarDisableDynamicTrackingForRuntimeGen.GetValueOnAnyThread())
+	{
+		return;
+	}*/
+
+	// If the component is local, we defer the tracking to the original component.
+	// So move everything to the original (while making sure we are not duplicating keys/settings).
+	// Since it can happen in parallel, we need to lock.
+	if (bIsLocal && !bIsRuntime)
+	{
+		TRACE_CPUPROFILER_EVENT_SCOPE(FPCGChangeTracker::UpdateDynamicTracking::LocalComponent);
+		
+		OriginalSource->GetExecutionState().AddLocalSourceCurrentTrackingKeys(CurrentExecutionDynamicTracking, CurrentExecutionDynamicTrackingSettings);
+
+		CurrentExecutionDynamicTracking.Empty();
+		CurrentExecutionDynamicTrackingSettings.Empty();
+
+		return;
+	}
+
+	TArray<FPCGSelectionKey> ChangedKeys;
+
+	// Locking to make sure we never hit this multiple times.
+	{
+		UE::TScopeLock Lock(CurrentExecutionDynamicTrackingLock);
+
+		// Go over all dynamic keys gathered during this execution.
+		// If they are not already tracked, we need to register this key.
+		// Otherwise, we need to gather all the settings that tracked this key that were not executed (because of caching).
+		for (auto& It : CurrentExecutionDynamicTracking)
+		{
+			if (TArray<FPCGSettingsAndCulling>* AllSettingsAndCulling = DynamicallyTrackedKeysToSettings.Find(It.Key))
+			{
+				for (FPCGSettingsAndCulling& SettingsAndCulling : *AllSettingsAndCulling)
+				{
+					if (SettingsAndCulling.Key.Get() && !CurrentExecutionDynamicTrackingSettings.Contains(SettingsAndCulling.Key.Get()))
+					{
+						It.Value.AddUnique(std::move(SettingsAndCulling));
+					}
+				}
+			}
+			else
+			{
+				ChangedKeys.Add(It.Key);
+			}
+		}
+
+		// Go over all already registered dynamic keys
+		// If they are not in the current execution gathered keys, we check if they are associated with settings that 
+		// were executed. If so, we re-add them to the current execution gathered keys.
+		// If not, it means that the key is no longer tracked and should be unregistered.
+		for (auto& It : DynamicallyTrackedKeysToSettings)
+		{
+			if (!CurrentExecutionDynamicTracking.Contains(It.Key))
+			{
+				TArray<FPCGSettingsAndCulling>* AllSettingsAndCulling = nullptr;
+
+				for (FPCGSettingsAndCulling& SettingsAndCulling : It.Value)
+				{
+					if (SettingsAndCulling.Key.Get() && !CurrentExecutionDynamicTrackingSettings.Contains(SettingsAndCulling.Key.Get()))
+					{
+						if (!AllSettingsAndCulling)
+						{
+							AllSettingsAndCulling = &CurrentExecutionDynamicTracking.Add(It.Key);
+						}
+
+						check(AllSettingsAndCulling);
+						// No need for Add Unique since they are already unique in the original map.
+						AllSettingsAndCulling->Add(std::move(SettingsAndCulling));
+					}
+				}
+
+				if (!AllSettingsAndCulling)
+				{
+					ChangedKeys.Add(It.Key);
+				}
+			}
+		}
+
+		DynamicallyTrackedKeysToSettings = std::move(CurrentExecutionDynamicTracking);
+		CurrentExecutionDynamicTracking.Empty();
+		CurrentExecutionDynamicTrackingSettings.Empty();
+	}
+
+	if (!ChangedKeys.IsEmpty())
+	{
+		if (bIsLocal)
+		{
+			OriginalSource->GetExecutionState().AddLocalSourceDynamicTrackingKeys(DynamicallyTrackedKeysToSettings);
+		}
+
+		Subsystem->UpdateTracking(OriginalSource, &ChangedKeys);
+	}
+}
+
+void FPCGChangeTracker::AddLocalSourceDynamicTrackingKeys(const FPCGSelectionKeyToSettingsMap& InLocalSourceDynamicTrackingKeys)
+{
+	PCG::TScopeLock Lock(CurrentExecutionDynamicTrackingLock);
+	for (const auto& It : InLocalSourceDynamicTrackingKeys)
+	{
+		TArray<FPCGSettingsAndCulling>& OriginalSettingsAndCulling = DynamicallyTrackedKeysToSettings.FindOrAdd(It.Key);
+		for (FPCGSettingsAndCulling SettingsAndCulling : It.Value)
+		{
+			OriginalSettingsAndCulling.AddUnique(std::move(SettingsAndCulling));
+		}
+	}
+}
+
+void FPCGChangeTracker::AddLocalSourceCurrentTrackingKeys(const FPCGSelectionKeyToSettingsMap& InLocalSourceCurrentExecutionTrackingKeys, const TSet<const UPCGSettings*>& InLocalSourceCurrentExecutionTrackingSettings)
+{
+	PCG::TScopeLock Lock(CurrentExecutionDynamicTrackingLock);
+	for (const auto& It : InLocalSourceCurrentExecutionTrackingKeys)
+	{
+		TArray<FPCGSettingsAndCulling>& OriginalSettingsAndCulling = CurrentExecutionDynamicTracking.FindOrAdd(It.Key);
+		for (FPCGSettingsAndCulling SettingsAndCulling : It.Value)
+		{
+			OriginalSettingsAndCulling.AddUnique(std::move(SettingsAndCulling));
+		}
+	}
+	CurrentExecutionDynamicTrackingSettings.Append(InLocalSourceCurrentExecutionTrackingSettings);
+}
+
+void FPCGChangeTracker::Serialize(FArchive& Ar)
+{
+	if (!Ar.IsCooking() && !Ar.IsLoadingFromCookedPackage())
+	{
+		Ar << DynamicallyTrackedKeysToSettings;
+	}
+}
+
+#endif
\ No newline at end of file
diff --git a/Engine/Plugins/PCG/Source/PCG/Private/ChangeTracking/PCGChangeTrackingRegistry.cpp b/Engine/Plugins/PCG/Source/PCG/Private/ChangeTracking/PCGChangeTrackingRegistry.cpp
new file mode 100644
--- /dev/null
+++ b/Engine/Plugins/PCG/Source/PCG/Private/ChangeTracking/PCGChangeTrackingRegistry.cpp
@@ -0,0 +1,53 @@
+// Copyright Epic Games, Inc. All Rights Reserved.
+
+#include "ChangeTracking/PCGChangeTrackingRegistry.h"
+
+#if WITH_EDITOR
+
+void FPCGChangeTrackingRegistry::RegisterTrackerFactory(FName InName, FPCGCreateTrackerInstance InTrackerFactory)
+{
+	check(!TrackerFactories.Contains(InName));
+	TrackerFactories.Add(InName, MoveTemp(InTrackerFactory));
+}
+
+void FPCGChangeTrackingRegistry::UnregisterTrackerFactory(FName InName)
+{
+	TrackerFactories.Remove(InName);
+}
+
+TArray<TUniquePtr<IPCGChangeTracker>> FPCGChangeTrackingRegistry::CreateChangeTrackers(FPCGTrackingManager* InManager) const
+{
+	TArray<TUniquePtr<IPCGChangeTracker>> ChangeTrackers;
+
+	for (const auto& KeyValuePair : TrackerFactories)
+	{
+		ChangeTrackers.Add(KeyValuePair.Value.Execute(InManager));
+	}
+
+	return ChangeTrackers;
+}
+
+void FPCGChangeTrackingRegistry::RegisterHandlerFactory(FName InName, FPCGCreateHandlerInstance InHandlerFactory)
+{
+	check(!HandlerFactories.Contains(InName));
+	HandlerFactories.Add(InName, MoveTemp(InHandlerFactory));
+}
+
+void FPCGChangeTrackingRegistry::UnregisterHandlerFactory(FName InName)
+{
+	HandlerFactories.Remove(InName);
+}
+
+TArray<TUniquePtr<IPCGChangeHandler>> FPCGChangeTrackingRegistry::CreateChangeHandlers(FPCGTrackingManager* InManager) const
+{
+	TArray<TUniquePtr<IPCGChangeHandler>> ChangeHandlers;
+
+	for (const auto& KeyValuePair : HandlerFactories)
+	{
+		ChangeHandlers.Add(KeyValuePair.Value.Execute(InManager));
+	}
+
+	return ChangeHandlers;
+}
+
+#endif // WITH_EDITOR
\ No newline at end of file
diff --git a/Engine/Plugins/PCG/Source/PCG/Private/Elements/PCGLoopElement.cpp b/Engine/Plugins/PCG/Source/PCG/Private/Elements/PCGLoopElement.cpp
--- a/Engine/Plugins/PCG/Source/PCG/Private/Elements/PCGLoopElement.cpp
+++ b/Engine/Plugins/PCG/Source/PCG/Private/Elements/PCGLoopElement.cpp
@@ -201,7 +201,7 @@ bool FPCGLoopElement::ExecuteInternal(FPCGContext* InContext) const
 			{
 				if (const UPCGGraph* Graph = Settings->SubgraphOverride->GetGraph())
 				{
-					InContext->ExecutionSource->GetExecutionState().RegisterDynamicTracking(Graph->GetTrackedActorKeysToSettings());
+					InContext->ExecutionSource->GetExecutionState().RegisterDynamicTracking(Graph->GetTrackedSelectionKeysToSettings());
 				}
 			}
 #endif // WITH_EDITOR
diff --git a/Engine/Plugins/PCG/Source/PCG/Private/Grid/PCGExecutionSourceOctree.h b/Engine/Plugins/PCG/Source/PCG/Private/Grid/PCGExecutionSourceOctree.h
deleted file mode 100644
diff --git a/Engine/Plugins/PCG/Source/PCG/Private/Landscape/PCGLandscapeCVars.cpp b/Engine/Plugins/PCG/Source/PCG/Private/Landscape/PCGLandscapeCVars.cpp
new file mode 100644
--- /dev/null
+++ b/Engine/Plugins/PCG/Source/PCG/Private/Landscape/PCGLandscapeCVars.cpp
@@ -0,0 +1,21 @@
+// Copyright Epic Games, Inc. All Rights Reserved.
+
+#include "Landscape/PCGLandscapeCVars.h"
+
+namespace PCGLandscapeCVars
+{
+	TAutoConsoleVariable<int> CVarLandscapeRefreshTimeDelay(
+		TEXT("pcg.LandscapeRefreshTimeDelayMS"),
+		1000,
+		TEXT("Time in MS between a landscape change and PCG refresh. Set it to 0 or negative value to disable the delay."));
+
+	TAutoConsoleVariable<bool> CVarLandscapeDisableRefreshTracking(
+		TEXT("pcg.LandscapeDisableRefreshTracking"),
+		false,
+		TEXT("Completely disable landscape refresh when it changes."));
+
+	TAutoConsoleVariable<bool> CVarLandscapeDisableRefreshTrackingInLandscapeEditingMode(
+		TEXT("pcg.LandscapeDisableRefreshTrackingInLandscapeEditingMode"),
+		false,
+		TEXT("Disable landscape refresh when it changes in landscape editing mode."));
+}
\ No newline at end of file
diff --git a/Engine/Plugins/PCG/Source/PCG/Private/PCGComponent.cpp b/Engine/Plugins/PCG/Source/PCG/Private/PCGComponent.cpp
--- a/Engine/Plugins/PCG/Source/PCG/Private/PCGComponent.cpp
+++ b/Engine/Plugins/PCG/Source/PCG/Private/PCGComponent.cpp
@@ -572,8 +572,7 @@ FPCGTaskId UPCGComponent::CreateGenerateTask(bool bForce, const TArray<FPCGTaskI
 
 #if WITH_EDITOR
 	// No need for lock since it is not executed in parallel.
-	CurrentExecutionDynamicTracking.Empty();
-	CurrentExecutionDynamicTrackingSettings.Empty();
+	ChangeTracker.ResetCurrentExecution();
 #endif // WITH_EDITOR
 
 	return GetSubsystem()->ScheduleGraph(this, *AllDependencies);
@@ -836,7 +835,7 @@ void UPCGComponent::PostCleanupGraph(bool bRemoveComponents)
 	ClearPerPinGeneratedOutput();
 
 #if WITH_EDITOR
-	DynamicallyTrackedKeysToSettings.Reset();
+	ChangeTracker.DynamicallyTrackedKeysToSettings.Reset();
 
 	// Notify Subsystem first
 	if (UPCGSubsystem* Subsystem = GetSubsystem())
@@ -1834,7 +1833,7 @@ void UPCGComponent::Serialize(FArchive& Ar)
 
 		if (DataVersion >= FPCGCustomVersion::DynamicTrackingKeysSerializedInComponent)
 		{
-			Ar << DynamicallyTrackedKeysToSettings;
+			ChangeTracker.Serialize(Ar);
 		}
 	}
 #endif // WITH_EDITORONLY_DATA
@@ -2113,7 +2112,7 @@ void UPCGComponent::RefreshAfterGraphChanged(UPCGGraphInterface* InGraph, EPCGCh
 			TArray<FPCGSelectionKey> ChangedKeys;
 			if (UpdateTrackingCache(&ChangedKeys))
 			{
-				Subsystem->UpdateComponentTracking(this, /*bInShouldDirtyActors=*/ true, &ChangedKeys);
+				Subsystem->UpdateTracking(this, &ChangedKeys);
 			}
 		}
 
@@ -2286,7 +2285,7 @@ void UPCGComponent::PostEditChangeProperty(FPropertyChangedEvent& PropertyChange
 			TArray<FPCGSelectionKey> ChangedKeys;
 			if (UpdateTrackingCache(&ChangedKeys))
 			{
-				Subsystem->UpdateComponentTracking(this, /*bInShouldDirtyActors=*/true, &ChangedKeys);
+				Subsystem->UpdateTracking(this, &ChangedKeys);
 			}
 		}
 
@@ -2428,7 +2427,7 @@ void UPCGComponent::PostEditUndo()
 
 	if (UPCGSubsystem* Subsystem = GetSubsystem())
 	{
-		Subsystem->UpdateComponentTracking(this, /*bInShouldDirtyActors=*/ true);
+		Subsystem->UpdateTracking(this);
 	}
 
 	if (bGenerated)
@@ -2452,63 +2451,42 @@ bool UPCGComponent::UpdateTrackingCache(TArray<FPCGSelectionKey>* OptionalChange
 		return false;
 	}
 
-	int32 FoundKeys = 0;
-
 	// Store in a temporary map to detect key changes.
 	FPCGSelectionKeyToSettingsMap NewTrackedKeysToSettings;
 
+	const int32 PreviousKeyCount = ChangeTracker.StaticallyTrackedKeysToSettings.Num();
+	
+	bool bHasChanged = false;
+
 	if (UPCGGraph* PCGGraph = GetGraph())
 	{
-		NewTrackedKeysToSettings = PCGGraph->GetTrackedActorKeysToSettings();
-
+		// @todo_pcg: to remove usage
+PRAGMA_DISABLE_DEPRECATION_WARNINGS
 		// Also add a key for the landscape, with settings null and always culled, if we should track the landscape
 		if (ShouldTrackLandscape())
+PRAGMA_ENABLE_DEPRECATION_WARNINGS
 		{
 			FPCGSelectionKey LandscapeKey = FPCGSelectionKey(ALandscapeProxy::StaticClass());
 			NewTrackedKeysToSettings.FindOrAdd(LandscapeKey).Emplace(/*Settings*/nullptr, /*bIsCulled*/true);
 		}
 
-		// A tag should be culled, if only all the settings that track this tag should cull.
-		// Note that is only impact the fact that we track (or not) this tag.
-		// If a setting is marked as "should cull", it will only be dirtied (at least by default), if the actor with the
-		// given tag intersect with the component.
-		for (const TPair<FPCGSelectionKey, TArray<FPCGSettingsAndCulling>>& It : NewTrackedKeysToSettings)
-		{
-			const FPCGSelectionKey& Key = It.Key;
-
-			// Should cull only if all the settings requires a cull.
-			const bool bShouldCull = PCGSettings::IsKeyCulled(It.Value);
-			const TArray<FPCGSettingsAndCulling>* OldSettingsAndCulling = StaticallyTrackedKeysToSettings.Find(It.Key);
-			const bool OldCulling = OldSettingsAndCulling && PCGSettings::IsKeyCulled(*OldSettingsAndCulling);
-			const bool bNewKeyOrCullChanged = !OldSettingsAndCulling || (OldCulling != bShouldCull);
-
-			StaticallyTrackedKeysToSettings.Remove(Key);
-
-			if (!bNewKeyOrCullChanged)
-			{
-				++FoundKeys;
-			}
-			else if (OptionalChangedKeys)
-			{
-				OptionalChangedKeys->Add(Key);
-			}
-		}
-
-		// At the end, we also have keys that were tracked but no more, so add them at the list of tracked keys
+		bHasChanged = PCGGraph->UpdateTrackedSelectionKeys(ChangeTracker.StaticallyTrackedKeysToSettings, NewTrackedKeysToSettings, OptionalChangedKeys);
+	}
+	else
+	{
+		// Handle case were we no longer have a graph
+		bHasChanged = ChangeTracker.StaticallyTrackedKeysToSettings.Num() > 0;
 		if (OptionalChangedKeys)
 		{
-			OptionalChangedKeys->Reserve(OptionalChangedKeys->Num() + StaticallyTrackedKeysToSettings.Num());
-
-			for (const TPair<FPCGSelectionKey, TArray<FPCGSettingsAndCulling>>& It : StaticallyTrackedKeysToSettings)
+			OptionalChangedKeys->Reserve(ChangeTracker.StaticallyTrackedKeysToSettings.Num());
+			for (const TPair<FPCGSelectionKey, TArray<FPCGSettingsAndCulling>>& It : ChangeTracker.StaticallyTrackedKeysToSettings)
 			{
 				OptionalChangedKeys->Add(It.Key);
 			}
 		}
 	}
 
-	bool bHasChanged = NewTrackedKeysToSettings.Num() != FoundKeys;
-
-	StaticallyTrackedKeysToSettings = MoveTemp(NewTrackedKeysToSettings);
+	ChangeTracker.StaticallyTrackedKeysToSettings = MoveTemp(NewTrackedKeysToSettings);
 
 	return bHasChanged;
 }
@@ -2777,6 +2755,7 @@ bool UPCGComponent::ShouldGenerateBPPCGAddedToWorld() const
 	}
 }
 
+// deprecated
 bool UPCGComponent::IsObjectTracked(const UObject* InObject, bool& bOutIsCulled) const
 {
 	check(InObject);
@@ -2794,12 +2773,15 @@ bool UPCGComponent::IsObjectTracked(const UObject* InObject, bool& bOutIsCulled)
 		return true;
 	}
 
+	// @todo_pcg: to remove usage
+PRAGMA_DISABLE_DEPRECATION_WARNINGS
 	// If we track the landscape using legacy methods and it is a landscape, it should be tracked as culled
 	if (InObject && InObject->IsA<ALandscapeProxy>() && ShouldTrackLandscape())
 	{
 		bOutIsCulled = true;
 		return true;
 	}
+PRAGMA_ENABLE_DEPRECATION_WARNINGS
 
 	const FPCGSelectionKey ObjectKey = FPCGSelectionKey::CreateFromObjectPtr(InObject);
 
@@ -2819,7 +2801,7 @@ bool UPCGComponent::IsObjectTracked(const UObject* InObject, bool& bOutIsCulled)
 		return false;
 	};
 
-	return CheckMap(StaticallyTrackedKeysToSettings) || CheckMap(DynamicallyTrackedKeysToSettings);
+	return CheckMap(ChangeTracker.StaticallyTrackedKeysToSettings) || CheckMap(ChangeTracker.DynamicallyTrackedKeysToSettings);
 }
 
 void UPCGComponent::OnRefresh(bool bForceRefresh)
@@ -3455,66 +3437,20 @@ UPCGSubsystem* UPCGComponent::GetSubsystem() const
 #if WITH_EDITOR
 void UPCGComponent::ApplyToEachSettings(const FPCGSelectionKey& InKey, const TFunctionRef<void(const FPCGSelectionKey&, const FPCGSettingsAndCulling&)> InCallback) const
 {
-	TRACE_CPUPROFILER_EVENT_SCOPE(UPCGComponent::ApplyToEachSettings);
-
-	auto FindAndApplyInMap = [&InKey, &InCallback](const FPCGSelectionKeyToSettingsMap& InMap)
-	{
-		if (const TArray<FPCGSettingsAndCulling>* StaticallyTrackedSettings = InMap.Find(InKey))
-		{
-			for (const FPCGSettingsAndCulling& SettingsAndCulling : *StaticallyTrackedSettings)
-			{
-				InCallback(InKey, SettingsAndCulling);
-			}
-		}
-	};
-
-	FindAndApplyInMap(StaticallyTrackedKeysToSettings);
-	FindAndApplyInMap(DynamicallyTrackedKeysToSettings);
+	ChangeTracker.ForEachSettingTrackingKey(InKey, InCallback);
 }
 
 TArray<FPCGSelectionKey> UPCGComponent::GatherTrackingKeys() const
 {
-	TArray<FPCGSelectionKey> Keys;
-	Keys.Reserve(StaticallyTrackedKeysToSettings.Num() + DynamicallyTrackedKeysToSettings.Num());
-	for (const auto& It : StaticallyTrackedKeysToSettings)
-	{
-		Keys.Add(It.Key);
-	}
-
-	for (const auto& It : DynamicallyTrackedKeysToSettings)
-	{
-		Keys.Add(It.Key);
-	}
-
-	return Keys;
+	return ChangeTracker.GatherTrackingKeys();
 }
 
 bool UPCGComponent::IsKeyTrackedAndCulled(const FPCGSelectionKey& Key, bool& bOutIsCulled) const
 {
-	bool bIsTracked = false;
-
-	bool bStaticallyCulled = true;
-	bool bDynamicallyCulled = true;
-
-	if (auto* It = StaticallyTrackedKeysToSettings.Find(Key))
-	{
-		bIsTracked = true;
-		bStaticallyCulled = PCGSettings::IsKeyCulled(*It);
-	}
-
-	if (auto* It = DynamicallyTrackedKeysToSettings.Find(Key))
-	{
-		bIsTracked = true;
-		bDynamicallyCulled = PCGSettings::IsKeyCulled(*It);
-	}
-
-	// If it is tracked statically and dynamically, we will cull only and only if both are culling.
-	// Otherwise, it means that at least one key requires to always track, so bOutIsCulled needs to be False.
-	bOutIsCulled = bIsTracked && bStaticallyCulled && bDynamicallyCulled;
-
-	return bIsTracked;
+	return ChangeTracker.IsKeyTrackedAndCulled(Key, bOutIsCulled);
 }
 
+// deprecated
 bool UPCGComponent::ShouldTrackLandscape() const
 {
 	const UPCGGraph* PCGGraph = GetGraph();
@@ -3533,176 +3469,17 @@ bool UPCGComponent::ShouldTrackLandscape() const
 
 void UPCGComponent::RegisterDynamicTracking(const UPCGSettings* InSettings, const TArrayView<TPair<FPCGSelectionKey, bool>>& InDynamicKeysAndCulling)
 {
-	if (!InSettings)
-	{
-		return;
-	}
-
-	UE::TScopeLock Lock(CurrentExecutionDynamicTrackingLock);
-	CurrentExecutionDynamicTrackingSettings.Add(InSettings);
-
-	for (const TPair<FPCGSelectionKey, bool>& It : InDynamicKeysAndCulling)
-	{
-		// Make sure to not register null assets
-		if (It.Key.Selection == EPCGActorSelection::ByPath && It.Key.ObjectPath.IsNull())
-		{
-			continue;
-		}
-
-		CurrentExecutionDynamicTracking.FindOrAdd(It.Key).Emplace(InSettings, It.Value);
-	}
+	ChangeTracker.RegisterDynamicTracking(InSettings, InDynamicKeysAndCulling);
 }
 
 void UPCGComponent::RegisterDynamicTracking(const FPCGSelectionKeyToSettingsMap& InKeysToSettings)
 {
-	if (InKeysToSettings.IsEmpty())
-	{
-		return;
-	}
-
-	UE::TScopeLock Lock(CurrentExecutionDynamicTrackingLock);
-
-	CurrentExecutionDynamicTracking.Append(InKeysToSettings);
-
-	for (const auto& It : InKeysToSettings)
-	{
-		Algo::Transform(It.Value, CurrentExecutionDynamicTrackingSettings, [](const FPCGSettingsAndCulling& SettingsAndCulling) { return SettingsAndCulling.Key.Get();});
-	}
+	ChangeTracker.RegisterDynamicTracking(InKeysToSettings);
 }
 
 void UPCGComponent::UpdateDynamicTracking()
 {
-	TRACE_CPUPROFILER_EVENT_SCOPE(UPCGComponent::UpdateDynamicTracking);
-
-	UPCGSubsystem* Subsystem = GetSubsystem();
-	if (!Subsystem)
-	{
-		return;
-	}
-
-	UPCGComponent* OriginalComponent = GetOriginalComponent();
-	if (!OriginalComponent)
-	{
-		return;
-	}
-
-	const bool bIsLocal = IsLocalComponent();
-	const bool bIsRuntime = IsManagedByRuntimeGenSystem();
-
-	if (bIsRuntime && PCGComponent::CVarDisableDynamicTrackingForRuntimeGen.GetValueOnAnyThread())
-	{
-		return;
-	}
-
-	auto PushToOriginal = [](FPCGSelectionKeyToSettingsMap& LocalDynamicTrackingMap, FPCGSelectionKeyToSettingsMap& OriginalDynamicTrackingMap)
-	{
-		for (auto& It : LocalDynamicTrackingMap)
-		{
-			TArray<FPCGSettingsAndCulling>& OriginalSettingsAndCulling = OriginalDynamicTrackingMap.FindOrAdd(It.Key);
-			for (FPCGSettingsAndCulling& SettingsAndCulling : It.Value)
-			{
-				OriginalSettingsAndCulling.AddUnique(std::move(SettingsAndCulling));
-			}
-		}
-	};
-	
-	// If the component is local, we defer the tracking to the original component.
-	// So move everything to the original (while making sure we are not duplicating keys/settings).
-	// Since it can happen in parallel, we need to lock.
-	if (bIsLocal && !bIsRuntime)
-	{
-		TRACE_CPUPROFILER_EVENT_SCOPE(UPCGComponent::UpdateDynamicTracking::LocalComponent);
-		UE::TScopeLock Lock(OriginalComponent->CurrentExecutionDynamicTrackingLock);
-
-		PushToOriginal(CurrentExecutionDynamicTracking, OriginalComponent->CurrentExecutionDynamicTracking);
-		OriginalComponent->CurrentExecutionDynamicTrackingSettings.Append(CurrentExecutionDynamicTrackingSettings);
-
-		CurrentExecutionDynamicTracking.Empty();
-		CurrentExecutionDynamicTrackingSettings.Empty();
-
-		return;
-	}
-
-	TArray<FPCGSelectionKey> ChangedKeys;
-
-	// Locking to make sure we never hit this multiple times.
-	{
-		UE::TScopeLock Lock(CurrentExecutionDynamicTrackingLock);
-
-		// Go over all dynamic keys gathered during this execution.
-		// If they are not already tracked, we need to register this key.
-		// Otherwise, we need to gather all the settings that tracked this key that were not executed (because of caching).
-		for (auto& It : CurrentExecutionDynamicTracking)
-		{
-			if (TArray<FPCGSettingsAndCulling>* AllSettingsAndCulling = DynamicallyTrackedKeysToSettings.Find(It.Key))
-			{
-				for (FPCGSettingsAndCulling& SettingsAndCulling : *AllSettingsAndCulling)
-				{
-					if (SettingsAndCulling.Key.Get() && !CurrentExecutionDynamicTrackingSettings.Contains(SettingsAndCulling.Key.Get()))
-					{
-						It.Value.AddUnique(std::move(SettingsAndCulling));
-					}
-				}
-			}
-			else
-			{
-				ChangedKeys.Add(It.Key);
-			}
-		}
-
-		// Go over all already registered dynamic keys
-		// If they are not in the current execution gathered keys, we check if they are associated with settings that 
-		// were executed. If so, we re-add them to the current execution gathered keys.
-		// If not, it means that the key is no longer tracked and should be unregistered.
-		for (auto& It : DynamicallyTrackedKeysToSettings)
-		{
-			if (!CurrentExecutionDynamicTracking.Contains(It.Key))
-			{
-				TArray<FPCGSettingsAndCulling>* AllSettingsAndCulling = nullptr;
-
-				for (FPCGSettingsAndCulling& SettingsAndCulling : It.Value)
-				{
-					if (SettingsAndCulling.Key.Get() && !CurrentExecutionDynamicTrackingSettings.Contains(SettingsAndCulling.Key.Get()))
-					{
-						if (!AllSettingsAndCulling)
-						{
-							AllSettingsAndCulling = &CurrentExecutionDynamicTracking.Add(It.Key);
-						}
-
-						check(AllSettingsAndCulling);
-						// No need for Add Unique since they are already unique in the original map.
-						AllSettingsAndCulling->Add(std::move(SettingsAndCulling));
-					}
-				}
-
-				if (!AllSettingsAndCulling)
-				{
-					ChangedKeys.Add(It.Key);
-				}
-			}
-		}
-
-		DynamicallyTrackedKeysToSettings = std::move(CurrentExecutionDynamicTracking);
-		CurrentExecutionDynamicTracking.Empty();
-		CurrentExecutionDynamicTrackingSettings.Empty();
-	}
-
-	if (!ChangedKeys.IsEmpty())
-	{
-		if (bIsLocal)
-		{
-			OriginalComponent->CurrentExecutionDynamicTrackingLock.Lock();
-
-			PushToOriginal(DynamicallyTrackedKeysToSettings, OriginalComponent->DynamicallyTrackedKeysToSettings);
-		}
-		
-		Subsystem->UpdateComponentTracking(OriginalComponent, /*bShouldDirtyActors=*/false, &ChangedKeys);
-
-		if (bIsLocal)
-		{
-			OriginalComponent->CurrentExecutionDynamicTrackingLock.Unlock();
-		}
-	}
+	ChangeTracker.UpdateDynamicTracking(this);
 }
 
 void UPCGComponent::StartIgnoringChangeOriginDuringGeneration(const UObject* InChangeOriginToIgnore)
@@ -3976,7 +3753,7 @@ FPCGComponentInstanceData::FPCGComponentInstanceData(const UPCGComponent* InSour
 {
 #if WITH_EDITOR
 	// Convert map to array for serialization
-	for (const auto& KeyValuePair : InSourceComponent->DynamicallyTrackedKeysToSettings)
+	for (const auto& KeyValuePair : InSourceComponent->ChangeTracker.DynamicallyTrackedKeysToSettings)
 	{
 		for (const auto& SettingsAndCullingPair : KeyValuePair.Value)
 		{
@@ -4031,7 +3808,7 @@ void FPCGComponentInstanceData::ApplyToComponent(UActorComponent* Component, con
 
 		for (const FDynamicTrackedKeyInstanceData& DynamicTrackedKeyInstanceData : DynamicallyTrackedKeysToSettings)
 		{
-			PCGComponent->DynamicallyTrackedKeysToSettings.FindOrAdd(DynamicTrackedKeyInstanceData.SelectionKey).Add({ DynamicTrackedKeyInstanceData.Settings, DynamicTrackedKeyInstanceData.bValue });
+			PCGComponent->ChangeTracker.DynamicallyTrackedKeysToSettings.FindOrAdd(DynamicTrackedKeyInstanceData.SelectionKey).Add({ DynamicTrackedKeyInstanceData.Settings, DynamicTrackedKeyInstanceData.bValue });
 		}
 #endif
 
diff --git a/Engine/Plugins/PCG/Source/PCG/Private/PCGComponentExecutionState.cpp b/Engine/Plugins/PCG/Source/PCG/Private/PCGComponentExecutionState.cpp
--- a/Engine/Plugins/PCG/Source/PCG/Private/PCGComponentExecutionState.cpp
+++ b/Engine/Plugins/PCG/Source/PCG/Private/PCGComponentExecutionState.cpp
@@ -5,6 +5,7 @@
 #include "PCGComponent.h"
 #include "PCGSubsystem.h"
 
+#include "LandscapeProxy.h"
 #include "GameFramework/Actor.h"
 
 UPCGData* FPCGComponentExecutionState::GetSelfData() const
@@ -305,6 +306,24 @@ void FPCGComponentExecutionState::RegisterDynamicTracking(const FPCGSelectionKey
 	Component->RegisterDynamicTracking(InKeysToSettings);
 }
 
+TArray<FPCGSelectionKey> FPCGComponentExecutionState::GatherTrackingKeys() const
+{
+	check(Component);
+	return Component->GatherTrackingKeys();
+}
+
+bool FPCGComponentExecutionState::IsKeyTrackedAndCulled(const FPCGSelectionKey& Key, bool& bOutIsCulled) const
+{
+	check(Component);
+	return Component->IsKeyTrackedAndCulled(Key, bOutIsCulled);
+}
+
+void FPCGComponentExecutionState::ForEachSettingTrackingKey(const FPCGSelectionKey& InKey, const TFunctionRef<void(const FPCGSelectionKey&, const FPCGSettingsAndCulling&)> InCallback) const
+{
+	check(Component);
+	Component->ApplyToEachSettings(InKey, InCallback);
+}
+
 bool FPCGComponentExecutionState::IsRefreshInProgress() const
 {
 	check(Component);
diff --git a/Engine/Plugins/PCG/Source/PCG/Private/PCGGraph.cpp b/Engine/Plugins/PCG/Source/PCG/Private/PCGGraph.cpp
--- a/Engine/Plugins/PCG/Source/PCG/Private/PCGGraph.cpp
+++ b/Engine/Plugins/PCG/Source/PCG/Private/PCGGraph.cpp
@@ -1743,18 +1743,18 @@ void UPCGGraph::OnPCGQualityLevelChanged()
 	}
 }
 
-FPCGSelectionKeyToSettingsMap UPCGGraph::GetTrackedActorKeysToSettings() const
+FPCGSelectionKeyToSettingsMap UPCGGraph::GetTrackedSelectionKeysToSettings() const
 {
 	FPCGSelectionKeyToSettingsMap TagsToSettings;
 	TArray<TObjectPtr<const UPCGGraph>> VisitedGraphs;
 
-	GetTrackedActorKeysToSettings(TagsToSettings, VisitedGraphs);
+	GetTrackedSelectionKeysToSettings(TagsToSettings, VisitedGraphs);
 	return TagsToSettings;
 }
 
-void UPCGGraph::GetTrackedActorKeysToSettings(FPCGSelectionKeyToSettingsMap& OutTagsToSettings, TArray<TObjectPtr<const UPCGGraph>>& OutVisitedGraphs) const
+void UPCGGraph::GetTrackedSelectionKeysToSettings(FPCGSelectionKeyToSettingsMap& OutTagsToSettings, TArray<TObjectPtr<const UPCGGraph>>& OutVisitedGraphs) const
 {
-	TRACE_CPUPROFILER_EVENT_SCOPE(UPCGGraph::GetTrackedActorKeysToSettings);
+	TRACE_CPUPROFILER_EVENT_SCOPE(UPCGGraph::GetTrackedSelectionKeysToSettings);
 	
 	if (OutVisitedGraphs.Contains(this))
 	{
@@ -1789,6 +1789,52 @@ void UPCGGraph::GetTrackedActorKeysToSettings(FPCGSelectionKeyToSettingsMap& Out
 	}
 }
 
+bool UPCGGraph::UpdateTrackedSelectionKeys(FPCGSelectionKeyToSettingsMap& InExistingKeys, FPCGSelectionKeyToSettingsMap& InOutUpdatedKeys, TArray<FPCGSelectionKey>* OptionalChangedKeys) const
+{
+	int32 FoundKeys = 0;
+
+	InOutUpdatedKeys.Append(GetTrackedSelectionKeysToSettings());
+
+	// A tag should be culled, if only all the settings that track this tag should cull.
+	// Note that is only impact the fact that we track (or not) this tag.
+	// If a setting is marked as "should cull", it will only be dirtied (at least by default), if the actor with the
+	// given tag intersect with the component.
+	for (const TPair<FPCGSelectionKey, TArray<FPCGSettingsAndCulling>>& It : InOutUpdatedKeys)
+	{
+		const FPCGSelectionKey& Key = It.Key;
+
+		// Should cull only if all the settings requires a cull.
+		const bool bShouldCull = PCGSettings::IsKeyCulled(It.Value);
+		const TArray<FPCGSettingsAndCulling>* OldSettingsAndCulling = InExistingKeys.Find(It.Key);
+		const bool OldCulling = OldSettingsAndCulling && PCGSettings::IsKeyCulled(*OldSettingsAndCulling);
+		const bool bNewKeyOrCullChanged = !OldSettingsAndCulling || (OldCulling != bShouldCull);
+
+		InExistingKeys.Remove(Key);
+
+		if (!bNewKeyOrCullChanged)
+		{
+			++FoundKeys;
+		}
+		else if (OptionalChangedKeys)
+		{
+			OptionalChangedKeys->Add(Key);
+		}
+	}
+
+	// At the end, we also have keys that were tracked but no more, so add them at the list of tracked keys
+	if (OptionalChangedKeys)
+	{
+		OptionalChangedKeys->Reserve(OptionalChangedKeys->Num() + InExistingKeys.Num());
+
+		for (const TPair<FPCGSelectionKey, TArray<FPCGSettingsAndCulling>>& It : InExistingKeys)
+		{
+			OptionalChangedKeys->Add(It.Key);
+		}
+	}
+	
+	return InOutUpdatedKeys.Num() != FoundKeys || !InExistingKeys.IsEmpty();
+}
+
 void UPCGGraph::NotifyGraphStructureChanged(EPCGChangeType ChangeType, bool bForce)
 {
 	bool bExecutionAffected = true;
diff --git a/Engine/Plugins/PCG/Source/PCG/Private/PCGSubgraph.cpp b/Engine/Plugins/PCG/Source/PCG/Private/PCGSubgraph.cpp
--- a/Engine/Plugins/PCG/Source/PCG/Private/PCGSubgraph.cpp
+++ b/Engine/Plugins/PCG/Source/PCG/Private/PCGSubgraph.cpp
@@ -166,7 +166,7 @@ void UPCGBaseSubgraphSettings::GetStaticTrackedKeys(FPCGSelectionKeyToSettingsMa
 {
 	if (UPCGGraph* Subgraph = GetSubgraph())
 	{
-		Subgraph->GetTrackedActorKeysToSettings(OutKeysToSettings, VisitedGraphs);
+		Subgraph->GetTrackedSelectionKeysToSettings(OutKeysToSettings, VisitedGraphs);
 	}
 }
 
@@ -751,7 +751,7 @@ bool FPCGSubgraphElement::ExecuteInternal(FPCGContext* InContext) const
 		{
 			if (const UPCGGraph* Graph = Settings->SubgraphOverride->GetGraph())
 			{
-				InContext->ExecutionSource->GetExecutionState().RegisterDynamicTracking(Graph->GetTrackedActorKeysToSettings());
+				InContext->ExecutionSource->GetExecutionState().RegisterDynamicTracking(Graph->GetTrackedSelectionKeysToSettings());
 			}
 		}
 #endif // WITH_EDITOR
diff --git a/Engine/Plugins/PCG/Source/PCG/Private/PCGTrackingManager.cpp b/Engine/Plugins/PCG/Source/PCG/Private/PCGTrackingManager.cpp
--- a/Engine/Plugins/PCG/Source/PCG/Private/PCGTrackingManager.cpp
+++ b/Engine/Plugins/PCG/Source/PCG/Private/PCGTrackingManager.cpp
@@ -15,24 +15,18 @@
 #include "Utils/PCGGraphExecutionLogging.h"
 
 #include "Landscape.h"
-#include "LandscapeSubsystem.h"
 #include "StaticMeshCompiler.h"
 #include "TextureCompiler.h"
 #include "Algo/AnyOf.h"
 #include "Algo/ForEach.h"
-#include "Engine/DataTable.h"
 #include "Engine/Engine.h"
 #include "Engine/Level.h"
 #include "Engine/StaticMesh.h"
 #include "Kismet/GameplayStatics.h"
-#include "PackedLevelActor/PackedLevelActor.h"
-#include "LevelInstance/LevelInstanceEditorInstanceActor.h"
-#include "LevelInstance/LevelInstanceInterface.h"
-#include "LevelInstance/LevelInstanceEditorPivotInterface.h"
-#include "LevelInstance/LevelInstanceSubsystem.h"
 #include "Materials/MaterialInterface.h"
 
 #if WITH_EDITOR
+#include "ChangeTracking/PCGChangeTrackingRegistry.h"
 #include "Grid/PCGPartitionActorDesc.h"
 #include "UObject/PackageReload.h"
 #include "WorldPartition/WorldPartition.h"
@@ -48,36 +42,6 @@ namespace PCGTrackingManager
 		false,
 		TEXT("If depencencies are being unstable, disable the tracking, allowing people to continue working while we investigate."));
 
-	static TAutoConsoleVariable<bool> CVarDisableDelayedActorRegistering(
-		TEXT("pcg.DisableDelayedActorRegistering"),
-		false,
-		TEXT("If delayed actor registering when their components aren't registered yet is introducing bad behavior, disables it, allowing people to continue working while we investigate."));
-
-	static TAutoConsoleVariable<bool> CVarLandscapeDisableRefreshTracking(
-		TEXT("pcg.LandscapeDisableRefreshTracking"),
-		false,
-		TEXT("Completely disable landscape refresh when it changes."));
-
-	static TAutoConsoleVariable<bool> CVarLandscapeDisableRefreshTrackingInLandscapeEditingMode(
-		TEXT("pcg.LandscapeDisableRefreshTrackingInLandscapeEditingMode"),
-		false,
-		TEXT("Disable landscape refresh when it changes in landscape editing mode."));
-
-	static TAutoConsoleVariable<int> CVarLandscapeRefreshTimeDelay(
-		TEXT("pcg.LandscapeRefreshTimeDelayMS"),
-		1000,
-		TEXT("Time in MS between a landscape change and PCG refresh. Set it to 0 or negative value to disable the delay."));
-
-	static TAutoConsoleVariable<int> CVarActorModifiedPreviousDataCleanupDelayMS(
-		TEXT("pcg.ActorModifiedPreviousDataCleanupDelayMS"),
-		10000,
-		TEXT("Time in MS between a cached previous data from modified actor and the time we remove it from our cache."));
-
-	static TAutoConsoleVariable<int> CVarActorModifiedPreviousDataCheckDelayMS(
-		TEXT("pcg.ActorModifiedPreviousDataCheckDelayMS"),
-		1000,
-		TEXT("Delay in MS between cleanup checks on previous data from modified actors."));
-
 	static TAutoConsoleVariable<bool> CVarDisablePCGDataInterdependencyOptimization(
 		TEXT("pcg.DisablePCGDataInterdependencyOptimization"),
 		false,
@@ -123,68 +87,6 @@ namespace PCGTrackingManager
 			return true;
 		}
 	}
-
-#if WITH_EDITOR
-	bool ShouldIgnoreActor(AActor* InActor)
-	{
-		if(!InActor)
-		{
-			return true;
-		}
-				
-		return InActor->bIsEditorPreviewActor || InActor->IsA<ALevelInstanceEditorInstanceActor>() || InActor->IsA<APCGWorldActor>() || InActor->Implements<ULevelInstanceEditorPivotInterface>();
-	}
-
-	void PropagateToLevelInstanceActors(ILevelInstanceInterface* InLevelInstance, TFunctionRef<bool(AActor* LevelActor)> InFunc)
-	{
-		if (InLevelInstance)
-		{
-			if (const ULevelInstanceSubsystem* LevelInstanceSubsystem = InLevelInstance->GetLevelInstanceSubsystem())
-			{
-				LevelInstanceSubsystem->ForEachActorInLevelInstance(InLevelInstance, [InFunc](AActor* LevelActor)
-				{
-					if (!ShouldIgnoreActor(LevelActor))
-					{
-						return InFunc(LevelActor);
-					}
-					return true;
-				});
-			}
-		}
-	}
-	
-	void PropagateToLevelInstanceActors(AActor* InActor, TFunctionRef<bool(AActor* LevelActor)> InFunc)
-	{
-		PropagateToLevelInstanceActors(Cast<ILevelInstanceInterface>(InActor), InFunc);
-	}
-		
-	bool ShouldDiscardLandscapeRefresh(const ALandscapeProxy* InLandscape, bool& bIsInEditingMode, bool bIsExitingEditingMode)
-	{
-		bIsInEditingMode = false;
-		// If it is not a landscape, we should refresh.
-		if (!InLandscape)
-		{
-			return false;
-		}
-
-		// If refresh is globably disabled, never refresh
-		if (PCGTrackingManager::CVarLandscapeDisableRefreshTracking.GetValueOnAnyThread())
-		{
-			return true;
-		}
-
-		// If refresh is not disabled in editing, always refresh
-		if (!PCGTrackingManager::CVarLandscapeDisableRefreshTrackingInLandscapeEditingMode.GetValueOnAnyThread())
-		{
-			return false;
-		}
-
-		// Refresh only if we are not editing.
-		const ALandscape* Landscape = InLandscape->GetLandscapeActor();
-		bIsInEditingMode = Landscape && Landscape->HasLandscapeEdMode() && !bIsExitingEditingMode;
-		return bIsInEditingMode;
-	}
-#endif
 }
 
 FPCGTrackingManager::FPCGTrackingManager(UPCGSubsystem* InPCGSubsystem)
@@ -204,6 +106,9 @@ void FPCGTrackingManager::Initialize()
 #if WITH_EDITOR
 	BuildPartitionActorRecords();
 	RegisterDelegates();
+
+	ChangeTrackers = FPCGModule::GetPCGModuleChecked().GetConstChangeTrackingRegistry().CreateChangeTrackers(this);
+	ChangeHandlers = FPCGModule::GetPCGModuleChecked().GetConstChangeTrackingRegistry().CreateChangeHandlers(this);
 #endif // WITH_EDITOR
 }
 
@@ -211,6 +116,9 @@ void FPCGTrackingManager::Deinitialize()
 {
 #if WITH_EDITOR
 	UnregisterDelegates();
+
+	ChangeTrackers.Empty();
+	ChangeHandlers.Empty();
 #endif
 }
 
@@ -230,36 +138,9 @@ void FPCGTrackingManager::Tick()
 	Algo::ForEach(ComponentsToRegister, [this](UPCGComponent* Component) { RegisterOrUpdateExecutionSource(Component, Component->bGenerated, /*bForce=*/true); });
 	Algo::ForEach(ComponentsToUnregister, [this](UPCGComponent* Component) { UnregisterExecutionSource(Component, /*bForce=*/true); });
 
-	ProcessDelayedEvents();
-
-	const double CurrentTime = FApp::GetCurrentTime();
-
-	if (!DelayedModifiedLandscapes.IsEmpty() && LastLandscapeDirtyTime > 0.0 && ((CurrentTime - LastLandscapeDirtyTime) * 1000.0) > PCGTrackingManager::CVarLandscapeRefreshTimeDelay.GetValueOnAnyThread())
-	{
-		LastLandscapeDirtyTime = -1.0;
-		for (TObjectKey<ALandscapeProxy> Landscape : DelayedModifiedLandscapes)
-		{
-			ApplyLandscapeChanges(Landscape.ResolveObjectPtr());
-		}
-
-		DelayedModifiedLandscapes.Empty();
-	}
-
-	// Cleaning up previous data gathered by the OnObjectModified function but not consumed.
-	if (LastPreviousActorDataCleanup < 0.0 || ((CurrentTime - LastPreviousActorDataCleanup) * 1000) > PCGTrackingManager::CVarActorModifiedPreviousDataCheckDelayMS.GetValueOnAnyThread())
+	for (TUniquePtr<IPCGChangeTracker>& ChangeTracker : ChangeTrackers)
 	{
-		TRACE_CPUPROFILER_EVENT_SCOPE(FPCGTrackingManager::Tick::PreviousActorDataCleanupCheck);
-
-		LastPreviousActorDataCleanup = CurrentTime;
-		TArray<TObjectKey<AActor>> Keys;
-		ActorToPreviousDataMap.GetKeys(Keys);
-		for (const TObjectKey<AActor>& Key : Keys)
-		{
-			if (((CurrentTime - ActorToPreviousDataMap[Key].Get<2>()) * 1000) > PCGTrackingManager::CVarActorModifiedPreviousDataCleanupDelayMS.GetValueOnAnyThread())
-			{
-				ActorToPreviousDataMap.Remove(Key);
-			}
-		}
+		ChangeTracker->Tick();
 	}
 
 	if (PCGTrackingManager::CVarDisablePCGDataInterdependencyOptimization.GetValueOnAnyThread() && !ExecutionSourcesToDependencyMap.IsEmpty())
@@ -1120,10 +1001,10 @@ void FPCGTrackingManager::RegisterTracking(IPCGGraphExecutionSource* InExecution
 	// Components owner needs to be always tracked
 	AlwaysTrackedKeysToExecutionSourcesMap.FindOrAdd(FPCGSelectionKey(EPCGActorFilter::Self)).Add(InExecutionSource);
 
-	UpdateTracking(InExecutionSource, /*bInShouldDirtyActors=*/ false);
+	UpdateTracking(InExecutionSource);
 }
 
-void FPCGTrackingManager::UpdateTracking(IPCGGraphExecutionSource* InExecutionSource, bool bInShouldDirtyActors, const TArray<FPCGSelectionKey>* ChangedKeys)
+void FPCGTrackingManager::UpdateTracking(IPCGGraphExecutionSource* InExecutionSource, const TArray<FPCGSelectionKey>* ChangedKeys)
 {
 	if (!PCGTrackingManager::IsValidOriginalExecutionSource(InExecutionSource))
 	{
@@ -1133,18 +1014,11 @@ void FPCGTrackingManager::UpdateTracking(IPCGGraphExecutionSource* InExecutionSo
 	UWorld* World = GetWorld();
 	check(World);
 
-	// @todo_pcg: support execution source
-	UPCGComponent* PCGComponent = Cast<UPCGComponent>(InExecutionSource);
-	if (!PCGComponent)
-	{
-		return;
-	}
-
 	// If no keys are provided, update all tracking keys.
 	TArray<FPCGSelectionKey> AllKeys;
 	if (ChangedKeys == nullptr)
 	{
-		AllKeys = PCGComponent->GatherTrackingKeys();
+		AllKeys = InExecutionSource->GetExecutionState().GatherTrackingKeys();
 		ChangedKeys = &AllKeys;
 	}
 
@@ -1152,10 +1026,10 @@ void FPCGTrackingManager::UpdateTracking(IPCGGraphExecutionSource* InExecutionSo
 
 	auto RemoveFromMap = [InExecutionSource](TMap<FPCGSelectionKey, TSet<IPCGGraphExecutionSource*>>& InMap, const FPCGSelectionKey& InKey)
 	{
-		if (TSet<IPCGGraphExecutionSource*>* Components = InMap.Find(InKey))
+		if (TSet<IPCGGraphExecutionSource*>* ExecutionSources = InMap.Find(InKey))
 		{
-			Components->Remove(InExecutionSource);
-			if (Components->IsEmpty())
+			ExecutionSources->Remove(InExecutionSource);
+			if (ExecutionSources->IsEmpty())
 			{
 				InMap.Remove(InKey);
 			}
@@ -1166,7 +1040,7 @@ void FPCGTrackingManager::UpdateTracking(IPCGGraphExecutionSource* InExecutionSo
 	{
 		// bShouldBeCulled is modified in IsKeyTrackedAndCulled
 		bool bShouldBeCulled = false;
-		if (!PCGComponent->IsKeyTrackedAndCulled(Key, bShouldBeCulled))
+		if (!InExecutionSource->GetExecutionState().IsKeyTrackedAndCulled(Key, bShouldBeCulled))
 		{
 			// Untrack
 			RemoveFromMap(CulledTrackedKeysToExecutionSourcesMap, Key);
@@ -1175,12 +1049,12 @@ void FPCGTrackingManager::UpdateTracking(IPCGGraphExecutionSource* InExecutionSo
 		else if (bShouldBeCulled)
 		{
 			RemoveFromMap(AlwaysTrackedKeysToExecutionSourcesMap, Key);
-			CulledTrackedKeysToExecutionSourcesMap.FindOrAdd(Key).Add(PCGComponent);
+			CulledTrackedKeysToExecutionSourcesMap.FindOrAdd(Key).Add(InExecutionSource);
 		}
 		else
 		{
 			RemoveFromMap(CulledTrackedKeysToExecutionSourcesMap, Key);
-			AlwaysTrackedKeysToExecutionSourcesMap.FindOrAdd(Key).Add(PCGComponent);
+			AlwaysTrackedKeysToExecutionSourcesMap.FindOrAdd(Key).Add(InExecutionSource);
 		}
 	}
 }
@@ -1287,19 +1161,23 @@ bool FPCGTrackingManager::IsKeyTracked(const FPCGSelectionKey& InKey) const
 	return CulledTrackedKeysToExecutionSourcesMap.Contains(InKey) || AlwaysTrackedKeysToExecutionSourcesMap.Contains(InKey);
 }
 
-bool FPCGTrackingManager::IsActorTracked(const AActor* InActor) const
+void FPCGTrackingManager::ForEachTrackedKey(TFunctionRef<bool(const FPCGSelectionKey&, const TSet<IPCGGraphExecutionSource*>&)> InCallback) const
 {
-	TRACE_CPUPROFILER_EVENT_SCOPE(FPCGTrackingManager::IsActorTracked);
-
-	TSet<FName> EmptySet;
-
-	const FPCGSelectionKey ActorKey = FPCGSelectionKey::CreateFromObjectPtr(InActor);
-	auto Matching = [&ActorKey, &EmptySet](const TPair<FPCGSelectionKey, TSet<IPCGGraphExecutionSource*>>& It) -> bool
+	for (const auto& KeyValuePair : CulledTrackedKeysToExecutionSourcesMap)
 	{
-		return It.Key.IsMatching(ActorKey, EmptySet, It.Value, nullptr);
-	};
+		if (!InCallback(KeyValuePair.Key, KeyValuePair.Value))
+		{
+			return;
+		}
+	}
 
-	return Algo::AnyOf(CulledTrackedKeysToExecutionSourcesMap, Matching) || Algo::AnyOf(AlwaysTrackedKeysToExecutionSourcesMap, Matching);
+	for (const auto& KeyValuePair : AlwaysTrackedKeysToExecutionSourcesMap)
+	{
+		if (!InCallback(KeyValuePair.Key, KeyValuePair.Value))
+		{
+			return;
+		}
+	}
 }
 
 void FPCGTrackingManager::ResetPartitionActorsMap()
@@ -1308,41 +1186,54 @@ void FPCGTrackingManager::ResetPartitionActorsMap()
 	PartitionActorsMap.Empty();
 }
 
+bool FPCGTrackingManager::RemoveExecutionSourceDependency(const IPCGGraphExecutionSource* InOriginatingExecutionSource, const IPCGGraphExecutionSource* InDependentExecutionSource)
+{
+	const UObject* OriginatingObject = Cast<UObject>(InOriginatingExecutionSource);
+	const UObject* DependentObject = Cast<UObject>(InDependentExecutionSource);
+	if (OriginatingObject && DependentObject)
+	{
+		if (TArray<FObjectKey>* Dependencies = ExecutionSourcesToDependencyMap.Find(DependentObject))
+		{
+			Dependencies->RemoveSwap(OriginatingObject);
+			if (Dependencies->IsEmpty())
+			{
+				ExecutionSourcesToDependencyMap.Remove(DependentObject);
+				return true;
+			}
+			else
+			{
+				return false;
+			}
+		}
+	}
+
+	return true;
+}
+
+bool FPCGTrackingManager::IsTracking() const
+{
+	if (IsRunningCommandlet())
+	{
+		return false;
+	}
+
+	return !CulledTrackedKeysToExecutionSourcesMap.IsEmpty() || !AlwaysTrackedKeysToExecutionSourcesMap.IsEmpty();
+}
+
 void FPCGTrackingManager::RegisterDelegates()
 {
-	GEngine->OnLevelActorAdded().AddRaw(this, &FPCGTrackingManager::OnActorAdded);
-	GEngine->OnLevelActorDeleted().AddRaw(this, &FPCGTrackingManager::OnActorDeleted);
 	FCoreUObjectDelegates::OnObjectPropertyChanged.AddRaw(this, &FPCGTrackingManager::OnObjectPropertyChanged);
-	FCoreUObjectDelegates::OnObjectModified.AddRaw(this, &FPCGTrackingManager::OnObjectModified);
-	FCoreUObjectDelegates::OnObjectPreSave.AddRaw(this, &FPCGTrackingManager::OnObjectSaved);
 	FCoreUObjectDelegates::OnPackageReloaded.AddRaw(this, &FPCGTrackingManager::OnPackageReloaded);
 
 	UWorld* World = GetWorld();
 	check(World);
 
-	// Listen to added/removed ActorDescs to keep track of unloaded Partiton Actors
+	// @todo_pcg: move this out when partitioning gets abstracted
 	if (UWorldPartition* WorldPartition = World->GetWorldPartition())
 	{
 		WorldPartition->OnActorDescInstanceAddedEvent.AddRaw(this, &FPCGTrackingManager::OnActorDescInstanceAdded);
 		WorldPartition->OnActorDescInstanceRemovedEvent.AddRaw(this, &FPCGTrackingManager::OnActorDescInstanceRemoved);
 	}
-
-	if (World->PersistentLevel)
-	{
-		World->PersistentLevel->OnLoadedActorAddedToLevelEvent.AddRaw(this, &FPCGTrackingManager::OnActorLoaded);
-		World->PersistentLevel->OnLoadedActorRemovedFromLevelEvent.AddRaw(this, &FPCGTrackingManager::OnActorUnloaded);
-
-		if (ULevelInstanceSubsystem* LevelInstanceSubsystem = UWorld::GetSubsystem<ULevelInstanceSubsystem>(World))
-		{
-			LevelInstanceSubsystem->OnLevelInstancesUpdated().AddRaw(this, &FPCGTrackingManager::OnLevelInstancesUpdated);
-			LevelInstanceSubsystem->OnLevelInstanceEditCancelled().AddRaw(this, &FPCGTrackingManager::OnLevelInstanceEditCancelled);
-		}
-	}
-
-	if (ULandscapeSubsystem* LandscapeSubsystem = World->GetSubsystem<ULandscapeSubsystem>())
-	{
-		LandscapeSubsystem->OnLandscapeProxyComponentDataChanged().AddRaw(this, &FPCGTrackingManager::OnLandscapeChanged);
-	}
 }
 
 void FPCGTrackingManager::UnregisterDelegates()
@@ -1351,7 +1242,6 @@ void FPCGTrackingManager::UnregisterDelegates()
 	GEngine->OnLevelActorDeleted().RemoveAll(this);
 	FCoreUObjectDelegates::OnObjectPropertyChanged.RemoveAll(this);
 	FCoreUObjectDelegates::OnObjectModified.RemoveAll(this);
-	FCoreUObjectDelegates::OnObjectPreSave.RemoveAll(this);
 	FCoreUObjectDelegates::OnPackageReloaded.RemoveAll(this);
 
 	UWorld* World = GetWorld();
@@ -1362,123 +1252,6 @@ void FPCGTrackingManager::UnregisterDelegates()
 		WorldPartition->OnActorDescInstanceAddedEvent.RemoveAll(this);
 		WorldPartition->OnActorDescInstanceRemovedEvent.RemoveAll(this);
 	}
-
-	if (World->PersistentLevel)
-	{
-		World->PersistentLevel->OnLoadedActorAddedToLevelEvent.RemoveAll(this);
-		World->PersistentLevel->OnLoadedActorRemovedFromLevelEvent.RemoveAll(this);
-
-		if (ULevelInstanceSubsystem* LevelInstanceSubsystem = UWorld::GetSubsystem<ULevelInstanceSubsystem>(World))
-		{
-			LevelInstanceSubsystem->OnLevelInstancesUpdated().RemoveAll(this);
-			LevelInstanceSubsystem->OnLevelInstanceEditCancelled().RemoveAll(this);
-		}
-	}
-
-	if (ULandscapeSubsystem* LandscapeSubsystem = World->GetSubsystem<ULandscapeSubsystem>())
-	{
-		LandscapeSubsystem->OnLandscapeProxyComponentDataChanged().RemoveAll(this);
-	}
-}
-
-bool FPCGTrackingManager::ShouldDelayActor(AActor* InActor) const
-{
-	const bool bDisableDelayedActorRegistering = PCGTrackingManager::CVarDisableDelayedActorRegistering.GetValueOnAnyThread();
-	if (!InActor->HasActorRegisteredAllComponents() && !bDisableDelayedActorRegistering)
-	{
-		return true;
-	}
-
-	bool bShouldDelay = false;
-
-	// Check that whole Level Instance hierarchy is loaded
-	if (ILevelInstanceInterface* LevelInstance = Cast<ILevelInstanceInterface>(InActor))
-	{
-		// If Level Instance isn't loaded yet it should be delayed
-		if (!LevelInstance->IsLoaded() && !InActor->IsA<APackedLevelActor>())
-		{
-			return true;
-		}
-			
-		PCGTrackingManager::PropagateToLevelInstanceActors(LevelInstance, [this, &bShouldDelay](AActor* InChildActor)
-		{
-			if(ShouldDelayActor(InChildActor))
-			{
-				bShouldDelay = true;
-				return false;
-			}
-			return true;
-		});
-	}
-
-	return bShouldDelay;
-}
-
-void FPCGTrackingManager::ProcessDelayedEvents()
-{
-	// Safeguard, we can't add delayed actors if the subsystem is not initialized
-	if (!PCGSubsystem->IsInitialized() || DelayedChangedActors.IsEmpty())
-	{
-		return;
-	}
-			
-	auto LocalDelayedChangedActors = MoveTemp(DelayedChangedActors);
-
-	for (const auto& ActorToObject : LocalDelayedChangedActors)
-	{
-		AActor* Actor = ActorToObject.Key.ResolveObjectPtr();
-		UObject* Object = ActorToObject.Value.Get<0>().ResolveObjectPtr();
-		const EChangedObjectReason Reason = ActorToObject.Value.Get<1>();
-		
-		if (!Actor)
-		{
-			continue;
-		}
-
-		OnActorChanged_Internal(Actor, Object, Reason);
-	}
-}
-
-void FPCGTrackingManager::OnLevelInstancesUpdated(const TArray<ILevelInstanceInterface*>& InLevelInstances)
-{
-	for(ILevelInstanceInterface* LevelInstance : InLevelInstances)
-	{
-		OnActorAdded(Cast<AActor>(LevelInstance));
-	}
-}
-
-void FPCGTrackingManager::OnLevelInstanceEditCancelled(ILevelInstanceInterface* InLevelInstance, bool bInHasDiscardedChanges)
-{
-	if (bInHasDiscardedChanges)
-	{
-		OnActorAdded(Cast<AActor>(InLevelInstance));
-	}
-}
-
-void FPCGTrackingManager::OnActorLoaded(AActor& InActor)
-{
-	// We have to make sure to not create a infinite loop
-	if (PCGTrackingManager::ShouldIgnoreActor(&InActor) || InActor.GetWorld() != GetWorld())
-	{
-		return;
-	}
-
-	// Loaded actors should not dirty.
-	OnActorAdded_Internal(&InActor, /*bShouldDirty=*/ false);
-}
-
-void FPCGTrackingManager::OnActorAdded(AActor* InActor)
-{
-	// We have to make sure to not create a infinite loop
-	if (PCGTrackingManager::ShouldIgnoreActor(InActor) || InActor->GetWorld() != GetWorld())
-	{
-		return;
-	}
-			
-	// Implementation note: We delay adding because OnActorAdded fires before an actor's properties are set,
-	// so the actor is not ready for processing until the next tick.
-	// Only dirty if we aren't currently inside a loading path
-	OnActorAdded_Internal(InActor, /*bShouldDirty=*/ !UE::GetIsEditorLoadingPackage());
 }
 
 void FPCGTrackingManager::OnPackageReloaded(const EPackageReloadPhase InPackageReloadPhase, FPackageReloadedEvent* InPackageReloadedEvent)
@@ -1602,197 +1375,25 @@ void FPCGTrackingManager::BuildPartitionActorRecords()
 	}
 }
 
-void FPCGTrackingManager::OnActorChanged_Recursive(AActor* InActor, UObject* InOriginatingChangeObject, const EChangedObjectReason InChangeReason)
-{
-	check(!DelayedChangedActors.Contains(InActor));
-
-	// Process actors recursively
-	auto OnActorChanged = [this, InOriginatingChangeObject, InChangeReason](AActor* InActor, auto&& OnActorChangedRecursive) -> void
-	{
-		if (InActor)
-		{
-			FActorPreviousData* PreviousData = ActorToPreviousDataMap.Find(InActor);
-			OnObjectChanged(InActor, InChangeReason, PreviousData, InOriginatingChangeObject);
-			if (PreviousData)
-			{
-				ActorToPreviousDataMap.Remove(InActor);
-			}
-		}
-		
-		PCGTrackingManager::PropagateToLevelInstanceActors(InActor, [this, OnActorChangedRecursive](AActor* LevelActor)
-		{
-			OnActorChangedRecursive(LevelActor, OnActorChangedRecursive);
-			return true;
-		});
-	};
-
-	OnActorChanged(InActor, OnActorChanged);
-}
-
-void FPCGTrackingManager::OnActorAdded_Internal(AActor* InActor, bool bShouldDirty)
-{
-	TRACE_CPUPROFILER_EVENT_SCOPE(FPCGTrackingManager::OnActorAdded);
-	check(!PCGTrackingManager::ShouldIgnoreActor(InActor) && InActor->GetWorld() == GetWorld());
-
-
-	if (!bShouldDirty || DelayedChangedActors.Contains(InActor))
-	{
-		// Nothing to do
-		return;
-	}
-
-	// We delay adding because OnActorAdded fires before an actor's properties are set,
-	// so the actor is not ready for processing until the next tick.
-	DelayedChangedActors.Emplace(InActor, {nullptr, EChangedObjectReason::ActorAdded});
-}
-
-void FPCGTrackingManager::OnActorChanged_Internal(AActor* InActor, UObject* InOriginatingChangeObject, const EChangedObjectReason InChangeReason)
-{
-	// Currently delayed, return
-	if (DelayedChangedActors.Contains(InActor))
-	{
-		return;
-	}
-
-	// Should still be delayed, add to delayed actors, return
-	if (ShouldDelayActor(InActor))
-	{
-		DelayedChangedActors.Emplace(InActor, {InOriginatingChangeObject, InChangeReason});
-		return;
-	}
-	
-	// Call Recursive method
-	OnActorChanged_Recursive(InActor, InOriginatingChangeObject, InChangeReason);
-}
-
-void FPCGTrackingManager::OnActorUnloaded(AActor& InActor)
-{
-	// Don't dirty on unload (to mirror the behavior in load)
-	OnActorDeleted_Internal(&InActor, /*bShouldDirty=*/false);
-}
-
-void FPCGTrackingManager::OnActorDeleted(AActor* InActor)
-{
-	if (PCGTrackingManager::ShouldIgnoreActor(InActor) || InActor->GetWorld() != GetWorld())
-	{
-		return;
-	}
-
-	// Only dirty if we aren't currently inside a loading path
-	OnActorDeleted_Internal(InActor, /*bShouldDirty=*/ !UE::GetIsEditorLoadingPackage());
-}
-
-void FPCGTrackingManager::OnActorDeleted_Internal(AActor* InActor, bool bShouldDirty)
-{
-	TRACE_CPUPROFILER_EVENT_SCOPE(UPCGTrackingManager::OnActorDeleted_Internal);
-	check(InActor && InActor->GetWorld() == GetWorld());
-
-	if (!bShouldDirty)
-	{
-		// Nothing to do
-		return;
-	}
-
-#if WITH_EDITOR
-	PCGTrackingManager::PropagateToLevelInstanceActors(InActor, [this, bShouldDirty](AActor* LevelActor)
-	{
-		OnActorDeleted_Internal(LevelActor, bShouldDirty);
-		return true;
-	});
-#endif // WITH_EDITOR
-
-	// Notify all components that the actor has changed (was removed), but the Refresh will only happen AFTER the actor was actually removed from the world (because of delayed refresh).
-	OnObjectChanged(InActor, EChangedObjectReason::ActorDeleted, /*InPreviousData=*/nullptr, /*InOriginatingObject=*/nullptr);
-}
-
-void FPCGTrackingManager::OnObjectModified(UObject* InObject)
+void FPCGTrackingManager::OnObjectPropertyChanged(UObject* InObject, FPropertyChangedEvent& InEvent)
 {
-	TRACE_CPUPROFILER_EVENT_SCOPE(FPCGTrackingManager::OnObjectModified);
+	TRACE_CPUPROFILER_EVENT_SCOPE(FPCGTrackingManager::OnObjectPropertyChanged);
 
 	// Nothing to do if we track nothing
-	if (CulledTrackedKeysToExecutionSourcesMap.IsEmpty() && AlwaysTrackedKeysToExecutionSourcesMap.IsEmpty())
+	if (!IsTracking())
 	{
 		return;
 	}
 
-	AActor* Actor = Cast<AActor>(InObject);
-
-	// Otherwise, if it's an actor component, track it as well
-	if (!Actor)
+	// Check if the event is handled by one of our trackers
+	for (TUniquePtr<IPCGChangeTracker>& ChangeTracker : ChangeTrackers)
 	{
-		if (UActorComponent* ActorComponent = Cast<UActorComponent>(InObject))
+		if (ChangeTracker->OnObjectPropertyChanged(InObject, InEvent))
 		{
-			Actor = ActorComponent->GetOwner();
-		}
-	}
-
-	if (!Actor || (Actor->GetWorld() != GetWorld()))
-	{
-		return;
-	}
-
-#if WITH_EDITOR
-	if (PCGTrackingManager::ShouldIgnoreActor(Actor))
-	{
-		return;
-	}
-#endif
-
-	auto StorePreviousData = [this](AActor* InActor, auto&& RecursiveCall) -> void
-	{
-		if (!ActorToPreviousDataMap.Contains(InActor))
-		{
-			if (IsActorTracked(InActor))
-			{
-				FActorPreviousData& PreviousData = ActorToPreviousDataMap.Add(InActor);
-				PreviousData.Get<0>() = PCGHelpers::GetActorBounds(InActor, /*bIgnorePCGCreatedComponents=*/false);
-				PreviousData.Get<1>() = TSet<FName>(InActor->Tags);
-				PreviousData.Get<2>() = FApp::GetCurrentTime();
-			}
-
-			// Also propagate the pre-change to all child actors if it is within a level instance.
-			PCGTrackingManager::PropagateToLevelInstanceActors(InActor, [this, RecursiveCall](AActor* LevelActor)
-			{
-				RecursiveCall(LevelActor, RecursiveCall);
-				return true;
-			});
+			return;
 		}
-	};
-
-	StorePreviousData(Actor, StorePreviousData);
-}
-
-void FPCGTrackingManager::OnObjectSaved(UObject* InObject, FObjectPreSaveContext InObjectSaveContext)
-{
-	// Nothing to do if we track nothing
-	if (!InObject || (CulledTrackedKeysToExecutionSourcesMap.IsEmpty() && AlwaysTrackedKeysToExecutionSourcesMap.IsEmpty()))
-	{
-		return;
-	}
-
-	// Only trigger a refresh on save for limited data classes.
-	// At this point in time, We only track data tables and PCG data assets because in most cases we probably will catch other changes with OnObjectPropertyChanged.
-	// This is especially important to make sure we don't trigger refresh multiple times (less a problem for PCG assets, but a big problem for data tables).
-	if (!InObjectSaveContext.IsProceduralSave() && (InObject->IsA<UDataTable>() || InObject->IsA<UPCGDataAsset>() || InObject->IsA<UStaticMesh>()))
-	{
-		FPropertyChangedEvent Event{ nullptr };
-		OnObjectChanged(InObject, EChangedObjectReason::ObjectSaved, /*InPreviousData=*/nullptr, /*InOriginatingChangeObject=*/ InObject);
-	}
-}
-
-void FPCGTrackingManager::OnObjectPropertyChanged(UObject* InObject, FPropertyChangedEvent& InEvent)
-{
-	TRACE_CPUPROFILER_EVENT_SCOPE(FPCGTrackingManager::OnObjectPropertyChanged);
-
-	// Nothing to do if we track nothing
-	if (CulledTrackedKeysToExecutionSourcesMap.IsEmpty() && AlwaysTrackedKeysToExecutionSourcesMap.IsEmpty())
-	{
-		return;
 	}
 
-	const bool bValueNotInteractive = (InEvent.ChangeType != EPropertyChangeType::Interactive);
-	// Special exception for actor tags, as we can't track otherwise an actor "losing" a tag
-	const bool bActorTagChange = (InEvent.Property && InEvent.Property->GetFName() == GET_MEMBER_NAME_CHECKED(AActor, Tags));
 	// Another special exception for texture/mesh compilation. If the InEvent is empty and the object is a texture/mesh, we ignore it.
 	const bool bEventIsEmpty = (InEvent.Property == nullptr) && (InEvent.ChangeType == EPropertyChangeType::Unspecified);
 	const bool bIsTextureCompilationResult = bEventIsEmpty && InObject && InObject->IsA<UTexture>()
@@ -1802,79 +1403,25 @@ void FPCGTrackingManager::OnObjectPropertyChanged(UObject* InObject, FPropertyCh
 	const bool bIsStaticMeshCompilationResult = bEventIsEmpty && InObject && InObject->IsA<UStaticMesh>()
 		&& FStaticMeshCompilingManager::Get().GetNumRemainingMeshes() > 0;
 
-	// First check if it is an actor
-	AActor* Actor = Cast<AActor>(InObject);
-
-	// Otherwise, if it's an actor component, track it as well
-	if (!Actor)
-	{
-		if (UActorComponent* ActorComponent = Cast<UActorComponent>(InObject))
-		{
-			Actor = ActorComponent->GetOwner();
-		}
-	}
-
-	if (Actor && Actor->GetWorld() != GetWorld())
-	{
-		return;
-	}
-
-	if (Actor && PCGTrackingManager::ShouldIgnoreActor(Actor))
+	if (bIsTextureCompilationResult || bIsStaticMeshCompilationResult)
 	{
 		return;
 	}
 
-	const bool bNoOperation = (!bValueNotInteractive && !bActorTagChange)
-		|| bIsTextureCompilationResult
-		|| bIsStaticMeshCompilationResult
-		|| (Actor && (DelayedChangedActors.Contains(Actor)));
-
-	if (bNoOperation)
-	{
-		// Don't remove it as it will be needed when the delay will be called.
-		if (Actor && !(DelayedChangedActors.Contains(Actor)))
-		{
-			auto RemoveAll = [this](AActor* InActor, auto&& RecursiveCall) -> void
-			{
-				ActorToPreviousDataMap.Remove(InActor);
-				PCGTrackingManager::PropagateToLevelInstanceActors(InActor, [this, RecursiveCall](AActor* LevelActor)
-				{
-					RecursiveCall(LevelActor, RecursiveCall);
-					return true;
-				});
-			};
-
-			RemoveAll(Actor, RemoveAll);
-		}
-		
-		return;
-	}
-
-	if (Actor)
-	{
-		OnActorChanged_Internal(Actor, /*InOriginatingChangeObject=*/ InObject, EChangedObjectReason::ActorChanged);
-	}
-	else
-	{
-		OnObjectChanged(InObject, EChangedObjectReason::ObjectChanged, /*InPreviousData=*/nullptr, /*InOriginatingChangeObject=*/ InObject);
-	}
+	OnObjectChanged(InObject, PCGTrackingManager::ObjectChanged, /*InOriginatingChangeObject=*/ InObject);
 }
 
-void FPCGTrackingManager::OnObjectChanged(UObject* InObject, const EChangedObjectReason InChangeReason, const FActorPreviousData* InPreviousData, const UObject* InOriginatingChangeObject)
+void FPCGTrackingManager::OnObjectChanged(const UObject* InObject, const FPCGTrackingChangeID& InChangeType, const UObject* InOriginatingChangeObject)
 {
 	TRACE_CPUPROFILER_EVENT_SCOPE(FPCGTrackingManager::OnObjectChanged);
 
-	// Discard everything related to self for actor added/delete and changes coming from a PCG Component.
-	const bool bNoRefreshOwner = InChangeReason == EChangedObjectReason::ActorAdded
-		|| InChangeReason == EChangedObjectReason::ActorDeleted
-		|| InChangeReason == EChangedObjectReason::PCGComponentGenerated;
-
 	// Nothing to do if we track nothing or there is no object.
-	if (!InObject || (CulledTrackedKeysToExecutionSourcesMap.IsEmpty() && AlwaysTrackedKeysToExecutionSourcesMap.IsEmpty()))
+	if (!InObject || !IsTracking())
 	{
 		return;
 	}
 
+	// @todo_pcg: excluded classes should be done through handler/tracker api
 	// Don't react to what the PCG Component is already reacting to.
 	static const TArray<const UClass*> ExcludedClasses =
 	{
@@ -1886,32 +1433,31 @@ void FPCGTrackingManager::OnObjectChanged(UObject* InObject, const EChangedObjec
 		return;
 	}
 
-	const AActor* Actor = Cast<AActor>(InObject);
+	const bool bNoRefreshOwner = InChangeType == PCGTrackingManager::ExecutionSourceGenerated;
 
-	ensure(!Actor || Actor->GetWorld() == GetWorld());
+	FPCGSelectionKey SelectionKey = FPCGSelectionKey::CreateFromObjectPtr(InObject);
+	OnSelectionKeyChanged(SelectionKey, InOriginatingChangeObject, {}, bNoRefreshOwner, {FBox(EForceInit::ForceInit)});
+}
 
-	TSet<FName> RemovedTags;
-	if (Actor && InPreviousData && !InPreviousData->Get<1>().IsEmpty())
+void FPCGTrackingManager::NotifyObjectChanged(const UObject* InObject, const FPCGTrackingChangeID& InChangeType, const UObject* InOriginatingChangeObject)
+{
+	for (TUniquePtr<IPCGChangeTracker>& ChangeTracker : ChangeTrackers)
 	{
-		RemovedTags = InPreviousData->Get<1>().Difference(TSet<FName>(Actor->Tags));
-	}
-
-	FPCGSelectionKey SelectionKey = FPCGSelectionKey::CreateFromObjectPtr(InObject);
-	
-	TArray<FBox, TInlineAllocator<2>> Bounds;
-	const FBox ActorBounds = Actor ? PCGHelpers::GetActorBounds(Actor, /*bIgnorePCGCreatedComponents=*/false) : FBox{ EForceInit::ForceInit };
-	Bounds.Add(ActorBounds);
-	// Also provide previous bounds if we do have previous data
-	if (InPreviousData)
-	{
-		const FBox& PreviousBounds = InPreviousData->Get<0>();
-		if (PreviousBounds.IsValid && !PreviousBounds.Equals(ActorBounds))
+		if (ChangeTracker->OnNotifyObjectChanged(InObject, InChangeType, InOriginatingChangeObject))
 		{
-			Bounds.Add(PreviousBounds);
+			return;
 		}
 	}
 
-	OnSelectionKeyChanged(SelectionKey, InOriginatingChangeObject, RemovedTags, bNoRefreshOwner, Bounds);
+	OnObjectChanged(InObject, InChangeType, InOriginatingChangeObject);
+}
+
+void FPCGTrackingManager::OnSelectionKeyChanged(const FPCGSelectionKey& InSelectionKey)
+{
+	const UObject* OriginatingObject = InSelectionKey.GetObjectFromPath();
+	check(OriginatingObject);
+
+	OnSelectionKeyChanged(InSelectionKey, OriginatingObject, {}, /*bInNoRefreshOwner=*/false, { FBox(EForceInit::ForceInit) });
 }
 
 void FPCGTrackingManager::OnSelectionKeyChanged(const FPCGSelectionKey& InSelectionKey, const UObject* InOriginatingChangeObject, TArrayView<const FBox> InBounds)
@@ -1919,8 +1465,6 @@ void FPCGTrackingManager::OnSelectionKeyChanged(const FPCGSelectionKey& InSelect
 	const UObject* ChangedObject = InSelectionKey.GetObjectFromPath();
 	ensure(ChangedObject);
 
-	const AActor* Actor = Cast<AActor>(ChangedObject);
-
 	OnSelectionKeyChanged(InSelectionKey, InOriginatingChangeObject, {}, /*bInNoRefreshOwner=*/false, InBounds);
 }
 
@@ -1931,20 +1475,21 @@ void FPCGTrackingManager::OnSelectionKeyChanged(const FPCGSelectionKey& InSelect
 	/** First gather all the execution sources that are tracking this key */
 	TSet<IPCGGraphExecutionSource*> CulledTrackedExecutionSources;
 	TSet<IPCGGraphExecutionSource*> AlwaysTrackedExecutionSources;
-	TArray<FPCGSelectionKey> MatchedKeys;
-	
+
 	FPCGSelectionKey InOriginatingSelectionKey = FPCGSelectionKey::CreateFromObjectPtr(InOriginatingChangeObject);
+		
+	TSharedRef<IPCGChangeHandler::FPCGChangeHandlerChange> Change = MakeShared<IPCGChangeHandler::FPCGChangeHandlerChange>();
+	Change->ChangedObject = InSelectionKey.GetObjectFromPath();
+	Change->OriginatingChangeObject = InOriginatingChangeObject;
+	Change->bInNoRefreshOwner = bInNoRefreshOwner;
 	
-	const UObject* ChangedObject = InSelectionKey.GetObjectFromPath();
-	ensure(ChangedObject);
-
-	auto Gather = [InSelectionKey, InOriginatingSelectionKey, &InRemovedTags, &MatchedKeys](const TMap<FPCGSelectionKey, TSet<IPCGGraphExecutionSource*>>& InMap, TSet<IPCGGraphExecutionSource*>& OutSet)
+	auto Gather = [InSelectionKey, InOriginatingSelectionKey, &InRemovedTags, &Change](const TMap<FPCGSelectionKey, TSet<IPCGGraphExecutionSource*>>& InMap, TSet<IPCGGraphExecutionSource*>& OutSet)
 	{
 		for (const auto& It : InMap)
 		{
 			if (It.Key.IsMatching(InSelectionKey, InRemovedTags, It.Value, &OutSet) || It.Key.IsMatching(InOriginatingSelectionKey, InRemovedTags, It.Value, &OutSet))
 			{
-				MatchedKeys.Add(It.Key);
+				Change->MatchedKeys.Add(It.Key);
 			}
 		}
 	};
@@ -1958,63 +1503,11 @@ void FPCGTrackingManager::OnSelectionKeyChanged(const FPCGSelectionKey& InSelect
 		return;
 	}
 
-	TArray<const UObject*, TInlineAllocator<2>> ChangedObjects = { ChangedObject };
-	ChangedObjects.AddUnique(InOriginatingChangeObject);
-
-	TOptional<FPCGDynamicTrackingPriority> ChangedObjectDynamicTrackingPriority;
-	// @todo_pcg: Find a better way to get this info as this doesn't work for managed objects that were spawned under a different outer
-	if (AActor* OuterActor = InOriginatingChangeObject ? InOriginatingChangeObject->GetTypedOuter<AActor>() : nullptr)
+	// Prepare Change handlers
+	for (TUniquePtr<IPCGChangeHandler>& Handler : ChangeHandlers)
 	{
-		TArray<UPCGComponent*> OriginatingComponents;
-		OuterActor->GetComponents<UPCGComponent>(OriginatingComponents);
-		for (UPCGComponent* OriginatingComponent : OriginatingComponents)
-		{
-			if (!OriginatingComponent->IsCleaningUp() && OriginatingComponent->IsAnyObjectManagedByResource(ChangedObjects))
-			{
-				ChangedObjectDynamicTrackingPriority = OriginatingComponent->GetExecutionState().GetDynamicTrackingPriority();
-				break;
-			}
-		}
-	}
-
-	// We discard the component if:
-	//   * It is null.
-	//   * It is one of the objects that changed.
-	//   * We should not refresh the owner and its owner is one of the objects that changed.
-	//   * The objects that changed are managed by it (safe-guard, would probably induce an infinite loop).
-	//   * The objects that changed are managed by a lower tracking priority component. (lower tracking priority value means higher priority)
-	//   * One of the objects that changed is ignored by its original component.
-	auto ShouldDiscardComponent = [&ChangedObjects, &ChangedObjectDynamicTrackingPriority, bInNoRefreshOwner](const UPCGComponent* InComponent)
-	{
-		if (!InComponent
-			|| ChangedObjects.Contains(InComponent)
-			|| (bInNoRefreshOwner && ChangedObjects.Contains(InComponent->GetOwner()))
-			|| InComponent->IsCleaningUp()
-			|| InComponent->IsAnyObjectManagedByResource(ChangedObjects)
-			|| (ChangedObjectDynamicTrackingPriority.IsSet() && InComponent->GetExecutionState().GetDynamicTrackingPriority() < ChangedObjectDynamicTrackingPriority.GetValue()))
-		{
-			return true;
-		}
-
-		const UObject* ObjectIgnored = nullptr;
-		const UPCGComponent* OriginalComponent = InComponent->GetConstOriginalComponent();
-		EPCGIgnoreChangeOriginReason Reason;
-
-		if (ensure(OriginalComponent) && OriginalComponent->IsIgnoringAnyChangeOrigins(ChangedObjects, ObjectIgnored, Reason))
-		{
-			PCGGraphExecutionLogging::LogChangeOriginIgnoredForComponent(ObjectIgnored, InComponent, Reason);
-			return true;
-		}
-
-		return false;
-	};
-
-	TSet<UPCGComponent*> DirtyComponents;
-
-	EPCGComponentDirtyFlag DirtyFlag = EPCGComponentDirtyFlag::Actor;
-	if (ChangedObject->IsA<ALandscapeProxy>())
-	{
-		DirtyFlag = DirtyFlag | EPCGComponentDirtyFlag::Landscape;
+		// Give handler chance to store some information to be used to handle the current change
+		Handler->BeginChangeHandling(Change);
 	}
 
 	{
@@ -2022,30 +1515,9 @@ void FPCGTrackingManager::OnSelectionKeyChanged(const FPCGSelectionKey& InSelect
 
 		for (IPCGGraphExecutionSource* ExecutionSource : AlwaysTrackedExecutionSources)
 		{
-			UPCGComponent* PCGComponent = Cast<UPCGComponent>(ExecutionSource);
-			// @todo_pcg: support execution source
-			if (ShouldDiscardComponent(PCGComponent))
+			for (TUniquePtr<IPCGChangeHandler>& Handler : ChangeHandlers)
 			{
-				continue;
-			}
-
-			// Don't mark "Owner changed" if the change originate from a PCG Component. It will be delegated to the ClearCacheForActor.
-			// It is necessary to avoid infine loops when there are multiple PCG components on one actor, and one component was generated.
-			const bool bOwnerChanged = (PCGComponent->GetOwner() == ChangedObject) && (!InOriginatingChangeObject || !InOriginatingChangeObject->IsA<UPCGComponent>());
-			bool bShouldDirty = bOwnerChanged || (PCGComponent->ShouldTrackLandscape() && ChangedObject->IsA<ALandscapeProxy>());
-
-			// Since when we clear the cache for a settings, we clear it all, it's only necessary to do it once on the original component. And we will dirty all the local components (with dirty dispatch)
-			if (!DirtyComponents.Contains(PCGComponent))
-			{
-				// Technically, a key that is always tracked can also be culled tracked, so maybe having an intersection test here to not
-				// dirty components that culled tracked and don't intersect with the component would be more efficient.
-				// In practice, it's mostly relevent for partitioned graphs and we will dirty all the locals anyway.
-				bShouldDirty |= ClearCacheForKeys(MatchedKeys, PCGComponent, /*bIntersect=*/true, InOriginatingChangeObject);
-				if (bShouldDirty)
-				{
-					PCGComponent->DirtyGenerated(DirtyFlag, /*bDispatchToLocalComponents=*/true);
-					DirtyComponents.Add(PCGComponent);
-				}
+				Handler->HandleChange(ExecutionSource);
 			}
 		}
 	}
@@ -2055,78 +1527,35 @@ void FPCGTrackingManager::OnSelectionKeyChanged(const FPCGSelectionKey& InSelect
 		// Then do an octree find to get all components that intersect with this actor.
 		// If the actor has moved, we also need to find components that intersected with it before
 		// We first do it for non-partitioned, then we do it for partitioned
-		auto UpdateNonPartitioned = [this, &DirtyComponents, ChangedObject, CulledTrackedExecutionSources, &InRemovedTags, DirtyFlag, InOriginatingChangeObject, bInNoRefreshOwner, &MatchedKeys, &ShouldDiscardComponent](const FPCGExecutionSourceRef& ExecutionSourceRef, const FBox& InBounds) -> void
+		auto UpdateNonPartitioned = [this, CulledTrackedExecutionSources, InOriginatingChangeObject](const FPCGExecutionSourceRef& ExecutionSourceRef, const FBox& InBounds) -> void
 		{
 			TRACE_CPUPROFILER_EVENT_SCOPE(FPCGTrackingManager::OnObjectChanged::UpdateNonPartitioned);
 
-			if (UPCGComponent* PCGComponent = Cast<UPCGComponent>(ExecutionSourceRef.ExecutionSource))
+			if (!CulledTrackedExecutionSources.Contains(ExecutionSourceRef.ExecutionSource))
 			{
-				// Don't dirty if the component was already dirtied, not tracked, the origin of the change or its owner is the changed actor and we should not refresh.
-				if (DirtyComponents.Contains(PCGComponent) ||
-					!CulledTrackedExecutionSources.Contains(PCGComponent) ||
-					ShouldDiscardComponent(PCGComponent))
-				{
-					return;
-				}
+				return;
+			}
 
-				bool bShouldDirty = PCGComponent->ShouldTrackLandscape() && ChangedObject->IsA<ALandscapeProxy>();
-				bShouldDirty |= ClearCacheForKeys(MatchedKeys, PCGComponent, /*bIntersect=*/true, InOriginatingChangeObject);
-				if (bShouldDirty)
-				{
-					PCGComponent->DirtyGenerated(DirtyFlag);
-					DirtyComponents.Add(PCGComponent);
-				}
+			for (TUniquePtr<IPCGChangeHandler>& Handler : ChangeHandlers)
+			{
+				Handler->HandleBoundedChange(ExecutionSourceRef.ExecutionSource, ExecutionSourceRef.Bounds.GetBox(), InBounds);
 			}
 		};
 
 		// For partitioned, we first need check if the original component intersect with the bounds, then forward the dirty call only to locals that intersect with the bounds.
 		// Note: CurrentActorBoundsPtr is passed by reference because it will be modified between lambda calls (cf comment above).
-		auto UpdatePartitioned = [this, &DirtyComponents, ChangedObject, CulledTrackedExecutionSources, &InRemovedTags, DirtyFlag, InOriginatingChangeObject, bInNoRefreshOwner, &MatchedKeys, &ShouldDiscardComponent](const FPCGExecutionSourceRef& ExecutionSourceRef, const FBox& InBounds)  -> void
+		auto UpdatePartitioned = [this, CulledTrackedExecutionSources, InOriginatingChangeObject](const FPCGExecutionSourceRef& ExecutionSourceRef, const FBox& InBounds)  -> void
 		{
 			TRACE_CPUPROFILER_EVENT_SCOPE(FPCGTrackingManager::OnObjectChanged::UpdatePartitioned);
 
-			if (UPCGComponent* PCGComponent = Cast<UPCGComponent>(ExecutionSourceRef.ExecutionSource))
+			if (!CulledTrackedExecutionSources.Contains(ExecutionSourceRef.ExecutionSource))
 			{
-				// Don't dirty if the component is not tracked, the origin of the change or its owner is the changed actor and we should not refresh.
-				// We can "re-dirty" it because changes can impact different local components, from the same
-				// original component.
-				if (!CulledTrackedExecutionSources.Contains(PCGComponent) ||
-					ShouldDiscardComponent(PCGComponent))
-				{
-					return;
-				}
-
-				const FBox ComponentBounds = ExecutionSourceRef.Bounds.GetBox();
-				const bool bIntersect = InBounds.Intersect(ComponentBounds);
-				if (!bIntersect)
-				{
-					return;
-				}
-
-				const FBox Overlap = InBounds.Overlap(ComponentBounds);
-				bool bShouldDirty = PCGComponent->ShouldTrackLandscape() && ChangedObject->IsA<ALandscapeProxy>();
-				bShouldDirty |= ClearCacheForKeys(MatchedKeys, PCGComponent, /*bIntersect=*/true, InOriginatingChangeObject);
-				bool bWasDirtied = false;
-
-				// Since when we clear the cache for a settings, we clear it all, it's only necessary to do it once on the original component, then only dirty the local that intersects.
-				if (bShouldDirty)
-				{
-					ForAllIntersectingPartitionActors(Overlap, [ChangedObject, Component = PCGComponent, &bWasDirtied, DirtyFlag, InOriginatingChangeObject](APCGPartitionActor* InPartitionActor) -> void
-						{
-							if (UPCGComponent* LocalComponent = InPartitionActor->GetLocalComponent(Component))
-							{
-								bWasDirtied = true;
-								LocalComponent->DirtyGenerated(DirtyFlag);
-							}
-						});
-				}
+				return;
+			}
 
-				if (bWasDirtied)
-				{
-					// Don't dispatch
-					PCGComponent->DirtyGenerated(DirtyFlag, /*bDispatchToLocalComponents=*/false);
-					DirtyComponents.Add(PCGComponent);
-				}
+			for (TUniquePtr<IPCGChangeHandler>& Handler : ChangeHandlers)
+			{
+				Handler->HandleBoundedChange(ExecutionSourceRef.ExecutionSource, ExecutionSourceRef.Bounds.GetBox(), InBounds);
 			}
 		};
 
@@ -2138,139 +1567,21 @@ void FPCGTrackingManager::OnSelectionKeyChanged(const FPCGSelectionKey& InSelect
 		}
 	}
 
-	// Finally, dirty all components that always track this actor that are not yet notified.
-
-	// If it is a landscape and we should discard the refresh, early out.
-	bool bIsInEditingMode = false;
-	const ALandscapeProxy* Landscape = Cast<ALandscapeProxy>(ChangedObject);
-	if (PCGTrackingManager::ShouldDiscardLandscapeRefresh(Landscape, bIsInEditingMode, bIsCurrentlyExitingLandscapeEditMode))
+	// Allow trackers to skip refresh
+	bool bSkipRefresh = false;
+	for (TUniquePtr<IPCGChangeTracker>& ChangeTracker : ChangeTrackers)
 	{
-		// If we are in editing, keep track of all the dirtied landscape to refresh them when we exit.
-		if (bIsInEditingMode)
+		if (ChangeTracker->ShouldSkipRefresh(Change->ChangedObject))
 		{
-			DirtiedLandscapes.AddUnique(Landscape);
+			bSkipRefresh = true;
+			break;
 		}
-
-		return;
 	}
 
-	ULevelInstanceSubsystem* LevelInstanceSubsystem = GetWorld()->GetSubsystem<ULevelInstanceSubsystem>();
-	const UPCGComponent* OriginatingComponent = Cast<const UPCGComponent>(InOriginatingChangeObject);
-
-	// And refresh all dirtied components
-	for (UPCGComponent* Component : DirtyComponents)
+	for (TUniquePtr<IPCGChangeHandler>& Handler : ChangeHandlers)
 	{
-		if (!ensure(Component))
-		{
-			continue;
-		}
-
-		// This part checks if the change originates from a PCG Component. If so, we check the dirty component has no more other PCG dependencies.
-		// In that case we can proceed for the refresh, otherwise early out, this will be woken up by another dependency change.
-		if (OriginatingComponent)
-		{
-			if (TArray<FObjectKey>* Dependencies = ExecutionSourcesToDependencyMap.Find(Component))
-			{
-				Dependencies->RemoveSwap(OriginatingComponent);
-				if (Dependencies->IsEmpty())
-				{
-					ExecutionSourcesToDependencyMap.Remove(Component);
-				}
-				else
-				{
-					continue;
-				}
-			}
-		}
-
-		const bool bOwnerHasChanged = Component->GetOwner() == ChangedObject;
-
-		if ((!bInNoRefreshOwner || !bOwnerHasChanged) && (!Component->bOnlyTrackItself || bOwnerHasChanged))
-		{
-			const AActor* Actor = Cast<AActor>(ChangedObject);
-
-			// When an object changes, we need to make sure that we don't trigger a refresh on PCG components that are "higher" in the
-			// level hierarchy, otherwise we will end up generating in the Level Instance level, which is wrong.
-			// Note that in some instances (e.g. when something happens at the Level Instance level) we need to make sure that the
-			// PCG components higher-up are properly updated, hence the level instance depth
-			if (LevelInstanceSubsystem && Actor)
-			{
-				// Immediate Level Instance Parent (Actor)
-				const ILevelInstanceInterface* ActorLevelInstance = LevelInstanceSubsystem->GetParentLevelInstance(Actor);
-
-				// Immediate Level Instance Parent (PCG Component)
-				const ILevelInstanceInterface* ComponentLevelInstance = LevelInstanceSubsystem->GetParentLevelInstance(Component->GetOwner());
-								
-				// Is the Actor under a Level Instance that is the same or under the Components owning Level Instance
-				bool bActorLevelInstanceInComponentLevelInstanceHierarchy = false;
-
-				// Go up Level Instance hierarchy for Actor and see if we are part of the Component's level instance hierarchy, stop if we find an Editing Level instance as we don't want to impact top level 
-				// Components from an editing level instance
-				if (ComponentLevelInstance)
-				{
-					LevelInstanceSubsystem->ForEachLevelInstanceAncestors(Actor, [ComponentLevelInstance, &bActorLevelInstanceInComponentLevelInstanceHierarchy](const ILevelInstanceInterface* Ancestor)
-					{
-						if (Ancestor == ComponentLevelInstance)
-						{
-							bActorLevelInstanceInComponentLevelInstanceHierarchy = true;
-							return false; // stop iterating
-						}
-
-						return !Ancestor->IsEditing();
-					});
-				}
-				
-				// Actor in same level instance as Component or under the Components Level instance
-				if (bActorLevelInstanceInComponentLevelInstanceHierarchy)
-				{
-					// Allow update if Component is inside Editing level instance or if it is only in preview mode (generates transient data because it is in a non editing Level instance)
-					if (!ComponentLevelInstance->IsEditing() && Component->GetEditingMode() != EPCGEditorDirtyMode::Preview)
-					{
-						continue;
-					}
-				} // Allow update of Component from a Editing Level instance contained actor if the Component being refreshed is in preview mode and not itself in a Level Instance
-				else if (Actor->IsInEditLevelInstanceHierarchy() && (Component->GetEditingMode() != EPCGEditorDirtyMode::Preview || ComponentLevelInstance != nullptr))
-				{
-					continue;
-				}
-				// If ComponentLevelInstance is non-null and Actor Level instance isn't in its hierarchy then skip it
-				else if (ActorLevelInstance && ComponentLevelInstance != nullptr)
-				{
-					continue;
-				}
-			}
-
-			Component->Refresh();
-		}
-	}
-}
-
-void FPCGTrackingManager::OnLandscapeChanged(ALandscapeProxy* InLandscape, const FLandscapeProxyComponentDataChangedParams& InChangeParams)
-{
-	if (!InLandscape)
-	{
-		return;
-	}
-
-	if (PCGTrackingManager::CVarLandscapeRefreshTimeDelay.GetValueOnAnyThread() > 0)
-	{
-		LastLandscapeDirtyTime = FApp::GetCurrentTime();
-		DelayedModifiedLandscapes.AddUnique(InLandscape);
+		Handler->EndChangeHandling(bSkipRefresh);
 	}
-	else
-	{
-		ApplyLandscapeChanges(InLandscape);
-	}
-}
-
-void FPCGTrackingManager::ApplyLandscapeChanges(ALandscapeProxy* InLandscape)
-{
-	if (!InLandscape)
-	{
-		return;
-	}
-
-	OnObjectChanged(InLandscape, EChangedObjectReason::ActorChanged, /*InPreviousData=*/nullptr, InLandscape);
 }
 
 void FPCGTrackingManager::OnPCGGraphStartsGenerating(IPCGGraphExecutionSource* InExecutionSource)
@@ -2395,23 +1706,16 @@ void FPCGTrackingManager::OnPCGGraphGeneratedOrCleaned(IPCGGraphExecutionSource*
 	// @todo_pcg: support execution source
 	if (UPCGComponent* PCGComponent = Cast<UPCGComponent>(InExecutionSource))
 	{
-		OnObjectChanged(PCGComponent->GetOwner(), EChangedObjectReason::PCGComponentGenerated, /*InPreviousData=*/nullptr, PCGComponent);
+		OnObjectChanged(PCGComponent->GetOwner(), PCGTrackingManager::ExecutionSourceGenerated, PCGComponent);
 	}
 }
 
-bool FPCGTrackingManager::ClearCacheForKeys(const TArray<FPCGSelectionKey>& InKeys, const IPCGGraphExecutionSource* InExecutionSource, const bool bIntersect, const UObject* InOriginatingChange) const
+bool FPCGTrackingManager::ClearCacheForKeys(TArrayView<FPCGSelectionKey> InKeys, const IPCGGraphExecutionSource* InExecutionSource, const bool bIntersect, const UObject* InOriginatingChange) const
 {
 	check(InExecutionSource);
-
-	// @todo_pcg: support execution source
-	const UPCGComponent* PCGComponent = Cast<UPCGComponent>(InExecutionSource);
-	if (!PCGComponent)
-	{
-		return false;
-	}
-
-	bool bShouldDirty = PCGComponent && PCGComponent->ShouldTrackLandscape() && InOriginatingChange && InOriginatingChange->IsA<ALandscapeProxy>();
 	
+	bool bShouldDirty = false;
+
 	auto ClearCache = [this, InOriginatingChange, bIntersect, &bShouldDirty](const FPCGSelectionKey& Key, const FPCGSettingsAndCulling& SettingsAndCulling)
 	{
 		if (!SettingsAndCulling.Key.IsValid() || (SettingsAndCulling.Value && !bIntersect))
@@ -2419,6 +1723,7 @@ bool FPCGTrackingManager::ClearCacheForKeys(const TArray<FPCGSelectionKey>& InKe
 			return;
 		}
 
+		// todo_pcg: validate if this is needed for all execution source types
 		// Extra care if the change originates from a PCGComponent. Only dirty if we are tracking a PCG component.
 		if (InOriginatingChange && InOriginatingChange->IsA<UPCGComponent>()
 			&& (!Key.OptionalExtraDependency || !Key.OptionalExtraDependency->IsChildOf(UPCGComponent::StaticClass())))
@@ -2433,30 +1738,10 @@ bool FPCGTrackingManager::ClearCacheForKeys(const TArray<FPCGSelectionKey>& InKe
 
 	for (const FPCGSelectionKey& Key : InKeys)
 	{
-		PCGComponent->ApplyToEachSettings(Key, ClearCache);
+		InExecutionSource->GetExecutionState().ForEachSettingTrackingKey(Key, ClearCache);
 	}
 
 	return bShouldDirty;
 }
 
-void FPCGTrackingManager::NotifyLandscapeEditModeExited()
-{
-	bIsCurrentlyExitingLandscapeEditMode = true;
-	// When the landscape edit mode is exited, force the refresh on all modified/dirtied landscapes.
-	for (TObjectKey<ALandscapeProxy> Landscape : DelayedModifiedLandscapes)
-	{
-		DirtiedLandscapes.AddUnique(Landscape);
-	}
-
-	DelayedModifiedLandscapes.Empty();
-
-	for (TObjectKey<ALandscapeProxy> Landscape : DirtiedLandscapes)
-	{
-		ApplyLandscapeChanges(Landscape.ResolveObjectPtr());
-	}
-
-	DirtiedLandscapes.Empty();
-
-	bIsCurrentlyExitingLandscapeEditMode = false;
-}
 #endif // WITH_EDITOR
diff --git a/Engine/Plugins/PCG/Source/PCG/Private/PCGTrackingManager.h b/Engine/Plugins/PCG/Source/PCG/Private/PCGTrackingManager.h
deleted file mode 100644
diff --git a/Engine/Plugins/PCG/Source/PCG/Private/Subsystems/PCGSubsystem.cpp b/Engine/Plugins/PCG/Source/PCG/Private/Subsystems/PCGSubsystem.cpp
--- a/Engine/Plugins/PCG/Source/PCG/Private/Subsystems/PCGSubsystem.cpp
+++ b/Engine/Plugins/PCG/Source/PCG/Private/Subsystems/PCGSubsystem.cpp
@@ -1797,11 +1797,6 @@ void UPCGSubsystem::NotifySelectionKeyChanged(const FPCGSelectionKey& InSelectio
 	TrackingManager->OnSelectionKeyChanged(InSelectionKey, InOriginatingChangeObject, InChangedBounds);
 }
 
-void UPCGSubsystem::NotifyLandscapeEditModeExited()
-{
-	TrackingManager->NotifyLandscapeEditModeExited();
-}
-
 void UPCGSubsystem::CreateMissingPartitionActors()
 {
 	if (!PCGHelpers::IsRuntimeOrPIE())
@@ -1945,9 +1940,9 @@ void UPCGSubsystem::OnPCGGraphCleaned(IPCGGraphExecutionSource* InExecutionSourc
 	TrackingManager->OnPCGGraphGeneratedOrCleaned(InExecutionSource);
 }
 
-void UPCGSubsystem::UpdateComponentTracking(IPCGGraphExecutionSource* InExecutionSource, bool bShouldDirtyActors, const TArray<FPCGSelectionKey>* OptionalChangedKeys)
+void UPCGSubsystem::UpdateTracking(IPCGGraphExecutionSource* InExecutionSource, const TArray<FPCGSelectionKey>* OptionalChangedKeys)
 {
-	TrackingManager->UpdateTracking(InExecutionSource, bShouldDirtyActors, OptionalChangedKeys);
+	TrackingManager->UpdateTracking(InExecutionSource, OptionalChangedKeys);
 }
 
 void UPCGSubsystem::ResetPartitionActorsMap()
diff --git a/Engine/Plugins/PCG/Source/PCG/Private/Utils/PCGGraphExecutionLogging.cpp b/Engine/Plugins/PCG/Source/PCG/Private/Utils/PCGGraphExecutionLogging.cpp
--- a/Engine/Plugins/PCG/Source/PCG/Private/Utils/PCGGraphExecutionLogging.cpp
+++ b/Engine/Plugins/PCG/Source/PCG/Private/Utils/PCGGraphExecutionLogging.cpp
@@ -1,6 +1,6 @@
 // Copyright Epic Games, Inc. All Rights Reserved.
 
-#include "PCGGraphExecutionLogging.h"
+#include "Utils/PCGGraphExecutionLogging.h"
 
 #include "PCGCommon.h"
 #include "PCGComponent.h"
diff --git a/Engine/Plugins/PCG/Source/PCG/Private/Utils/PCGGraphExecutionLogging.h b/Engine/Plugins/PCG/Source/PCG/Private/Utils/PCGGraphExecutionLogging.h
deleted file mode 100644
diff --git a/Engine/Plugins/PCG/Source/PCG/Public/ChangeTracking/PCGChangeTracker.h b/Engine/Plugins/PCG/Source/PCG/Public/ChangeTracking/PCGChangeTracker.h
new file mode 100644
--- /dev/null
+++ b/Engine/Plugins/PCG/Source/PCG/Public/ChangeTracking/PCGChangeTracker.h
@@ -0,0 +1,48 @@
+// Copyright Epic Games, Inc. All Rights Reserved.
+
+#pragma once
+
+#if WITH_EDITOR
+
+#include "PCGCommon.h"
+#include "PCGSettings.h"
+#include "Elements/PCGActorSelector.h"
+
+class UPCGGraph;
+class IPCGGraphExecutionSource;
+
+struct FPCGChangeTracker
+{
+public:
+	PCG_API void AddLocalSourceDynamicTrackingKeys(const FPCGSelectionKeyToSettingsMap& InLocalSourceDynamicTrackingKeys);
+	PCG_API void AddLocalSourceCurrentTrackingKeys(const FPCGSelectionKeyToSettingsMap& InLocalSourceCurrentExecutionTrackingKeys, const TSet<const UPCGSettings*>& InLocalSourceCurrentExecutionTrackingSettings);
+	PCG_API void RegisterDynamicTracking(const UPCGSettings* InSettings, const TArrayView<TPair<FPCGSelectionKey, bool>>& InDynamicKeysAndCulling);
+	PCG_API void RegisterDynamicTracking(const FPCGSelectionKeyToSettingsMap& InKeysToSettings);
+	PCG_API void UpdateDynamicTracking(IPCGGraphExecutionSource* InExecutionSource);
+	PCG_API void ResetCurrentExecution();
+
+	PCG_API TArray<FPCGSelectionKey> GatherTrackingKeys() const;
+	PCG_API bool IsKeyTrackedAndCulled(const FPCGSelectionKey& Key, bool& bOutIsCulled) const;
+	PCG_API bool UpdateTrackingCache(const UPCGGraph* InGraph, TArray<FPCGSelectionKey>* OptionalChangedKeys);
+	PCG_API void ForEachSettingTrackingKey(const FPCGSelectionKey& InKey, const TFunctionRef<void(const FPCGSelectionKey&, const FPCGSettingsAndCulling&)> InCallback) const;
+	PCG_API void Serialize(FArchive& Ar);
+
+private: 
+	// @todo_pcg: remove friending and expose what is needed
+	friend class UPCGComponent;
+	friend struct FPCGComponentInstanceData;
+
+	FPCGSelectionKeyToSettingsMap StaticallyTrackedKeysToSettings;
+
+	// Temporary storage for dynamic tracking that will be filled during execution source execution.
+	FPCGSelectionKeyToSettingsMap CurrentExecutionDynamicTracking;
+	// Temporary storage for dynamic tracking that will keep all settings that could have dynamic tracking, in order to detect changes.
+	TSet<const UPCGSettings*> CurrentExecutionDynamicTrackingSettings;
+	mutable FTransactionallySafeCriticalSection CurrentExecutionDynamicTrackingLock;
+
+	// Need to keep a reference to all tracked settings to still react to changes after a map load (since the execution source won't have been executed).
+	// Serialization will be done in the Serialize function
+	FPCGSelectionKeyToSettingsMap DynamicallyTrackedKeysToSettings;
+};
+
+#endif
\ No newline at end of file
diff --git a/Engine/Plugins/PCG/Source/PCG/Public/ChangeTracking/PCGChangeTrackingRegistry.h b/Engine/Plugins/PCG/Source/PCG/Public/ChangeTracking/PCGChangeTrackingRegistry.h
new file mode 100644
--- /dev/null
+++ b/Engine/Plugins/PCG/Source/PCG/Public/ChangeTracking/PCGChangeTrackingRegistry.h
@@ -0,0 +1,101 @@
+// Copyright Epic Games, Inc. All Rights Reserved.
+
+#pragma once
+
+#if WITH_EDITOR
+
+#include "PCGCommon.h"
+#include "Delegates/Delegate.h"
+#include "Elements/PCGActorSelector.h"
+
+class FPCGTrackingManager;
+
+// Unique change types
+using FPCGTrackingChangeID = FGuid;
+
+// This interface provides change tracking for different systems that may live across different plugins and have those trackers send out formatted events to their FPCGTrackingManager owner.
+class IPCGChangeTracker
+{
+public:
+	IPCGChangeTracker(FPCGTrackingManager* InOwner)
+		: Owner(InOwner)
+	{
+		check(Owner);
+	}
+
+	virtual ~IPCGChangeTracker() = default;
+
+	// Called every frame from FPCGTrackingManager.
+	virtual void Tick() {}
+
+	// Instead of hooking up to the core delegate, the FPCGTrackingManager will call this method and returning true means that this tracker is handling the event.
+	virtual bool OnObjectPropertyChanged(const UObject* InObject, FPropertyChangedEvent& InEvent) { return false; }
+
+	// Can be implemented to react to other tracker events (currently used to send out Landscape tracker event to Actor tracker).
+	virtual bool OnNotifyObjectChanged(const UObject* InObject, const FPCGTrackingChangeID& InChangeType, const UObject* InOriginatingChangeObject) { return false; }
+
+	// Called from FPCGTrackingManager::OnSelectionKeyChanged to allow skipping the execution source refresh
+	virtual bool ShouldSkipRefresh(const UObject* InChangedObject) { return false; }
+
+protected:
+	FPCGTrackingManager* Owner = nullptr;
+};
+
+// This interface provides handling of changes once they are processed through the FPCGTrackingManager
+class IPCGChangeHandler
+{
+public:
+	IPCGChangeHandler(FPCGTrackingManager* InOwner)
+		: Owner(InOwner)
+	{
+		check(Owner);
+	}
+
+	struct FPCGChangeHandlerChange
+	{
+	public:
+		// Root object to attribute change to
+		const UObject* ChangedObject = nullptr;
+		// Object the change originates from (can be sub-object of ChangedObject)
+		const UObject* OriginatingChangeObject = nullptr;
+		// Tracked keys that match the ChangedObject or OriginatingChangedObject
+		TArray<FPCGSelectionKey> MatchedKeys;
+		// Flag passed down to FPCGTrackingManager::OnSelectionKeyChanged
+		bool bInNoRefreshOwner = false;
+	};
+
+	virtual ~IPCGChangeHandler() = default;
+
+	virtual void BeginChangeHandling(const TSharedRef<FPCGChangeHandlerChange>& InChange) {}
+	virtual void HandleChange(IPCGGraphExecutionSource* InExecutionSource) {}
+	virtual void HandleBoundedChange(IPCGGraphExecutionSource* InExecutionSource, const FBox& InExecutionSourceBounds, const FBox& InChangeBounds) {}
+	virtual void EndChangeHandling(bool bSkipRefresh) {}
+
+protected:
+	FPCGTrackingManager* Owner = nullptr;
+};
+
+DECLARE_DELEGATE_RetVal_OneParam(TUniquePtr<IPCGChangeTracker>, FPCGCreateTrackerInstance, FPCGTrackingManager*);
+DECLARE_DELEGATE_RetVal_OneParam(TUniquePtr<IPCGChangeHandler>, FPCGCreateHandlerInstance, FPCGTrackingManager*);
+
+class FPCGChangeTrackingRegistry
+{
+public:
+	// Creates the tracker instances for the FPCGTrackingManager
+	TArray<TUniquePtr<IPCGChangeTracker>> CreateChangeTrackers(FPCGTrackingManager* InManager) const;
+	TArray<TUniquePtr<IPCGChangeHandler>> CreateChangeHandlers(FPCGTrackingManager* InManager) const;
+
+	// Register/unregister custom tracker factories (allows other modules to register new trackers)
+	PCG_API void RegisterTrackerFactory(FName InName, FPCGCreateTrackerInstance InTrackerFactory);
+	PCG_API void UnregisterTrackerFactory(FName InName);
+
+	// Register/unregister custom handler factories (allows other modules to register new handlers)
+	PCG_API void RegisterHandlerFactory(FName InName, FPCGCreateHandlerInstance InHandlerFactory);
+	PCG_API void UnregisterHandlerFactory(FName InName);
+
+private:
+	TMap<FName, FPCGCreateTrackerInstance> TrackerFactories;
+	TMap<FName, FPCGCreateHandlerInstance> HandlerFactories;
+};
+
+#endif // WITH_EDITOR
\ No newline at end of file
diff --git a/Engine/Plugins/PCG/Source/PCG/Public/PCGComponent.h b/Engine/Plugins/PCG/Source/PCG/Public/PCGComponent.h
--- a/Engine/Plugins/PCG/Source/PCG/Public/PCGComponent.h
+++ b/Engine/Plugins/PCG/Source/PCG/Public/PCGComponent.h
@@ -9,6 +9,7 @@
 #include "PCGManagedResourceContainer.h"
 #include "PCGNode.h"
 #include "PCGSettings.h"
+#include "ChangeTracking/PCGChangeTracker.h"
 #include "Data/PCGToolData.h"
 #include "Graph/PCGStackContext.h"
 #include "Grid/PCGGridDescriptor.h"
@@ -443,6 +444,7 @@ public:
 	/** Did the given node produce one or more data items in the given stack during the last execution. */
 	UE_API bool HasNodeProducedData(const UPCGNode* InNode, const FPCGStack& Stack) const;
 
+	UE_DEPRECATED(5.8, "No longer supported")
 	UE_API bool IsObjectTracked(const UObject* InObject, bool& bOutIsCulled) const;
 
 	/** Know if we need to force a generation, in case of BP added to the world in editor */
@@ -553,6 +555,9 @@ public:
 
 #if WITH_EDITOR
 
+	UE_DEPRECATED(all, "Do not use")
+	UE_API bool ShouldTrackLandscape() const;
+
 	UE_DEPRECATED(5.6, "Use FPCGGraphExecutionInspection::IsInspecting() instead")
 	UE_API bool IsInspecting() const;
 
@@ -708,8 +713,6 @@ private:
 	/** Compare the temp map to the stored map for dynamic tracking and register/unregister accordingly. If it is a local component, it will push the info to the original. */
 	UE_API void UpdateDynamicTracking();
 
-	UE_API bool ShouldTrackLandscape() const;
-
 	UE_API void MarkResourcesAsTransientOnLoad();
 	UE_API bool DeletePreviewResources();
 
@@ -792,17 +795,7 @@ private:
 #endif
 
 #if WITH_EDITORONLY_DATA
-	FPCGSelectionKeyToSettingsMap StaticallyTrackedKeysToSettings;
-
-	// Temporary storage for dynamic tracking that will be filled during component execution.
-	FPCGSelectionKeyToSettingsMap CurrentExecutionDynamicTracking;
-	// Temporary storage for dynamic tracking that will keep all settings that could have dynamic tracking, in order to detect changes.
-	TSet<const UPCGSettings*> CurrentExecutionDynamicTrackingSettings;
-	mutable FTransactionallySafeCriticalSection CurrentExecutionDynamicTrackingLock;
-
-	// Need to keep a reference to all tracked settings to still react to changes after a map load (since the component won't have been executed).
-	// Serialization will be done in the Serialize function
-	FPCGSelectionKeyToSettingsMap DynamicallyTrackedKeysToSettings;
+	FPCGChangeTracker ChangeTracker;
 #endif
 
 #if WITH_EDITOR
diff --git a/Engine/Plugins/PCG/Source/PCG/Public/PCGComponentExecutionState.h b/Engine/Plugins/PCG/Source/PCG/Public/PCGComponentExecutionState.h
--- a/Engine/Plugins/PCG/Source/PCG/Public/PCGComponentExecutionState.h
+++ b/Engine/Plugins/PCG/Source/PCG/Public/PCGComponentExecutionState.h
@@ -62,6 +62,9 @@ public:
 
 	virtual void RegisterDynamicTracking(const UPCGSettings* InSettings, const TArrayView<TPair<FPCGSelectionKey, bool>>& InDynamicKeysAndCulling) override;
 	virtual void RegisterDynamicTracking(const FPCGSelectionKeyToSettingsMap& InKeysToSettings) override;
+	virtual TArray<FPCGSelectionKey> GatherTrackingKeys() const override;
+	virtual bool IsKeyTrackedAndCulled(const FPCGSelectionKey& Key, bool& bOutIsCulled) const override;
+	virtual void ForEachSettingTrackingKey(const FPCGSelectionKey& InKey, const TFunctionRef<void(const FPCGSelectionKey&, const FPCGSettingsAndCulling&)> InCallback) const override;
 	
 	virtual bool IsRefreshInProgress() const override;
 	virtual FPCGDynamicTrackingPriority GetDynamicTrackingPriority() const override;
diff --git a/Engine/Plugins/PCG/Source/PCG/Public/PCGGraph.h b/Engine/Plugins/PCG/Source/PCG/Public/PCGGraph.h
--- a/Engine/Plugins/PCG/Source/PCG/Public/PCGGraph.h
+++ b/Engine/Plugins/PCG/Source/PCG/Public/PCGGraph.h
@@ -483,8 +483,15 @@ public:
 #endif
 
 #if WITH_EDITOR
-	UE_API FPCGSelectionKeyToSettingsMap GetTrackedActorKeysToSettings() const;
-	UE_API void GetTrackedActorKeysToSettings(FPCGSelectionKeyToSettingsMap& OutKeysToSettings, TArray<TObjectPtr<const UPCGGraph>>& OutVisitedGraphs) const;
+	UE_API FPCGSelectionKeyToSettingsMap GetTrackedSelectionKeysToSettings() const;
+	UE_API void GetTrackedSelectionKeysToSettings(FPCGSelectionKeyToSettingsMap& OutKeysToSettings, TArray<TObjectPtr<const UPCGGraph>>& OutVisitedGraphs) const;
+	UE_API bool UpdateTrackedSelectionKeys(FPCGSelectionKeyToSettingsMap& InExistingKeys, FPCGSelectionKeyToSettingsMap& InOutUpdatedKeys, TArray<FPCGSelectionKey>* OptionalChangedKeys) const;
+
+	UE_DEPRECATED(5.8, "Use GetTrackedSelectionKeysToSettings instead")
+	FPCGSelectionKeyToSettingsMap GetTrackedActorKeysToSettings() const { return GetTrackedSelectionKeysToSettings(); }
+
+	UE_DEPRECATED(5.8, "Use GetTrackedSelectionKeysToSettings instead")
+	void GetTrackedActorKeysToSettings(FPCGSelectionKeyToSettingsMap& OutKeysToSettings, TArray<TObjectPtr<const UPCGGraph>>& OutVisitedGraphs) const { GetTrackedSelectionKeysToSettings(OutKeysToSettings, OutVisitedGraphs); }
 #endif
 
 	/** Size of grid on which this node should be executed. Nodes execute at the minimum of all input grid sizes. */
diff --git a/Engine/Plugins/PCG/Source/PCG/Public/PCGGraphExecutionStateInterface.h b/Engine/Plugins/PCG/Source/PCG/Public/PCGGraphExecutionStateInterface.h
--- a/Engine/Plugins/PCG/Source/PCG/Public/PCGGraphExecutionStateInterface.h
+++ b/Engine/Plugins/PCG/Source/PCG/Public/PCGGraphExecutionStateInterface.h
@@ -181,9 +181,23 @@ public:
 	/** Register multiple tracking dependencies, so ExecutionState can be updated when they change */
 	virtual void RegisterDynamicTracking(const FPCGSelectionKeyToSettingsMap& InKeysToSettings) = 0;
 
+	/** Returns the tracked keys */
+	virtual TArray<FPCGSelectionKey> GatherTrackingKeys() const { return {}; }
+
+	/** Return true if the key is tracked, and if so, bOutIsCulled will contains if the key is culled or not. */
+	virtual bool IsKeyTrackedAndCulled(const FPCGSelectionKey& Key, bool& bOutIsCulled) const { return false; }
+
+	/** Apply a function to all settings that track a given key. */
+	virtual void ForEachSettingTrackingKey(const FPCGSelectionKey& InKey, const TFunctionRef<void(const FPCGSelectionKey&, const FPCGSettingsAndCulling&)> InCallback) const {}
+
+	/** Add dynamic tracking keys from local source */
+	virtual void AddLocalSourceDynamicTrackingKeys(const FPCGSelectionKeyToSettingsMap& LocalSourceDynamicTrackingKeys) {}
+
+	/** Add current tracking keys from local source */
+	virtual void AddLocalSourceCurrentTrackingKeys(const FPCGSelectionKeyToSettingsMap& LocalSourceCurrentExecutionTrackingKeys, const TSet<const UPCGSettings*>& LocalSourceCurrentExecutionTrackingSettings) {}
+
 	/** Indicates if the current source is waiting for a refresh */
 	virtual bool IsRefreshInProgress() const = 0;
-
 #endif
 
 	/** If this is a local execution source returns the original execution source, otherwise returns self. */
diff --git a/Engine/Plugins/PCG/Source/PCG/Public/PCGModule.h b/Engine/Plugins/PCG/Source/PCG/Public/PCGModule.h
--- a/Engine/Plugins/PCG/Source/PCG/Public/PCGModule.h
+++ b/Engine/Plugins/PCG/Source/PCG/Public/PCGModule.h
@@ -4,6 +4,7 @@
 
 #include "PCGDataVisualizationRegistry.h"
 #include "PCGGraphExecutionRegistry.h"
+#include "ChangeTracking/PCGChangeTrackingRegistry.h"
 #include "Data/Registry/PCGDataTypeRegistry.h"
 #include "Data/Registry/PCGDataViewRegistry.h"
 #include "Data/Registry/PCGGetDataFunctionRegistry.h"
@@ -66,6 +67,9 @@ public:
 #if WITH_EDITOR
 	static const FPCGObjectHashFactory& GetConstObjectHashFactory() { return GetPCGModuleChecked().ObjectHashFactory; }
 	static FPCGObjectHashFactory& GetMutableObjectHashFactory() { return GetPCGModuleChecked().ObjectHashFactory; }
+
+	static const FPCGChangeTrackingRegistry& GetConstChangeTrackingRegistry() { return GetPCGModuleChecked().ChangeTrackingRegistry; }
+	static FPCGChangeTrackingRegistry& GetMutableChangeTrackingRegistry() { return GetPCGModuleChecked().ChangeTrackingRegistry; }
 #endif
 
 	PCG_API static bool IsPCGModuleLoaded();
@@ -91,6 +95,7 @@ private:
 	FPCGDataViewRegistry DataViewRegistry;
 
 #if WITH_EDITOR
+	FPCGChangeTrackingRegistry ChangeTrackingRegistry;
 	FPCGObjectHashFactory ObjectHashFactory;
 #endif
 
diff --git a/Engine/Plugins/PCG/Source/PCG/Public/Subsystems/PCGSubsystem.h b/Engine/Plugins/PCG/Source/PCG/Public/Subsystems/PCGSubsystem.h
--- a/Engine/Plugins/PCG/Source/PCG/Public/Subsystems/PCGSubsystem.h
+++ b/Engine/Plugins/PCG/Source/PCG/Public/Subsystems/PCGSubsystem.h
@@ -305,7 +305,10 @@ public:
 	PCG_API void DeleteSerializedPartitionActors(bool bOnlyDeleteUnused, bool bOnlyChildren = false);
 
 	/** Update the tracking on a given component. */
-	PCG_API void UpdateComponentTracking(IPCGGraphExecutionSource* InExecutionSource, bool bShouldDirtyActors, const TArray<FPCGSelectionKey>* OptionalChangedKeys = nullptr);
+	PCG_API void UpdateTracking(IPCGGraphExecutionSource* InExecutionSource, const TArray<FPCGSelectionKey>* OptionalChangedKeys = nullptr);
+
+	UE_DEPRECATED(5.8, "Use UpdateTracking instead")
+	void UpdateComponentTracking(IPCGGraphExecutionSource* InExecutionSource, bool bShouldDirtyActors, const TArray<FPCGSelectionKey>* OptionalChangedKeys = nullptr) { UpdateTracking(InExecutionSource, OptionalChangedKeys); }
 
 	/** Propagates transient state change from an original component to the relevant partition actors */
 	PCG_API void PropagateEditingModeToLocalComponents(UPCGComponent* InOriginalComponent, EPCGEditorDirtyMode EditingMode);
@@ -336,7 +339,8 @@ public:
 	PCG_API void NotifySelectionKeyChanged(const FPCGSelectionKey& InSelectionKey, const UObject* InOriginatingChangeObject, TArrayView<const FBox> InBounds);
 
 	/** Notify that we exited the Landscape edit mode. */
-	PCG_API void NotifyLandscapeEditModeExited();
+	UE_DEPRECATED(5.8, "No longer used")
+	void NotifyLandscapeEditModeExited() {}
 
 	FPCGOnPCGComponentUnregistered OnPCGComponentUnregistered;
 	
diff --git a/Engine/Plugins/PCG/Source/PCGEditor/Internal/ChangeTracking/PCGActorTracker.h b/Engine/Plugins/PCG/Source/PCGEditor/Internal/ChangeTracking/PCGActorTracker.h
new file mode 100644
--- /dev/null
+++ b/Engine/Plugins/PCG/Source/PCGEditor/Internal/ChangeTracking/PCGActorTracker.h
@@ -0,0 +1,72 @@
+// Copyright Epic Games, Inc. All Rights Reserved.
+
+#pragma once
+
+#if WITH_EDITOR
+
+#include "ChangeTracking/PCGChangeTrackingRegistry.h"
+
+#include "Misc/Guid.h"
+#include "UObject/ObjectKey.h"
+
+class FPCGTrackingManager;
+
+class AActor;
+class ILevelInstanceInterface;
+class UObject;
+
+namespace PCGActorTracker
+{
+	static const FPCGTrackingChangeID ActorDeleted(0x99FF9E41, 0x3940482C, 0xBD5F0BA0, 0x7DAE68EE);
+	static const FPCGTrackingChangeID ActorAdded(0x10A0FD3D, 0x9A8546B8, 0xA7D5FB5F, 0xCD398DC1);
+	static const FPCGTrackingChangeID ActorChanged(0xD822BD05, 0xE0B94C35, 0xA8A36F1D, 0x7A05F444);
+}
+
+class FPCGActorTracker : public IPCGChangeTracker
+{
+public:
+	virtual ~FPCGActorTracker();
+
+	static TUniquePtr<IPCGChangeTracker> MakeInstance(FPCGTrackingManager* InOwner);
+	static FName GetName();
+
+	virtual void Tick() override;
+
+	virtual bool OnObjectPropertyChanged(const UObject* InObject, FPropertyChangedEvent& InEvent) override;
+	virtual bool OnNotifyObjectChanged(const UObject* InObject, const FPCGTrackingChangeID& InChangeType, const UObject* InOriginatingChangeObject) override;
+
+private:
+	// Typedef to store the previous position and previous tags of changed actors.
+	using FActorPreviousData = TTuple<FBox, TSet<FName>, double>;
+
+	FPCGActorTracker(FPCGTrackingManager* InOwner);
+
+	bool IsActorTracked(const AActor* InActor) const;
+	bool ShouldDelayActor(const AActor* InActor) const;
+
+	void OnLevelInstancesUpdated(const TArray<ILevelInstanceInterface*>& InLevelInstances);
+	void OnLevelInstanceEditCancelled(ILevelInstanceInterface* InLevelInstance, bool bInHasDiscardedChanges);
+	void OnActorLoaded(AActor& InActor);
+	void OnActorAdded(AActor* InActor);
+	void OnActorAdded_Internal(const AActor* InActor, bool bShouldDirty);
+	void OnActorChanged_Internal(const AActor* InActor, const UObject* InOriginatingChangeObject, const FPCGTrackingChangeID& InChangeType);
+	void OnActorChanged_Recursive(const AActor* InActor, const UObject* InOriginatingChangeObject, const FPCGTrackingChangeID& InChangeType);
+	void OnActorChanged_Impl(const AActor* InActor, const UObject* InOriginatingChangeObject, const FPCGTrackingChangeID& InChangeType, FActorPreviousData* InPreviousData);
+	void OnActorUnloaded(AActor& InActor);
+	void OnActorDeleted(AActor* InActor);
+	void OnActorDeleted_Internal(const AActor* InActor, bool bShouldDirty);
+	void OnObjectModified(UObject* InObject);
+	
+	static const FLazyName Name;
+
+	// Keep track of actors that aren't yet ready (or if the subsystem is not yet ready), whether we should dirty them in next tick.
+	TMap<TObjectKey<AActor>, TTuple<TObjectKey<UObject>, FPCGTrackingChangeID>> DelayedChangedActors;
+
+	/** Transient map of owners and their previous data, it's set in the pre object change to be able to track changes (such as tags or positions) */
+	TMap<TObjectKey<AActor>, FActorPreviousData> ActorToPreviousDataMap;
+
+	// Time keeper for cleaning up cached previous actor data
+	double LastPreviousActorDataCleanup = -1.0;
+};
+
+#endif // WITH_EDITOR
\ No newline at end of file
diff --git a/Engine/Plugins/PCG/Source/PCGEditor/Internal/ChangeTracking/PCGComponentChangeHandler.h b/Engine/Plugins/PCG/Source/PCGEditor/Internal/ChangeTracking/PCGComponentChangeHandler.h
new file mode 100644
--- /dev/null
+++ b/Engine/Plugins/PCG/Source/PCGEditor/Internal/ChangeTracking/PCGComponentChangeHandler.h
@@ -0,0 +1,47 @@
+// Copyright Epic Games, Inc. All Rights Reserved.
+
+#pragma once
+
+#if WITH_EDITOR
+
+#include "PCGComponent.h"
+#include "PCGGraphExecutionStateInterface.h"
+#include "ChangeTracking/PCGChangeTrackingRegistry.h"
+
+class FPCGTrackingManager;
+
+class FPCGComponentChangeHandler : public IPCGChangeHandler
+{
+public:
+	virtual ~FPCGComponentChangeHandler();
+
+	static TUniquePtr<IPCGChangeHandler> MakeInstance(FPCGTrackingManager* InOwner);
+	static FName GetName();
+
+	virtual void BeginChangeHandling(const TSharedRef<FPCGChangeHandlerChange>& InChange) override;
+	virtual void HandleChange(IPCGGraphExecutionSource* InExecutionSource) override;
+	virtual void HandleBoundedChange(IPCGGraphExecutionSource* InExecutionSource, const FBox& InExecutionSourceBounds, const FBox& InChangeBounds) override;
+	virtual void EndChangeHandling(bool bSkipRefresh) override;
+
+private:
+	FPCGComponentChangeHandler(FPCGTrackingManager* InOwner);
+
+	bool ShouldDiscardComponent(UPCGComponent* InComponent) const;
+
+	struct FCurrentChange
+	{
+		TWeakPtr<FPCGChangeHandlerChange> Change;
+
+		TArray<const UObject*> ChangedObjects;
+
+		TOptional<FPCGDynamicTrackingPriority> ChangedObjectDynamicTrackingPriority;
+		EPCGComponentDirtyFlag DirtyFlag = EPCGComponentDirtyFlag::None;
+		TSet<UPCGComponent*> DirtyComponents;
+	};
+
+	TUniquePtr<FCurrentChange> CurrentChange;
+	
+	static const FLazyName Name;
+};
+
+#endif // WITH_EDITOR
\ No newline at end of file
diff --git a/Engine/Plugins/PCG/Source/PCGEditor/Internal/ChangeTracking/PCGLandscapeTracker.h b/Engine/Plugins/PCG/Source/PCGEditor/Internal/ChangeTracking/PCGLandscapeTracker.h
new file mode 100644
--- /dev/null
+++ b/Engine/Plugins/PCG/Source/PCGEditor/Internal/ChangeTracking/PCGLandscapeTracker.h
@@ -0,0 +1,50 @@
+// Copyright Epic Games, Inc. All Rights Reserved.
+
+#pragma once
+
+#if WITH_EDITOR
+
+#include "ChangeTracking/PCGChangeTrackingRegistry.h"
+
+#include "UObject/ObjectKey.h"
+
+class FPCGTrackingManager;
+
+class ALandscapeProxy;
+class FLandscapeProxyComponentDataChangedParams;
+
+typedef FName FEditorModeID;
+
+class FPCGLandscapeTracker : public IPCGChangeTracker
+{
+public:
+	virtual ~FPCGLandscapeTracker();
+
+	static TUniquePtr<IPCGChangeTracker> MakeInstance(FPCGTrackingManager* InOwner);
+	static FName GetName();
+
+	virtual void Tick() override;
+	virtual bool ShouldSkipRefresh(const UObject* InChangedObject) override;
+
+private:
+	FPCGLandscapeTracker(FPCGTrackingManager* InOwner);
+
+	void OnLandscapeChanged(ALandscapeProxy* InLandscape, const FLandscapeProxyComponentDataChangedParams& InChangeParams);
+	void ApplyLandscapeChanges(ALandscapeProxy* InLandscape);
+	void OnEditorModeIDChanged(const FEditorModeID& EditorModeID, bool bIsEntering);
+	void OnLandscapeEditModeExited();
+
+	static const FLazyName Name;
+
+	// Part for the delayed landscape change update
+	double LastLandscapeDirtyTime = -1.0;
+	TArray<TObjectKey<ALandscapeProxy>> DelayedModifiedLandscapes;
+
+	// Keep track of all dirtied landscapes and force a refresh on them when exiting edit mode.
+	TArray<TObjectKey<ALandscapeProxy>> DirtiedLandscapes;
+	
+	// Temp boolean to indicate we are currently exiting the landscape edit mode
+	bool bIsCurrentlyExitingLandscapeEditMode = false;
+};
+
+#endif // WITH_EDITOR
\ No newline at end of file
diff --git a/Engine/Plugins/PCG/Source/PCGEditor/Internal/ChangeTracking/PCGMiscTracker.h b/Engine/Plugins/PCG/Source/PCGEditor/Internal/ChangeTracking/PCGMiscTracker.h
new file mode 100644
--- /dev/null
+++ b/Engine/Plugins/PCG/Source/PCGEditor/Internal/ChangeTracking/PCGMiscTracker.h
@@ -0,0 +1,31 @@
+// Copyright Epic Games, Inc. All Rights Reserved.
+
+#pragma once
+
+#if WITH_EDITOR
+
+#include "ChangeTracking/PCGChangeTrackingRegistry.h"
+
+#include "UObject/ObjectSaveContext.h"
+
+class FPCGTrackingManager;
+class UObject;
+
+// For now this is a "everything" else tracker that tracks specific asset type saves
+class FPCGMiscTracker : public IPCGChangeTracker
+{
+public:
+	virtual ~FPCGMiscTracker();
+
+	static TUniquePtr<IPCGChangeTracker> MakeInstance(FPCGTrackingManager* InOwner);
+	static FName GetName();
+
+private:
+	FPCGMiscTracker(FPCGTrackingManager* InOwner);
+
+	void OnObjectSaved(UObject* InObject, FObjectPreSaveContext InObjectSaveContext);
+
+	static const FLazyName Name;
+};
+
+#endif // WITH_EDITOR
\ No newline at end of file
diff --git a/Engine/Plugins/PCG/Source/PCGEditor/Private/ChangeTracking/PCGActorTracker.cpp b/Engine/Plugins/PCG/Source/PCGEditor/Private/ChangeTracking/PCGActorTracker.cpp
new file mode 100644
--- /dev/null
+++ b/Engine/Plugins/PCG/Source/PCGEditor/Private/ChangeTracking/PCGActorTracker.cpp
@@ -0,0 +1,549 @@
+// Copyright Epic Games, Inc. All Rights Reserved.
+
+#if WITH_EDITOR
+
+#include "ChangeTracking/PCGActorTracker.h"
+
+#include "PCGTrackingManager.h"
+#include "PCGWorldActor.h"
+#include "Helpers/PCGHelpers.h"
+
+#include "Engine/Engine.h"
+#include "Engine/Level.h"
+#include "GameFramework/Actor.h"
+#include "LevelInstance/LevelInstanceEditorInstanceActor.h"
+#include "LevelInstance/LevelInstanceEditorPivotInterface.h"
+#include "LevelInstance/LevelInstanceInterface.h"
+#include "LevelInstance/LevelInstanceSubsystem.h"
+#include "PackedLevelActor/PackedLevelActor.h"
+#include "UObject/UObjectGlobals.h"
+#include "WorldPartition/WorldPartition.h"
+
+namespace PCGActorTracker
+{
+	static TAutoConsoleVariable<bool> CVarDisableDelayedActorRegistering(
+		TEXT("pcg.DisableDelayedActorRegistering"),
+		false,
+		TEXT("If delayed actor registering when their components aren't registered yet is introducing bad behavior, disables it, allowing people to continue working while we investigate."));
+
+	static TAutoConsoleVariable<int> CVarActorModifiedPreviousDataCleanupDelayMS(
+		TEXT("pcg.ActorModifiedPreviousDataCleanupDelayMS"),
+		10000,
+		TEXT("Time in MS between a cached previous data from modified actor and the time we remove it from our cache."));
+
+	static TAutoConsoleVariable<int> CVarActorModifiedPreviousDataCheckDelayMS(
+		TEXT("pcg.ActorModifiedPreviousDataCheckDelayMS"),
+		1000,
+		TEXT("Delay in MS between cleanup checks on previous data from modified actors."));
+
+	bool ShouldIgnoreActor(const AActor* InActor)
+	{
+		if (!InActor)
+		{
+			return true;
+		}
+
+		return InActor->bIsEditorPreviewActor || InActor->IsA<ALevelInstanceEditorInstanceActor>() || InActor->IsA<APCGWorldActor>() || InActor->Implements<ULevelInstanceEditorPivotInterface>();
+	}
+
+	void PropagateToLevelInstanceActors(const ILevelInstanceInterface* InLevelInstance, TFunctionRef<bool(const AActor* LevelActor)> InFunc)
+	{
+		if (InLevelInstance)
+		{
+			if (const ULevelInstanceSubsystem* LevelInstanceSubsystem = InLevelInstance->GetLevelInstanceSubsystem())
+			{
+				LevelInstanceSubsystem->ForEachActorInLevelInstance(InLevelInstance, [InFunc](const AActor* LevelActor)
+				{
+					if (!ShouldIgnoreActor(LevelActor))
+					{
+						return InFunc(LevelActor);
+					}
+					return true;
+				});
+			}
+		}
+	}
+
+	void PropagateToLevelInstanceActors(const AActor* InActor, TFunctionRef<bool(const AActor* LevelActor)> InFunc)
+	{
+		PropagateToLevelInstanceActors(Cast<ILevelInstanceInterface>(InActor), InFunc);
+	}
+}
+
+const FLazyName FPCGActorTracker::Name = TEXT("ActorTracker");
+
+FPCGActorTracker::FPCGActorTracker(FPCGTrackingManager* InOwner)
+	: IPCGChangeTracker(InOwner)
+{
+	GEngine->OnLevelActorAdded().AddRaw(this, &FPCGActorTracker::OnActorAdded);
+	GEngine->OnLevelActorDeleted().AddRaw(this, &FPCGActorTracker::OnActorDeleted);
+	FCoreUObjectDelegates::OnObjectModified.AddRaw(this, &FPCGActorTracker::OnObjectModified);
+	
+	UWorld* World = Owner->GetWorld();
+	check(World);
+
+	if (World->PersistentLevel)
+	{
+		World->PersistentLevel->OnLoadedActorAddedToLevelEvent.AddRaw(this, &FPCGActorTracker::OnActorLoaded);
+		World->PersistentLevel->OnLoadedActorRemovedFromLevelEvent.AddRaw(this, &FPCGActorTracker::OnActorUnloaded);
+
+		if (ULevelInstanceSubsystem* LevelInstanceSubsystem = UWorld::GetSubsystem<ULevelInstanceSubsystem>(World))
+		{
+			LevelInstanceSubsystem->OnLevelInstancesUpdated().AddRaw(this, &FPCGActorTracker::OnLevelInstancesUpdated);
+			LevelInstanceSubsystem->OnLevelInstanceEditCancelled().AddRaw(this, &FPCGActorTracker::OnLevelInstanceEditCancelled);
+		}
+	}
+}
+
+FPCGActorTracker::~FPCGActorTracker()
+{
+	GEngine->OnLevelActorAdded().RemoveAll(this);
+	GEngine->OnLevelActorDeleted().RemoveAll(this);
+	FCoreUObjectDelegates::OnObjectModified.RemoveAll(this);
+
+	UWorld* World = Owner->GetWorld();
+	check(World);
+
+	if (UWorldPartition* WorldPartition = World->GetWorldPartition())
+	{
+		WorldPartition->OnActorDescInstanceAddedEvent.RemoveAll(this);
+		WorldPartition->OnActorDescInstanceRemovedEvent.RemoveAll(this);
+	}
+
+	if (World->PersistentLevel)
+	{
+		World->PersistentLevel->OnLoadedActorAddedToLevelEvent.RemoveAll(this);
+		World->PersistentLevel->OnLoadedActorRemovedFromLevelEvent.RemoveAll(this);
+
+		if (ULevelInstanceSubsystem* LevelInstanceSubsystem = UWorld::GetSubsystem<ULevelInstanceSubsystem>(World))
+		{
+			LevelInstanceSubsystem->OnLevelInstancesUpdated().RemoveAll(this);
+			LevelInstanceSubsystem->OnLevelInstanceEditCancelled().RemoveAll(this);
+		}
+	}
+}
+
+TUniquePtr<IPCGChangeTracker> FPCGActorTracker::MakeInstance(FPCGTrackingManager* InOwner)
+{
+	return TUniquePtr<IPCGChangeTracker>(new FPCGActorTracker(InOwner));
+}
+
+FName FPCGActorTracker::GetName()
+{
+	return Name.Resolve();
+}
+
+void FPCGActorTracker::Tick()
+{
+	if (DelayedChangedActors.IsEmpty())
+	{
+		return;
+	}
+
+	auto LocalDelayedChangedActors = MoveTemp(DelayedChangedActors);
+
+	for (const auto& ActorToObject : LocalDelayedChangedActors)
+	{
+		AActor* Actor = ActorToObject.Key.ResolveObjectPtr();
+		UObject* Object = ActorToObject.Value.Get<0>().ResolveObjectPtr();
+		const FPCGTrackingChangeID& Reason = ActorToObject.Value.Get<1>();
+
+		if (!Actor)
+		{
+			continue;
+		}
+
+		OnActorChanged_Internal(Actor, Object, Reason);
+	}
+
+	const double CurrentTime = FApp::GetCurrentTime();
+	// Cleaning up previous data gathered by the OnObjectModified function but not consumed.
+	if (LastPreviousActorDataCleanup < 0.0 || ((CurrentTime - LastPreviousActorDataCleanup) * 1000) > PCGActorTracker::CVarActorModifiedPreviousDataCheckDelayMS.GetValueOnAnyThread())
+	{
+		TRACE_CPUPROFILER_EVENT_SCOPE(FPCGTrackingManager::Tick::PreviousActorDataCleanupCheck);
+
+		LastPreviousActorDataCleanup = CurrentTime;
+		TArray<TObjectKey<AActor>> Keys;
+		ActorToPreviousDataMap.GetKeys(Keys);
+		for (const TObjectKey<AActor>& Key : Keys)
+		{
+			if (((CurrentTime - ActorToPreviousDataMap[Key].Get<2>()) * 1000) > PCGActorTracker::CVarActorModifiedPreviousDataCleanupDelayMS.GetValueOnAnyThread())
+			{
+				ActorToPreviousDataMap.Remove(Key);
+			}
+		}
+	}
+}
+
+bool FPCGActorTracker::IsActorTracked(const AActor* InActor) const
+{
+	TRACE_CPUPROFILER_EVENT_SCOPE(FPCGActorTracker::IsActorTracked);
+
+	const FPCGSelectionKey ActorKey = FPCGSelectionKey::CreateFromObjectPtr(InActor);
+	bool bFoundMatch = false;
+	TSet<FName> EmptySet;
+	
+	Owner->ForEachTrackedKey([&bFoundMatch, &ActorKey, &EmptySet](const FPCGSelectionKey& Key, const TSet<IPCGGraphExecutionSource*>& ExecutionSources) -> bool
+	{
+		if(Key.IsMatching(ActorKey, EmptySet, ExecutionSources, nullptr))
+		{
+			bFoundMatch = true;
+			return false;
+		}
+
+		return true;
+	});
+
+	return bFoundMatch;
+}
+
+bool FPCGActorTracker::ShouldDelayActor(const AActor* InActor) const
+{
+	const bool bDisableDelayedActorRegistering = PCGActorTracker::CVarDisableDelayedActorRegistering.GetValueOnAnyThread();
+	if (!InActor->HasActorRegisteredAllComponents() && !bDisableDelayedActorRegistering)
+	{
+		return true;
+	}
+
+	bool bShouldDelay = false;
+
+	// Check that whole Level Instance hierarchy is loaded
+	if (const ILevelInstanceInterface* LevelInstance = Cast<ILevelInstanceInterface>(InActor))
+	{
+		// If Level Instance isn't loaded yet it should be delayed
+		if (!LevelInstance->IsLoaded() && !InActor->IsA<APackedLevelActor>())
+		{
+			return true;
+		}
+
+		PCGActorTracker::PropagateToLevelInstanceActors(LevelInstance, [this, &bShouldDelay](const AActor* InChildActor)
+		{
+			if (ShouldDelayActor(InChildActor))
+			{
+				bShouldDelay = true;
+				return false;
+			}
+			return true;
+		});
+	}
+
+	return bShouldDelay;
+}
+
+void FPCGActorTracker::OnLevelInstancesUpdated(const TArray<ILevelInstanceInterface*>& InLevelInstances)
+{
+	for (ILevelInstanceInterface* LevelInstance : InLevelInstances)
+	{
+		OnActorAdded(Cast<AActor>(LevelInstance));
+	}
+}
+
+void FPCGActorTracker::OnLevelInstanceEditCancelled(ILevelInstanceInterface* InLevelInstance, bool bInHasDiscardedChanges)
+{
+	if (bInHasDiscardedChanges)
+	{
+		OnActorAdded(Cast<AActor>(InLevelInstance));
+	}
+}
+
+void FPCGActorTracker::OnActorLoaded(AActor& InActor)
+{
+	// We have to make sure to not create a infinite loop
+	if (PCGActorTracker::ShouldIgnoreActor(&InActor) || InActor.GetWorld() != Owner->GetWorld())
+	{
+		return;
+	}
+
+	// Loaded actors should not dirty.
+	OnActorAdded_Internal(&InActor, /*bShouldDirty=*/ false);
+}
+
+void FPCGActorTracker::OnActorAdded(AActor* InActor)
+{
+	// We have to make sure to not create a infinite loop
+	if (PCGActorTracker::ShouldIgnoreActor(InActor) || InActor->GetWorld() != Owner->GetWorld())
+	{
+		return;
+	}
+
+	// Implementation note: We delay adding because OnActorAdded fires before an actor's properties are set,
+	// so the actor is not ready for processing until the next tick.
+	// Only dirty if we aren't currently inside a loading path
+	OnActorAdded_Internal(InActor, /*bShouldDirty=*/ !UE::GetIsEditorLoadingPackage());
+}
+
+void FPCGActorTracker::OnActorAdded_Internal(const AActor* InActor, bool bShouldDirty)
+{
+	TRACE_CPUPROFILER_EVENT_SCOPE(FPCGActorTracker::OnActorAdded);
+	check(!PCGActorTracker::ShouldIgnoreActor(InActor) && InActor->GetWorld() == Owner->GetWorld());
+
+
+	if (!bShouldDirty || DelayedChangedActors.Contains(InActor))
+	{
+		// Nothing to do
+		return;
+	}
+
+	// We delay adding because OnActorAdded fires before an actor's properties are set,
+	// so the actor is not ready for processing until the next tick.
+	DelayedChangedActors.Emplace(InActor, { nullptr, PCGActorTracker::ActorAdded });
+}
+
+void FPCGActorTracker::OnActorChanged_Internal(const AActor* InActor, const UObject* InOriginatingChangeObject, const FPCGTrackingChangeID& InChangeType)
+{
+	// Currently delayed, return
+	if (DelayedChangedActors.Contains(InActor))
+	{
+		return;
+	}
+
+	// Should still be delayed, add to delayed actors, return
+	if (ShouldDelayActor(InActor))
+	{
+		DelayedChangedActors.Emplace(InActor, { InOriginatingChangeObject, InChangeType });
+		return;
+	}
+
+	// Call Recursive method
+	OnActorChanged_Recursive(InActor, InOriginatingChangeObject, InChangeType);
+}
+
+void FPCGActorTracker::OnActorChanged_Recursive(const AActor* InActor, const UObject* InOriginatingChangeObject, const FPCGTrackingChangeID& InChangeType)
+{
+	check(!DelayedChangedActors.Contains(InActor));
+
+	// Process actors recursively
+	auto OnActorChanged = [this, InOriginatingChangeObject, InChangeType](const AActor* InActor, auto&& OnActorChangedRecursive) -> void
+	{
+		if (InActor)
+		{
+			FActorPreviousData* PreviousData = ActorToPreviousDataMap.Find(InActor);
+			
+			OnActorChanged_Impl(InActor, InOriginatingChangeObject, InChangeType, PreviousData);
+
+			if (PreviousData)
+			{
+				ActorToPreviousDataMap.Remove(InActor);
+			}
+		}
+
+		PCGActorTracker::PropagateToLevelInstanceActors(InActor, [this, OnActorChangedRecursive](const AActor* LevelActor)
+		{
+			OnActorChangedRecursive(LevelActor, OnActorChangedRecursive);
+			return true;
+		});
+	};
+
+	OnActorChanged(InActor, OnActorChanged);
+}
+
+void FPCGActorTracker::OnActorChanged_Impl(const AActor* InActor, const UObject* InOriginatingChangeObject, const FPCGTrackingChangeID& InChangeType, FActorPreviousData* InPreviousData)
+{
+	check(InActor);
+	
+	if (!Owner->IsTracking())
+	{
+		return;
+	}
+
+	// Discard everything related to self for actor added/deleted
+	const bool bNoRefreshOwner = InChangeType == PCGActorTracker::ActorAdded || InChangeType == PCGActorTracker::ActorDeleted;
+
+	ensure(InActor->GetWorld() == Owner->GetWorld());
+
+	TSet<FName> RemovedTags;
+	if (InPreviousData && !InPreviousData->Get<1>().IsEmpty())
+	{
+		RemovedTags = InPreviousData->Get<1>().Difference(TSet<FName>(InActor->Tags));
+	}
+
+	FPCGSelectionKey SelectionKey = FPCGSelectionKey::CreateFromObjectPtr(InActor);
+
+	TArray<FBox, TInlineAllocator<2>> Bounds;
+	const FBox ActorBounds = PCGHelpers::GetActorBounds(InActor, /*bIgnorePCGCreatedComponents=*/false);
+	Bounds.Add(ActorBounds);
+	
+	// Also provide previous bounds if we do have previous data
+	if (InPreviousData)
+	{
+		const FBox& PreviousBounds = InPreviousData->Get<0>();
+		if (PreviousBounds.IsValid && !PreviousBounds.Equals(ActorBounds))
+		{
+			Bounds.Add(PreviousBounds);
+		}
+	}
+
+	Owner->OnSelectionKeyChanged(SelectionKey, InOriginatingChangeObject, RemovedTags, bNoRefreshOwner, Bounds);
+}
+
+void FPCGActorTracker::OnActorUnloaded(AActor& InActor)
+{
+	// Don't dirty on unload (to mirror the behavior in load)
+	OnActorDeleted_Internal(&InActor, /*bShouldDirty=*/false);
+}
+
+void FPCGActorTracker::OnActorDeleted(AActor* InActor)
+{
+	if (PCGActorTracker::ShouldIgnoreActor(InActor) || InActor->GetWorld() != Owner->GetWorld())
+	{
+		return;
+	}
+
+	// Only dirty if we aren't currently inside a loading path
+	OnActorDeleted_Internal(InActor, /*bShouldDirty=*/ !UE::GetIsEditorLoadingPackage());
+}
+
+void FPCGActorTracker::OnActorDeleted_Internal(const AActor* InActor, bool bShouldDirty)
+{
+	TRACE_CPUPROFILER_EVENT_SCOPE(FPCGActorTracker::OnActorDeleted_Internal);
+	check(InActor && InActor->GetWorld() == Owner->GetWorld());
+
+	if (!bShouldDirty)
+	{
+		// Nothing to do
+		return;
+	}
+
+#if WITH_EDITOR
+	PCGActorTracker::PropagateToLevelInstanceActors(InActor, [this, bShouldDirty](const AActor* LevelActor)
+	{
+		OnActorDeleted_Internal(LevelActor, bShouldDirty);
+		return true;
+	});
+#endif // WITH_EDITOR
+
+	// Notify all components that the actor has changed (was removed), but the Refresh will only happen AFTER the actor was actually removed from the world (because of delayed refresh).
+	OnActorChanged_Impl(InActor, nullptr, PCGActorTracker::ActorDeleted, /*InPreviousData=*/nullptr);
+}
+
+void FPCGActorTracker::OnObjectModified(UObject* InObject)
+{
+	TRACE_CPUPROFILER_EVENT_SCOPE(FPCGActorTracker::OnObjectModified);
+
+	if (!Owner->IsTracking())
+	{
+		return;
+	}
+
+	AActor* Actor = Cast<AActor>(InObject);
+
+	// Otherwise, if it's an actor component, track it as well
+	if (!Actor)
+	{
+		if (UActorComponent* ActorComponent = Cast<UActorComponent>(InObject))
+		{
+			Actor = ActorComponent->GetOwner();
+		}
+	}
+
+	if (!Actor || (Actor->GetWorld() != Owner->GetWorld()))
+	{
+		return;
+	}
+
+	if (PCGActorTracker::ShouldIgnoreActor(Actor))
+	{
+		return;
+	}
+
+	auto StorePreviousData = [this](const AActor* InActor, auto&& RecursiveCall) -> void
+	{
+		if (!ActorToPreviousDataMap.Contains(InActor))
+		{
+			if (IsActorTracked(InActor))
+			{
+				FActorPreviousData& PreviousData = ActorToPreviousDataMap.Add(InActor);
+				PreviousData.Get<0>() = PCGHelpers::GetActorBounds(InActor, /*bIgnorePCGCreatedComponents=*/false);
+				PreviousData.Get<1>() = TSet<FName>(InActor->Tags);
+				PreviousData.Get<2>() = FApp::GetCurrentTime();
+			}
+
+			// Also propagate the pre-change to all child actors if it is within a level instance.
+			PCGActorTracker::PropagateToLevelInstanceActors(InActor, [this, RecursiveCall](const AActor* LevelActor)
+			{
+				RecursiveCall(LevelActor, RecursiveCall);
+				return true;
+			});
+		}
+	};
+
+	StorePreviousData(Actor, StorePreviousData);
+}
+
+bool FPCGActorTracker::OnNotifyObjectChanged(const UObject* InObject, const FPCGTrackingChangeID& InChangeType, const UObject* InOriginatingChangeObject)
+{
+	if (InChangeType != PCGActorTracker::ActorChanged)
+	{
+		return false;
+	}
+
+	const AActor* Actor = Cast<AActor>(InObject);
+	if (PCGActorTracker::ShouldIgnoreActor(Actor) || Actor->GetWorld() != Owner->GetWorld())
+	{
+		return true;
+	}
+
+	OnActorChanged_Internal(Actor, InOriginatingChangeObject, InChangeType);
+	return true;
+}
+
+bool FPCGActorTracker::OnObjectPropertyChanged(const UObject* InObject, FPropertyChangedEvent& InEvent)
+{
+	TRACE_CPUPROFILER_EVENT_SCOPE(FPCGActorTracker::OnObjectPropertyChanged);
+	check(Owner->IsTracking());
+	
+	// First check if it is an actor
+	const AActor* Actor = Cast<AActor>(InObject);
+
+	// Otherwise, if it's an actor component, track it as well
+	if (!Actor)
+	{
+		if (const UActorComponent* ActorComponent = Cast<UActorComponent>(InObject))
+		{
+			Actor = ActorComponent->GetOwner();
+		}
+	}
+
+	// Not an actor, early out
+	if (!Actor)
+	{
+		return false;
+	}
+
+	if ((Actor->GetWorld() != Owner->GetWorld()) || PCGActorTracker::ShouldIgnoreActor(Actor))
+	{
+		return true;
+	}
+
+	const bool bInteractiveChange = (InEvent.ChangeType == EPropertyChangeType::Interactive);
+	// Special exception for actor tags, as we can't track otherwise an actor "losing" a tag
+	const bool bActorTagChange = (InEvent.Property && InEvent.Property->GetFName() == GET_MEMBER_NAME_CHECKED(AActor, Tags));
+
+	const bool bNoOperation = (bInteractiveChange && !bActorTagChange) || (Actor && (DelayedChangedActors.Contains(Actor)));
+
+	if (bNoOperation)
+	{
+		// Don't remove it as it will be needed when the delay will be called.
+		if (Actor && !(DelayedChangedActors.Contains(Actor)))
+		{
+			auto RemoveAll = [this](const AActor* InActor, auto&& RecursiveCall) -> void
+			{
+				ActorToPreviousDataMap.Remove(InActor);
+				PCGActorTracker::PropagateToLevelInstanceActors(InActor, [this, RecursiveCall](const AActor* LevelActor)
+				{
+					RecursiveCall(LevelActor, RecursiveCall);
+					return true;
+				});
+			};
+
+			RemoveAll(Actor, RemoveAll);
+		}
+
+		return true;
+	}
+
+	OnActorChanged_Internal(Actor, /*InOriginatingChangeObject=*/ InObject, PCGActorTracker::ActorChanged);
+	return true;
+}
+
+#endif // WITH_EDITOR
\ No newline at end of file
diff --git a/Engine/Plugins/PCG/Source/PCGEditor/Private/ChangeTracking/PCGComponentChangeHandler.cpp b/Engine/Plugins/PCG/Source/PCGEditor/Private/ChangeTracking/PCGComponentChangeHandler.cpp
new file mode 100644
--- /dev/null
+++ b/Engine/Plugins/PCG/Source/PCGEditor/Private/ChangeTracking/PCGComponentChangeHandler.cpp
@@ -0,0 +1,339 @@
+// Copyright Epic Games, Inc. All Rights Reserved.
+
+#if WITH_EDITOR
+
+#include "ChangeTracking/PCGComponentChangeHandler.h"
+
+#include "PCGTrackingManager.h"
+#include "Grid/PCGPartitionActor.h"
+#include "Utils/PCGGraphExecutionLogging.h"
+
+#include "LandscapeProxy.h"
+#include "Engine/World.h"
+#include "GameFramework/Actor.h"
+#include "LevelInstance/LevelInstanceInterface.h"
+#include "LevelInstance/LevelInstanceSubsystem.h"
+
+const FLazyName FPCGComponentChangeHandler::Name = TEXT("PCGComponentChangeHandler");
+
+namespace PCGComponentChangeHandler
+{
+	bool ShouldDirtyGenerated(UPCGComponent* InPCGComponent, const UObject* InChangedObject)
+	{
+		// @todo_pcg: to remove usage
+PRAGMA_DISABLE_DEPRECATION_WARNINGS
+		check(InPCGComponent);
+		return InPCGComponent->ShouldTrackLandscape() && InChangedObject && InChangedObject->IsA<ALandscapeProxy>();
+PRAGMA_ENABLE_DEPRECATION_WARNINGS
+	}
+}
+
+FPCGComponentChangeHandler::FPCGComponentChangeHandler(FPCGTrackingManager* InOwner)
+	: IPCGChangeHandler(InOwner)
+{
+}
+
+FPCGComponentChangeHandler::~FPCGComponentChangeHandler()
+{
+}
+
+TUniquePtr<IPCGChangeHandler> FPCGComponentChangeHandler::MakeInstance(FPCGTrackingManager* InOwner)
+{
+	return TUniquePtr<IPCGChangeHandler>(new FPCGComponentChangeHandler(InOwner));
+}
+
+FName FPCGComponentChangeHandler::GetName()
+{
+	return Name.Resolve();
+}
+
+void FPCGComponentChangeHandler::BeginChangeHandling(const TSharedRef<FPCGChangeHandlerChange>& InChange)
+{
+	check(InChange->ChangedObject);
+
+	CurrentChange = MakeUnique<FCurrentChange>();
+	CurrentChange->Change = InChange.ToWeakPtr();
+
+	CurrentChange->ChangedObjects.Add(InChange->ChangedObject);
+	if (InChange->OriginatingChangeObject != nullptr && InChange->ChangedObject != InChange->OriginatingChangeObject)
+	{
+		CurrentChange->ChangedObjects.Add(InChange->OriginatingChangeObject);
+	}
+
+	// @todo_pcg: Find a better way to get this info as this doesn't work for managed objects that were spawned under a different outer
+	if (AActor* OuterActor = InChange->OriginatingChangeObject ? InChange->OriginatingChangeObject->GetTypedOuter<AActor>() : nullptr)
+	{
+		TArray<UPCGComponent*> OriginatingComponents;
+		OuterActor->GetComponents<UPCGComponent>(OriginatingComponents);
+		for (UPCGComponent* OriginatingComponent : OriginatingComponents)
+		{
+			if (!OriginatingComponent->IsCleaningUp() && OriginatingComponent->IsAnyObjectManagedByResource(CurrentChange->ChangedObjects))
+			{
+				CurrentChange->ChangedObjectDynamicTrackingPriority = OriginatingComponent->GetExecutionState().GetDynamicTrackingPriority();
+				break;
+			}
+		}
+	}
+
+	CurrentChange->DirtyFlag = EPCGComponentDirtyFlag::Actor;
+	if (InChange->ChangedObject->IsA<ALandscapeProxy>())
+	{
+		CurrentChange->DirtyFlag = CurrentChange->DirtyFlag | EPCGComponentDirtyFlag::Landscape;
+	}
+}
+
+void FPCGComponentChangeHandler::HandleChange(IPCGGraphExecutionSource* InExecutionSource)
+{
+	check(CurrentChange.IsValid());
+
+	TSharedPtr<FPCGChangeHandlerChange> Change = CurrentChange->Change.Pin();
+	if (!Change)
+	{
+		return;
+	}
+
+	UPCGComponent* PCGComponent = Cast<UPCGComponent>(InExecutionSource);
+	if (!PCGComponent)
+	{
+		return;
+	}
+	
+	if (CurrentChange->DirtyComponents.Contains(PCGComponent) || ShouldDiscardComponent(PCGComponent))
+	{
+		return;
+	}
+
+	// Don't mark "Owner changed" if the change originate from a PCG Component. It will be delegated to the ClearCacheForActor.
+	// It is necessary to avoid infine loops when there are multiple PCG components on one actor, and one component was generated.
+	const bool bOwnerChanged = (PCGComponent->GetOwner() == Change->ChangedObject) && (!Change->OriginatingChangeObject || !Change->OriginatingChangeObject->IsA<UPCGComponent>());
+	bool bShouldDirty = bOwnerChanged || PCGComponentChangeHandler::ShouldDirtyGenerated(PCGComponent, Change->ChangedObject);
+
+	// Since when we clear the cache for a settings, we clear it all, it's only necessary to do it once on the original component. And we will dirty all the local components (with dirty dispatch)
+	// Technically, a key that is always tracked can also be culled tracked, so maybe having an intersection test here to not
+	// dirty components that culled tracked and don't intersect with the component would be more efficient.
+	// In practice, it's mostly relevent for partitioned graphs and we will dirty all the locals anyway.
+	bShouldDirty |= Owner->ClearCacheForKeys(Change->MatchedKeys, PCGComponent, /*bIntersect=*/true, Change->OriginatingChangeObject);
+	if (bShouldDirty)
+	{
+		PCGComponent->DirtyGenerated(CurrentChange->DirtyFlag, /*bDispatchToLocalComponents=*/true);
+		CurrentChange->DirtyComponents.Add(PCGComponent);
+	}
+}
+
+void FPCGComponentChangeHandler::HandleBoundedChange(IPCGGraphExecutionSource* InExecutionSource, const FBox& InExecutionSourceBounds, const FBox& InChangeBounds)
+{
+	check(CurrentChange.IsValid());
+
+	TSharedPtr<FPCGChangeHandlerChange> Change = CurrentChange->Change.Pin();
+	if (!Change)
+	{
+		return;
+	}
+
+	UPCGComponent* PCGComponent = Cast<UPCGComponent>(InExecutionSource);
+	if (!PCGComponent)
+	{
+		return;
+	}
+
+	if (ShouldDiscardComponent(PCGComponent))
+	{
+		return;
+	}
+
+	if (PCGComponent->IsPartitioned())
+	{
+		// Don't dirty if the component is not tracked, the origin of the change or its owner is the changed actor and we should not refresh.
+				// We can "re-dirty" it because changes can impact different local components, from the same
+				// original component.
+		const bool bIntersect = InChangeBounds.Intersect(InExecutionSourceBounds);
+		if (!bIntersect)
+		{
+			return;
+		}
+
+		const FBox Overlap = InChangeBounds.Overlap(InExecutionSourceBounds);
+		bool bShouldDirty = PCGComponentChangeHandler::ShouldDirtyGenerated(PCGComponent, Change->ChangedObject);
+		bShouldDirty |= Owner->ClearCacheForKeys(Change->MatchedKeys, PCGComponent, /*bIntersect=*/true, Change->OriginatingChangeObject);
+		bool bWasDirtied = false;
+
+		// Since when we clear the cache for a settings, we clear it all, it's only necessary to do it once on the original component, then only dirty the local that intersects.
+		if (bShouldDirty)
+		{
+			Owner->ForAllIntersectingPartitionActors(Overlap, [PCGComponent, &bWasDirtied, this](APCGPartitionActor* InPartitionActor) -> void
+			{
+				if (UPCGComponent* LocalComponent = InPartitionActor->GetLocalComponent(PCGComponent))
+				{
+					bWasDirtied = true;
+					LocalComponent->DirtyGenerated(CurrentChange->DirtyFlag);
+				}
+			});
+		}
+
+		if (bWasDirtied)
+		{
+			// Don't dispatch
+			PCGComponent->DirtyGenerated(CurrentChange->DirtyFlag, /*bDispatchToLocalComponents=*/false);
+			CurrentChange->DirtyComponents.Add(PCGComponent);
+		}
+	}
+	else
+	{
+		// Don't dirty if the component was already dirtied, not tracked, the origin of the change or its owner is the changed actor and we should not refresh.
+		if (CurrentChange->DirtyComponents.Contains(PCGComponent))
+		{
+			return;
+		}
+
+		bool bShouldDirty = PCGComponentChangeHandler::ShouldDirtyGenerated(PCGComponent, Change->ChangedObject);
+		bShouldDirty |= Owner->ClearCacheForKeys(Change->MatchedKeys, PCGComponent, /*bIntersect=*/true, Change->OriginatingChangeObject);
+		if (bShouldDirty)
+		{
+			PCGComponent->DirtyGenerated(CurrentChange->DirtyFlag);
+			CurrentChange->DirtyComponents.Add(PCGComponent);
+		}
+	}
+}
+
+// We discard the component if:
+	//   * It is null.
+	//   * It is one of the objects that changed.
+	//   * We should not refresh the owner and its owner is one of the objects that changed.
+	//   * The objects that changed are managed by it (safe-guard, would probably induce an infinite loop).
+	//   * The objects that changed are managed by a lower tracking priority component. (lower tracking priority value means higher priority)
+	//   * One of the objects that changed is ignored by its original component.
+bool FPCGComponentChangeHandler::ShouldDiscardComponent(UPCGComponent* InComponent) const
+{
+	check(CurrentChange.IsValid());
+
+	TSharedPtr<FPCGChangeHandlerChange> Change = CurrentChange->Change.Pin();
+	if (!Change)
+	{
+		return true;
+	}
+
+	if (!InComponent
+		|| CurrentChange->ChangedObjects.Contains(InComponent)
+		|| (Change->bInNoRefreshOwner && CurrentChange->ChangedObjects.Contains(InComponent->GetOwner()))
+		|| InComponent->IsCleaningUp()
+		|| InComponent->IsAnyObjectManagedByResource(CurrentChange->ChangedObjects)
+		|| (CurrentChange->ChangedObjectDynamicTrackingPriority.IsSet() && InComponent->GetExecutionState().GetDynamicTrackingPriority() < CurrentChange->ChangedObjectDynamicTrackingPriority.GetValue()))
+	{
+		return true;
+	}
+
+	const UObject* ObjectIgnored = nullptr;
+	const UPCGComponent* OriginalComponent = InComponent->GetConstOriginalComponent();
+	EPCGIgnoreChangeOriginReason Reason;
+
+	if (ensure(OriginalComponent) && OriginalComponent->IsIgnoringAnyChangeOrigins(CurrentChange->ChangedObjects, ObjectIgnored, Reason))
+	{
+		PCGGraphExecutionLogging::LogChangeOriginIgnoredForComponent(ObjectIgnored, InComponent, Reason);
+		return true;
+	}
+
+	return false;
+}
+
+void FPCGComponentChangeHandler::EndChangeHandling(bool bSkipRefresh)
+{
+	check(CurrentChange.IsValid());
+	
+	// Release current change
+	ON_SCOPE_EXIT
+	{
+		CurrentChange.Reset();
+	};
+
+	if (bSkipRefresh)
+	{
+		return;
+	}
+
+	TSharedPtr<FPCGChangeHandlerChange> Change = CurrentChange->Change.Pin();
+	if (!Change)
+	{
+		return;
+	}
+
+	ULevelInstanceSubsystem* LevelInstanceSubsystem = Owner->GetWorld()->GetSubsystem<ULevelInstanceSubsystem>();
+	const IPCGGraphExecutionSource* OriginatingExecutionSource = Cast<const IPCGGraphExecutionSource>(Change->OriginatingChangeObject);
+
+	// And refresh all dirtied components
+	for (UPCGComponent* Component : CurrentChange->DirtyComponents)
+	{
+		if (!ensure(Component))
+		{
+			continue;
+		}
+
+		// This part checks if the change originates from a PCG Component. If so, we check the dirty component has no more other PCG dependencies.
+		// In that case we can proceed for the refresh, otherwise early out, this will be woken up by another dependency change.
+		if (OriginatingExecutionSource && !Owner->RemoveExecutionSourceDependency(OriginatingExecutionSource, Component))
+		{
+			continue;
+		}
+
+		const bool bOwnerHasChanged = Component->GetOwner() == Change->ChangedObject;
+
+		if ((!Change->bInNoRefreshOwner || !bOwnerHasChanged) && (!Component->bOnlyTrackItself || bOwnerHasChanged))
+		{
+			const AActor* Actor = Cast<AActor>(Change->ChangedObject);
+
+			// When an object changes, we need to make sure that we don't trigger a refresh on PCG components that are "higher" in the
+			// level hierarchy, otherwise we will end up generating in the Level Instance level, which is wrong.
+			// Note that in some instances (e.g. when something happens at the Level Instance level) we need to make sure that the
+			// PCG components higher-up are properly updated, hence the level instance depth
+			if (LevelInstanceSubsystem && Actor)
+			{
+				// Immediate Level Instance Parent (Actor)
+				const ILevelInstanceInterface* ActorLevelInstance = LevelInstanceSubsystem->GetParentLevelInstance(Actor);
+
+				// Immediate Level Instance Parent (PCG Component)
+				const ILevelInstanceInterface* ComponentLevelInstance = LevelInstanceSubsystem->GetParentLevelInstance(Component->GetOwner());
+
+				// Is the Actor under a Level Instance that is the same or under the Components owning Level Instance
+				bool bActorLevelInstanceInComponentLevelInstanceHierarchy = false;
+
+				// Go up Level Instance hierarchy for Actor and see if we are part of the Component's level instance hierarchy, stop if we find an Editing Level instance as we don't want to impact top level 
+				// Components from an editing level instance
+				if (ComponentLevelInstance)
+				{
+					LevelInstanceSubsystem->ForEachLevelInstanceAncestors(Actor, [ComponentLevelInstance, &bActorLevelInstanceInComponentLevelInstanceHierarchy](const ILevelInstanceInterface* Ancestor)
+					{
+						if (Ancestor == ComponentLevelInstance)
+						{
+							bActorLevelInstanceInComponentLevelInstanceHierarchy = true;
+							return false; // stop iterating
+						}
+
+						return !Ancestor->IsEditing();
+					});
+				}
+
+				// Actor in same level instance as Component or under the Components Level instance
+				if (bActorLevelInstanceInComponentLevelInstanceHierarchy)
+				{
+					// Allow update if Component is inside Editing level instance or if it is only in preview mode (generates transient data because it is in a non editing Level instance)
+					if (!ComponentLevelInstance->IsEditing() && Component->GetEditingMode() != EPCGEditorDirtyMode::Preview)
+					{
+						continue;
+					}
+				} // Allow update of Component from a Editing Level instance contained actor if the Component being refreshed is in preview mode and not itself in a Level Instance
+				else if (Actor->IsInEditLevelInstanceHierarchy() && (Component->GetEditingMode() != EPCGEditorDirtyMode::Preview || ComponentLevelInstance != nullptr))
+				{
+					continue;
+				}
+				// If ComponentLevelInstance is non-null and Actor Level instance isn't in its hierarchy then skip it
+				else if (ActorLevelInstance && ComponentLevelInstance != nullptr)
+				{
+					continue;
+				}
+			}
+
+			Component->Refresh();
+		}
+	}
+}
+
+#endif // WITH_EDITOR
\ No newline at end of file
diff --git a/Engine/Plugins/PCG/Source/PCGEditor/Private/ChangeTracking/PCGLandscapeTracker.cpp b/Engine/Plugins/PCG/Source/PCGEditor/Private/ChangeTracking/PCGLandscapeTracker.cpp
new file mode 100644
--- /dev/null
+++ b/Engine/Plugins/PCG/Source/PCGEditor/Private/ChangeTracking/PCGLandscapeTracker.cpp
@@ -0,0 +1,170 @@
+// Copyright Epic Games, Inc. All Rights Reserved.
+
+#if WITH_EDITOR
+
+#include "ChangeTracking/PCGLandscapeTracker.h"
+
+#include "PCGTrackingManager.h"
+#include "ChangeTracking/PCGActorTracker.h"
+#include "Landscape/PCGLandscapeCVars.h"
+
+#include "Editor.h"
+#include "EditorModeManager.h"
+#include "EditorModes.h"
+#include "ILevelEditor.h"
+#include "Landscape.h"
+#include "LandscapeSubsystem.h"
+#include "LevelEditor.h"
+#include "Engine/World.h"
+#include "Modules/ModuleManager.h"
+
+const FLazyName FPCGLandscapeTracker::Name = TEXT("LandscapeTracker");
+
+FPCGLandscapeTracker::FPCGLandscapeTracker(FPCGTrackingManager* InOwner)
+	: IPCGChangeTracker(InOwner)
+{
+	if (ULandscapeSubsystem* LandscapeSubsystem = Owner->GetWorld()->GetSubsystem<ULandscapeSubsystem>())
+	{
+		LandscapeSubsystem->OnLandscapeProxyComponentDataChanged().AddRaw(this, &FPCGLandscapeTracker::OnLandscapeChanged);
+	}
+
+	FLevelEditorModule& LevelEditorModule = FModuleManager::GetModuleChecked<FLevelEditorModule>("LevelEditor");
+	if (TSharedPtr<ILevelEditor> FirstLevelEditor = LevelEditorModule.GetFirstLevelEditor())
+	{
+		FirstLevelEditor->GetEditorModeManager().OnEditorModeIDChanged().AddRaw(this, &FPCGLandscapeTracker::OnEditorModeIDChanged);
+	}
+}
+
+FPCGLandscapeTracker::~FPCGLandscapeTracker()
+{
+	if (ULandscapeSubsystem* LandscapeSubsystem = Owner->GetWorld()->GetSubsystem<ULandscapeSubsystem>())
+	{
+		LandscapeSubsystem->OnLandscapeProxyComponentDataChanged().RemoveAll(this);
+	}
+
+	if (FModuleManager::Get().IsModuleLoaded("LevelEditor"))
+	{
+		FLevelEditorModule& LevelEditorModule = FModuleManager::GetModuleChecked<FLevelEditorModule>("LevelEditor");
+		if (TSharedPtr<ILevelEditor> FirstLevelEditor = LevelEditorModule.GetFirstLevelEditor())
+		{
+			FirstLevelEditor->GetEditorModeManager().OnEditorModeIDChanged().RemoveAll(this);
+		}
+	}
+}
+
+TUniquePtr<IPCGChangeTracker> FPCGLandscapeTracker::MakeInstance(FPCGTrackingManager* InOwner)
+{
+	return TUniquePtr<IPCGChangeTracker>(new FPCGLandscapeTracker(InOwner));
+}
+
+FName FPCGLandscapeTracker::GetName()
+{
+	return Name.Resolve();
+}
+
+void FPCGLandscapeTracker::Tick()
+{
+	const double CurrentTime = FApp::GetCurrentTime();
+
+	if (!DelayedModifiedLandscapes.IsEmpty() && LastLandscapeDirtyTime > 0.0 && ((CurrentTime - LastLandscapeDirtyTime) * 1000.0) > PCGLandscapeCVars::CVarLandscapeRefreshTimeDelay.GetValueOnAnyThread())
+	{
+		LastLandscapeDirtyTime = -1.0;
+		for (TObjectKey<ALandscapeProxy> Landscape : DelayedModifiedLandscapes)
+		{
+			ApplyLandscapeChanges(Landscape.ResolveObjectPtr());
+		}
+
+		DelayedModifiedLandscapes.Empty();
+	}
+}
+
+bool FPCGLandscapeTracker::ShouldSkipRefresh(const UObject* InChangedObject)
+{
+	const ALandscapeProxy* LandscapeProxy = Cast<ALandscapeProxy>(InChangedObject);
+
+	if (!LandscapeProxy)
+	{
+		return false;
+	}
+
+	// If refresh is globably disabled, never refresh
+	if (PCGLandscapeCVars::CVarLandscapeDisableRefreshTracking.GetValueOnAnyThread())
+	{
+		return true;
+	}
+
+	// If refresh is not disabled in editing, always refresh
+	if (!PCGLandscapeCVars::CVarLandscapeDisableRefreshTrackingInLandscapeEditingMode.GetValueOnAnyThread())
+	{
+		return false;
+	}
+
+	// Skip refresh if we are currently in landscape edit mode
+	const ALandscape* Landscape = LandscapeProxy->GetLandscapeActor();
+	if (Landscape && Landscape->HasLandscapeEdMode() && !bIsCurrentlyExitingLandscapeEditMode)
+	{
+		DirtiedLandscapes.AddUnique(LandscapeProxy);
+		return true;
+	}
+
+	return false;
+}
+
+void FPCGLandscapeTracker::OnLandscapeChanged(ALandscapeProxy* InLandscape, const FLandscapeProxyComponentDataChangedParams& InChangeParams)
+{
+	if (!InLandscape)
+	{
+		return;
+	}
+
+	if (PCGLandscapeCVars::CVarLandscapeRefreshTimeDelay.GetValueOnAnyThread() > 0)
+	{
+		LastLandscapeDirtyTime = FApp::GetCurrentTime();
+		DelayedModifiedLandscapes.AddUnique(InLandscape);
+	}
+	else
+	{
+		ApplyLandscapeChanges(InLandscape);
+	}
+}
+
+void FPCGLandscapeTracker::ApplyLandscapeChanges(ALandscapeProxy* InLandscape)
+{
+	if (!InLandscape || !Owner->IsTracking())
+	{
+		return;
+	}
+
+	// Forward as an object change event that can be processed by another tracker
+	Owner->NotifyObjectChanged(InLandscape, PCGActorTracker::ActorChanged);
+}
+
+void FPCGLandscapeTracker::OnEditorModeIDChanged(const FEditorModeID& EditorModeID, bool bIsEntering)
+{
+	if (EditorModeID == FBuiltinEditorModes::EM_Landscape && !bIsEntering)
+	{
+		OnLandscapeEditModeExited();
+	}
+}
+
+void FPCGLandscapeTracker::OnLandscapeEditModeExited()
+{
+	TGuardValue<bool> Guard(bIsCurrentlyExitingLandscapeEditMode, true);
+
+	// When the landscape edit mode is exited, force the refresh on all modified/dirtied landscapes.
+	for (TObjectKey<ALandscapeProxy> Landscape : DelayedModifiedLandscapes)
+	{
+		DirtiedLandscapes.AddUnique(Landscape);
+	}
+
+	DelayedModifiedLandscapes.Empty();
+
+	for (TObjectKey<ALandscapeProxy> Landscape : DirtiedLandscapes)
+	{
+		ApplyLandscapeChanges(Landscape.ResolveObjectPtr());
+	}
+
+	DirtiedLandscapes.Empty();
+}
+
+#endif // WITH_EDITOR
\ No newline at end of file
diff --git a/Engine/Plugins/PCG/Source/PCGEditor/Private/ChangeTracking/PCGMiscTracker.cpp b/Engine/Plugins/PCG/Source/PCGEditor/Private/ChangeTracking/PCGMiscTracker.cpp
new file mode 100644
--- /dev/null
+++ b/Engine/Plugins/PCG/Source/PCGEditor/Private/ChangeTracking/PCGMiscTracker.cpp
@@ -0,0 +1,57 @@
+// Copyright Epic Games, Inc. All Rights Reserved.
+
+#if WITH_EDITOR
+
+#include "ChangeTracking/PCGMiscTracker.h"
+
+#include "PCGDataAsset.h"
+#include "PCGTrackingManager.h"
+
+#include "Engine/DataTable.h"
+#include "Engine/StaticMesh.h"
+#include "UObject/UObjectGlobals.h"
+
+const FLazyName FPCGMiscTracker::Name = TEXT("MiscTracker");
+
+FPCGMiscTracker::FPCGMiscTracker(FPCGTrackingManager* InOwner)
+	: IPCGChangeTracker(InOwner)
+{
+	FCoreUObjectDelegates::OnObjectPreSave.AddRaw(this, &FPCGMiscTracker::OnObjectSaved);
+}
+
+FPCGMiscTracker::~FPCGMiscTracker()
+{
+	FCoreUObjectDelegates::OnObjectPreSave.RemoveAll(this);
+}
+
+TUniquePtr<IPCGChangeTracker> FPCGMiscTracker::MakeInstance(FPCGTrackingManager* InOwner)
+{
+	return TUniquePtr<IPCGChangeTracker>(new FPCGMiscTracker(InOwner));
+}
+
+FName FPCGMiscTracker::GetName()
+{
+	return Name.Resolve();
+}
+
+void FPCGMiscTracker::OnObjectSaved(UObject* InObject, FObjectPreSaveContext InObjectSaveContext)
+{
+	check(Owner);
+
+	// Nothing to do if we track nothing
+	if (!InObject || !Owner->IsTracking())
+	{
+		return;
+	}
+
+	// Only trigger a refresh on save for limited data classes.
+	// At this point in time, We only track data tables and PCG data assets because in most cases we probably will catch other changes with OnObjectPropertyChanged.
+	// This is especially important to make sure we don't trigger refresh multiple times (less a problem for PCG assets, but a big problem for data tables).
+	if (!InObjectSaveContext.IsProceduralSave() && (InObject->IsA<UDataTable>() || InObject->IsA<UPCGDataAsset>() || InObject->IsA<UStaticMesh>()))
+	{
+		FPCGSelectionKey SelectionKey = FPCGSelectionKey::CreateFromObjectPtr(InObject);
+		Owner->OnSelectionKeyChanged(SelectionKey);
+	}
+}
+
+#endif // WITH_EDITOR
\ No newline at end of file
diff --git a/Engine/Plugins/PCG/Source/PCGEditor/Private/PCGEditorModule.cpp b/Engine/Plugins/PCG/Source/PCGEditor/Private/PCGEditorModule.cpp
--- a/Engine/Plugins/PCG/Source/PCGEditor/Private/PCGEditorModule.cpp
+++ b/Engine/Plugins/PCG/Source/PCGEditor/Private/PCGEditorModule.cpp
@@ -33,6 +33,10 @@
 #include "PCGEditorSettings.h"
 #include "PCGEditorStyle.h"
 #include "PCGEditorUtils.h"
+#include "ChangeTracking/PCGActorTracker.h"
+#include "ChangeTracking/PCGComponentChangeHandler.h"
+#include "ChangeTracking/PCGLandscapeTracker.h"
+#include "ChangeTracking/PCGMiscTracker.h"
 #include "DataVisualizations/PCGBaseTextureDataVisualization.h"
 #include "DataVisualizations/PCGLandscapeDataVisualization.h"
 #include "DataVisualizations/PCGParamDataVisualization.h"
@@ -103,6 +107,7 @@ void FPCGEditorModule::StartupModule()
 	RegisterSettings();
 	RegisterPCGDataVisualizations();
 	RegisterPinColorAndIcons();
+	RegisterTrackerFactories();
 
 	UToolMenus::RegisterStartupCallback(FSimpleMulticastDelegate::FDelegate::CreateRaw(this, &FPCGEditorModule::RegisterMenuExtensions));
 
@@ -146,6 +151,7 @@ void FPCGEditorModule::ShutdownModule()
 
 	IPCGEditorModule::SetEditorModule(nullptr);
 
+	UnregisterTrackerFactories();
 	UnregisterPCGDataVisualizations();
 	UnregisterSettings();
 	UnregisterDetailsCustomizations();
@@ -176,14 +182,6 @@ void FPCGEditorModule::ShutdownModule()
 	{
 		FLevelEditorModule& LevelEditorModule = FModuleManager::GetModuleChecked<FLevelEditorModule>("LevelEditor");
 		LevelEditorModule.OnLevelEditorCreated().RemoveAll(this);
-
-		if (!IsEngineExitRequested())
-		{
-			if (TSharedPtr<ILevelEditor> FirstLevelEditor = LevelEditorModule.GetFirstLevelEditor())
-			{
-				FirstLevelEditor->GetEditorModeManager().OnEditorModeIDChanged().RemoveAll(this);
-			}
-		}
 	}
 
 	FEditorBuildUtils::UnregisterCustomBuildType(PCGEditorModule::PCGBuildType);
@@ -191,8 +189,6 @@ void FPCGEditorModule::ShutdownModule()
 
 void FPCGEditorModule::OnLevelEditorCreated(TSharedPtr<ILevelEditor> InLevelEditor)
 {
-	RegisterOnEditorModeChange();
-
 	if (GEditor)
 	{
 		// Factory should be auto-discovered by UEditorEngine::InitEditor
@@ -221,17 +217,6 @@ bool FPCGEditorModule::ShouldDisableCPUThrottling()
 	return false;
 }
 
-void FPCGEditorModule::OnEditorModeIDChanged(const FEditorModeID& EditorModeID, bool bIsEntering)
-{
-	if (EditorModeID == FBuiltinEditorModes::EM_Landscape && !bIsEntering)
-	{
-		if (UPCGSubsystem* PCGSubsystem = UPCGSubsystem::GetSubsystemForCurrentWorld())
-		{
-			PCGSubsystem->NotifyLandscapeEditModeExited();
-		}
-	}
-}
-
 void FPCGEditorModule::OnPreExit()
 {
 	UnregisterPinColorAndIcons();
@@ -315,16 +300,6 @@ void FPCGEditorModule::SetPermissionMode(EPCGEditorPermissionMode InPermissionMo
 	PermissionModeChangedEvent.Broadcast(PermissionMode);
 }
 
-void FPCGEditorModule::RegisterOnEditorModeChange()
-{
-	// Have a callback that catches changes in the Editor modes, to catch when we exit the landscape edit mode.
-	FLevelEditorModule& LevelEditorModule = FModuleManager::GetModuleChecked<FLevelEditorModule>("LevelEditor");
-	if (TSharedPtr<ILevelEditor> FirstLevelEditor = LevelEditorModule.GetFirstLevelEditor())
-	{
-		FirstLevelEditor->GetEditorModeManager().OnEditorModeIDChanged().AddRaw(this, &FPCGEditorModule::OnEditorModeIDChanged);
-	}
-}
-
 void FPCGEditorModule::RegisterDetailsCustomizations()
 {
 	FPropertyEditorModule& PropertyEditor = FModuleManager::LoadModuleChecked<FPropertyEditorModule>("PropertyEditor");
@@ -875,6 +850,26 @@ void FPCGEditorModule::UnregisterPCGDataVisualizations()
 	FPCGModule::GetMutablePCGDataVisualizationRegistry().InternalRegistry.Empty();
 }
 
+void FPCGEditorModule::RegisterTrackerFactories()
+{
+	FPCGChangeTrackingRegistry& TrackingRegistry = FPCGModule::GetMutableChangeTrackingRegistry();
+	TrackingRegistry.RegisterTrackerFactory(FPCGActorTracker::GetName(), FPCGCreateTrackerInstance::CreateStatic(&FPCGActorTracker::MakeInstance));
+	TrackingRegistry.RegisterTrackerFactory(FPCGLandscapeTracker::GetName(), FPCGCreateTrackerInstance::CreateStatic(&FPCGLandscapeTracker::MakeInstance));
+	TrackingRegistry.RegisterTrackerFactory(FPCGMiscTracker::GetName(), FPCGCreateTrackerInstance::CreateStatic(&FPCGMiscTracker::MakeInstance));
+
+	TrackingRegistry.RegisterHandlerFactory(FPCGComponentChangeHandler::GetName(), FPCGCreateHandlerInstance::CreateStatic(&FPCGComponentChangeHandler::MakeInstance));
+}
+
+void FPCGEditorModule::UnregisterTrackerFactories()
+{
+	FPCGChangeTrackingRegistry& TrackingRegistry = FPCGModule::GetMutableChangeTrackingRegistry();
+	TrackingRegistry.UnregisterTrackerFactory(FPCGActorTracker::GetName());
+	TrackingRegistry.UnregisterTrackerFactory(FPCGLandscapeTracker::GetName());
+	TrackingRegistry.UnregisterTrackerFactory(FPCGMiscTracker::GetName());
+
+	TrackingRegistry.UnregisterHandlerFactory(FPCGComponentChangeHandler::GetName());
+}
+
 void FPCGEditorModule::RegisterPinColorAndIcons()
 {
 	FPCGDataTypeRegistry& InRegistry = FPCGModule::GetMutableDataTypeRegistry();
diff --git a/Engine/Plugins/PCG/Source/PCGEditor/Public/PCGEditorModule.h b/Engine/Plugins/PCG/Source/PCGEditor/Public/PCGEditorModule.h
--- a/Engine/Plugins/PCG/Source/PCGEditor/Public/PCGEditorModule.h
+++ b/Engine/Plugins/PCG/Source/PCGEditor/Public/PCGEditorModule.h
@@ -65,10 +65,11 @@ protected:
 	void UnregisterPCGDataVisualizations();
 	void RegisterPinColorAndIcons();
 	void UnregisterPinColorAndIcons();
+	void RegisterTrackerFactories();
+	void UnregisterTrackerFactories();
 
 	void OnLevelEditorCreated(TSharedPtr<ILevelEditor> InLevelEditor);
 	void RegisterOnEditorModeChange();
-	void OnEditorModeIDChanged(const FEditorModeID& EditorModeID, bool bIsEntering);
 
 	void OnPreExit();
 
