[PCG] Fix node wires showing "no data" for GPU node outputs.

#rnx
#Synced-CL 48603294

diff --git a/Engine/Plugins/PCG/Source/PCGEditor/Private/Schema/PCGEditorGraphSchema.cpp b/Engine/Plugins/PCG/Source/PCGEditor/Private/Schema/PCGEditorGraphSchema.cpp
--- a/Engine/Plugins/PCG/Source/PCGEditor/Private/Schema/PCGEditorGraphSchema.cpp
+++ b/Engine/Plugins/PCG/Source/PCGEditor/Private/Schema/PCGEditorGraphSchema.cpp
@@ -1105,6 +1105,12 @@ bool FPCGEditorConnectionDrawingPolicy::UpdateParamsIfDebugging(UEdGraphPin* Out
 		return false;
 	}
 
+	// Use defaults if data is emitted by GPU node as we don't have entries in the inspection cache to drive this info.
+	if (const UPCGSettings* NodeSettings = NodeToInspect->GetSettings(); NodeSettings && NodeSettings->bEnabled && NodeSettings->ShouldExecuteOnGPU())
+	{
+		return false;
+	}
+
 	const FPCGEditor* Editor = Graph->GetEditor().Pin().Get();
 	const FPCGStack* PCGStack = Editor ? Editor->GetStackBeingInspected() : nullptr;
 	IPCGGraphExecutionSource* PCGSource = Editor ? Editor->GetPCGSourceBeingInspected() : nullptr;
