[PCG] Fix error when trying to override enum properties that require GPU readback

#rnx
#Synced-CL 48600440

diff --git a/Engine/Plugins/PCG/Source/PCG/Private/Compute/PCGKernelParamsCache.cpp b/Engine/Plugins/PCG/Source/PCG/Private/Compute/PCGKernelParamsCache.cpp
--- a/Engine/Plugins/PCG/Source/PCG/Private/Compute/PCGKernelParamsCache.cpp
+++ b/Engine/Plugins/PCG/Source/PCG/Private/Compute/PCGKernelParamsCache.cpp
@@ -91,6 +91,38 @@ namespace PCGKernelParamsHelpers
 				NumericProperty->SetIntPropertyValue(OverrideValuePtr, Value);
 			}
 		}
+		else if (const FEnumProperty* EnumPropety = CastField<FEnumProperty>(OverrideProperty))
+		{
+			const FNumericProperty* UnderlyingNumericProperty = EnumPropety->GetUnderlyingProperty();
+
+			if (!ensure(UnderlyingNumericProperty))
+			{
+				PCG_KERNEL_VALIDATION_ERR(
+					InContext,
+					InSettings,
+					FText::Format(LOCTEXT("MissingUnderlyingNumericProperty", "Tried to override pin '{0}', but enum property '{1}' does not have a valid underlying numeric property."),
+						FText::FromName(InParamName),
+						FText::FromString(OverrideProperty->GetClass()->GetName())));
+				return;
+			}
+
+			int64 Value = 0;
+
+			switch (AttributeBase->GetTypeId())
+			{
+			case PCG::Private::MetadataTypes<int32>::Id:
+				Value = static_cast<const FPCGMetadataAttribute<int32>*>(AttributeBase)->GetValueFromItemKey(PCGFirstEntryKey);
+				break;
+			case PCG::Private::MetadataTypes<int64>::Id:
+				Value = static_cast<const FPCGMetadataAttribute<int64>*>(AttributeBase)->GetValueFromItemKey(PCGFirstEntryKey);
+				break;
+			default:
+				LogUnsupportedAttributeTypeError();
+				return;
+			}
+
+			UnderlyingNumericProperty->SetIntPropertyValue(OverrideValuePtr, Value);
+		}
 		else if (const FBoolProperty* BoolProperty = CastField<FBoolProperty>(OverrideProperty))
 		{
 			bool Value = false;
