[PCG] Expose texture settings in GPU pin properties

Can set the format and filter on texture data GPU pins.

#Synced-CL 48600440

diff --git a/Engine/Plugins/PCG/Source/PCG/Private/Compute/DataInterfaces/PCGTextureDataInterface.cpp b/Engine/Plugins/PCG/Source/PCG/Private/Compute/DataInterfaces/PCGTextureDataInterface.cpp
--- a/Engine/Plugins/PCG/Source/PCG/Private/Compute/DataInterfaces/PCGTextureDataInterface.cpp
+++ b/Engine/Plugins/PCG/Source/PCG/Private/Compute/DataInterfaces/PCGTextureDataInterface.cpp
@@ -81,13 +81,14 @@ FPCGTextureBindingInfo::FPCGTextureBindingInfo(const UPCGBaseTextureData* InText
 	Texture = InTextureData->GetTextureRHI();
 	ExportedTexture = InTextureData->GetRefCountedTexture();
 	Transform = InTextureData->GetTransform();
-	bPointSample = InTextureData->Filter == EPCGTextureFilter::Point;
+	Filter = InTextureData->Filter;
 
 	if (Texture.IsValid())
 	{
 		const FRHITextureDesc& Desc = Texture->GetDesc();
 		Size = Desc.Extent;
 		Dimension = Desc.Dimension;
+		Format = Desc.Format;
 	}
 	else
 	{
@@ -99,6 +100,8 @@ FPCGTextureBindingInfo::FPCGTextureBindingInfo(const FPCGDataDesc& InDataDesc, c
 {
 	// @todo_pcg: Need to support other information here too, like format, slice index, resource type, etc.
 	Size = FIntPoint(FMath::Max(InDataDesc.GetElementCount().X, 1), FMath::Max(InDataDesc.GetElementCount().Y, 1));
+	Format = PCGComputeHelpers::GetPixelFormatFromPCGRenderTargetFormat(InDataDesc.GetRenderTargetFormat());
+	Filter = InDataDesc.GetTextureFilter();
 	Transform = InTransform;
 }
 
@@ -109,7 +112,8 @@ bool FPCGTextureBindingInfo::operator==(const FPCGTextureBindingInfo& Other) con
 		&& ExportedTexture == Other.ExportedTexture
 		&& Size == Other.Size
 		&& Dimension == Other.Dimension
-		&& bPointSample == Other.bPointSample
+		&& Format == Other.Format
+		&& Filter == Other.Filter
 		&& Transform.Equals(Other.Transform);
 }
 
@@ -524,11 +528,10 @@ void FPCGTextureDataProviderProxy::CreateTextures(FRDGBuilder& GraphBuilder)
 		}
 		else if (BindingIndex < PCGTextureDataInterfaceConstants::MAX_NUM_UAV_BINDINGS)
 		{
-			// @todo_pcg: Expose options like size, format, dimensions (e.g. Tex2D vs TexArray), etc.
 			Texture = GraphBuilder.CreateTexture(
 				FRDGTextureDesc::Create2D(
 					BindingInfos[BindingIndex].Size,
-					EPixelFormat::PF_FloatRGBA,
+					BindingInfos[BindingIndex].Format,
 					FClearValueBinding::Black,
 					ETextureCreateFlags::ShaderResource | ETextureCreateFlags::UAV),
 				TEXT("PCGTextureDI_UAV"));
@@ -689,7 +692,7 @@ void FPCGTextureDataProviderProxy::PackTextureInfos(FRDGBuilder& GraphBuilder)
 			const float TexelWorldSizeX = BindingInfo.Size.X != 0 ? (TextureWorldScale.X * 2.0) / BindingInfo.Size.X : 0.0f;
 			const float TexelWorldSizeY = BindingInfo.Size.Y != 0 ? (TextureWorldScale.Y * 2.0) / BindingInfo.Size.Y : 0.0f;
 
-			TextureInfosPackedData[TextureInfoIndex * 2 + 0] = FUintVector4(TextureInfo.BindingIndex, TextureInfo.SliceIndex, static_cast<uint32>(BindingInfo.Dimension), BindingInfo.bPointSample);
+			TextureInfosPackedData[TextureInfoIndex * 2 + 0] = FUintVector4(TextureInfo.BindingIndex, TextureInfo.SliceIndex, static_cast<uint32>(BindingInfo.Dimension), BindingInfo.Filter == EPCGTextureFilter::Point);
 			TextureInfosPackedData[TextureInfoIndex * 2 + 1] = FUintVector4(BindingInfo.Size.X, BindingInfo.Size.Y, FMath::AsUInt(TexelWorldSizeX), FMath::AsUInt(TexelWorldSizeY));
 		}
 
diff --git a/Engine/Plugins/PCG/Source/PCG/Private/Compute/DataInterfaces/PCGTextureDataInterface.h b/Engine/Plugins/PCG/Source/PCG/Private/Compute/DataInterfaces/PCGTextureDataInterface.h
--- a/Engine/Plugins/PCG/Source/PCG/Private/Compute/DataInterfaces/PCGTextureDataInterface.h
+++ b/Engine/Plugins/PCG/Source/PCG/Private/Compute/DataInterfaces/PCGTextureDataInterface.h
@@ -37,8 +37,8 @@ struct FPCGTextureBindingInfo
 	TRefCountPtr<IPooledRenderTarget> ExportedTexture = nullptr;
 	FIntPoint Size = FIntPoint::ZeroValue;
 	ETextureDimension Dimension = ETextureDimension::Texture2D;
-	bool bPointSample = false;
-
+	EPixelFormat Format = EPixelFormat::PF_Unknown;
+	EPCGTextureFilter Filter = EPCGTextureFilter::Bilinear;
 	FTransform Transform = FTransform::Identity;
 };
 
diff --git a/Engine/Plugins/PCG/Source/PCG/Private/Compute/Elements/PCGCustomHLSL.h b/Engine/Plugins/PCG/Source/PCG/Private/Compute/Elements/PCGCustomHLSL.h
--- a/Engine/Plugins/PCG/Source/PCG/Private/Compute/Elements/PCGCustomHLSL.h
+++ b/Engine/Plugins/PCG/Source/PCG/Private/Compute/Elements/PCGCustomHLSL.h
@@ -156,6 +156,12 @@ public:
 	UPROPERTY(EditAnywhere, Category = "Settings", meta = (EditCondition = "KernelType == EPCGKernelType::TextureGenerator", EditConditionHides, PCG_OverridableCPUAndGPUWithReadback))
 	int NumElementsY = 64;
 
+	UPROPERTY(EditAnywhere, Category = "Settings", meta = (EditCondition = "KernelType == EPCGKernelType::TextureGenerator", EditConditionHides, PCG_OverridableCPUAndGPUWithReadback))
+	EPCGRenderTargetFormat TextureFormat = EPCGRenderTargetFormat::RGBA16f;
+
+	UPROPERTY(EditAnywhere, Category = "Settings", meta = (EditCondition = "KernelType == EPCGKernelType::TextureGenerator", EditConditionHides, PCG_OverridableCPUAndGPUWithReadback))
+	EPCGTextureFilter TextureFilter = EPCGTextureFilter::Bilinear;
+
 	UPROPERTY(EditAnywhere, Category = "Settings|Thread Count", meta = (EditCondition = "KernelType == EPCGKernelType::Custom", EditConditionHides))
 	EPCGDispatchThreadCount DispatchThreadCount = EPCGDispatchThreadCount::FromFirstOutputPin;
 
diff --git a/Engine/Plugins/PCG/Source/PCG/Private/Compute/Elements/PCGCustomHLSLKernel.cpp b/Engine/Plugins/PCG/Source/PCG/Private/Compute/Elements/PCGCustomHLSLKernel.cpp
--- a/Engine/Plugins/PCG/Source/PCG/Private/Compute/Elements/PCGCustomHLSLKernel.cpp
+++ b/Engine/Plugins/PCG/Source/PCG/Private/Compute/Elements/PCGCustomHLSLKernel.cpp
@@ -283,9 +283,13 @@ TSharedPtr<const FPCGDataCollectionDesc> UPCGCustomHLSLKernel::ComputeOutputBind
 
 		const int NumElementsX = KernelParams->GetValueInt(GET_MEMBER_NAME_CHECKED(UPCGCustomHLSLSettings, NumElementsX));
 		const int NumElementsY = KernelParams->GetValueInt(GET_MEMBER_NAME_CHECKED(UPCGCustomHLSLSettings, NumElementsY));
+		const EPCGRenderTargetFormat TextureFormat = KernelParams->GetValueEnum<EPCGRenderTargetFormat>(GET_MEMBER_NAME_CHECKED(UPCGCustomHLSLSettings, TextureFormat));
+		const EPCGTextureFilter TextureFilter = KernelParams->GetValueEnum<EPCGTextureFilter>(GET_MEMBER_NAME_CHECKED(UPCGCustomHLSLSettings, TextureFilter));
 
 		// Generators always produce a single texture data with known size.
 		OutDataDesc->GetDataDescriptionsMutable().Emplace(EPCGDataType::BaseTexture, FIntVector4(NumElementsX, NumElementsY, 0, 0));
+		OutDataDesc->SetRenderTargetFormatForAllData(TextureFormat);
+		OutDataDesc->SetTextureFilterForAllData(TextureFilter);
 	}
 	else
 	{
diff --git a/Engine/Plugins/PCG/Source/PCG/Private/Compute/PCGComputeCommon.cpp b/Engine/Plugins/PCG/Source/PCG/Private/Compute/PCGComputeCommon.cpp
--- a/Engine/Plugins/PCG/Source/PCG/Private/Compute/PCGComputeCommon.cpp
+++ b/Engine/Plugins/PCG/Source/PCG/Private/Compute/PCGComputeCommon.cpp
@@ -241,6 +241,86 @@ namespace PCGComputeHelpers
 		return FString::Format(TEXT("{0}_DataResolver"), { *InPinLabel.ToString() });
 	}
 
+	EPixelFormat GetPixelFormatFromPCGRenderTargetFormat(EPCGRenderTargetFormat InRenderTargetFormat)
+	{
+		switch (InRenderTargetFormat)
+		{
+		// 8 bit
+		case EPCGRenderTargetFormat::R8:      return PF_G8;
+		case EPCGRenderTargetFormat::RG8:     return PF_R8G8;
+		case EPCGRenderTargetFormat::RGBA8:   return PF_B8G8R8A8;
+
+		// 16 bit
+		case EPCGRenderTargetFormat::R16f:    return PF_R16F;
+		case EPCGRenderTargetFormat::RG16f:   return PF_G16R16F;
+		case EPCGRenderTargetFormat::RGBA16f: return PF_FloatRGBA;
+
+		// 32 bit
+		case EPCGRenderTargetFormat::R32f:    return PF_R32_FLOAT;
+		case EPCGRenderTargetFormat::RG32f:   return PF_G32R32F;
+		case EPCGRenderTargetFormat::RGBA32f: return PF_A32B32G32R32F;
+		}
+
+		UE_LOG(LogPCG, Verbose, TEXT("No mapping from EPCGRenderTargetFormat to EPixelFormat is possible for entry %u, falling back to EPixelFormat::FloatRGBA."), (uint32)InRenderTargetFormat);
+
+		// Default to FloatRGBA as it is a safe value that won't crash the renderer.
+		return PF_FloatRGBA;
+	}
+
+	EPCGRenderTargetFormat GetPCGRenderTargetFormatFromPixelFormat(EPixelFormat InPixelFormat)
+	{
+		switch (InPixelFormat)
+		{
+		// 8 bit
+		case EPixelFormat::PF_R8:             // fallthrough
+		case EPixelFormat::PF_G8:             return EPCGRenderTargetFormat::R8;
+		case EPixelFormat::PF_R8G8:           return EPCGRenderTargetFormat::RG8;
+		case EPixelFormat::PF_B8G8R8A8:       // fallthrough
+		case EPixelFormat::PF_R8G8B8A8:       return EPCGRenderTargetFormat::RGBA8;
+
+		// 16 bit
+		case EPixelFormat::PF_R16F:           // fallthrough
+		case EPixelFormat::PF_G16:            return EPCGRenderTargetFormat::R16f;
+		case EPixelFormat::PF_G16R16F:        return EPCGRenderTargetFormat::RG16f;
+		case EPixelFormat::PF_FloatRGBA:      // fallthrough
+		case EPixelFormat::PF_FloatRGB:       // fallthrough
+		case EPixelFormat::PF_FloatR11G11B10: return EPCGRenderTargetFormat::RGBA16f;
+
+		// 32 bit
+		case EPixelFormat::PF_R32_FLOAT:      return EPCGRenderTargetFormat::R32f;
+		case EPixelFormat::PF_G32R32F:        return EPCGRenderTargetFormat::RG32f;
+		case EPixelFormat::PF_A32B32G32R32F:  // fallthrough
+		case EPixelFormat::PF_R32G32B32F:     return EPCGRenderTargetFormat::RGBA32f;
+		}
+
+		const FPixelFormatInfo& FormatInfo = GPixelFormats[InPixelFormat];
+
+		if (IsFloatFormat(InPixelFormat) || IsHDR(InPixelFormat))
+		{
+			switch (FormatInfo.NumComponents)
+			{
+			case 1: return EPCGRenderTargetFormat::R16f;
+			case 2: return EPCGRenderTargetFormat::RG16f;
+			case 3: // fallthrough
+			case 4: return EPCGRenderTargetFormat::RGBA16f;
+			}
+		}
+		else
+		{
+			switch (FormatInfo.NumComponents)
+			{
+			case 1: return EPCGRenderTargetFormat::R8;
+			case 2: return EPCGRenderTargetFormat::RG8;
+			case 3: // fallthrough
+			case 4: return EPCGRenderTargetFormat::RGBA8;
+			}
+		}
+
+		UE_LOG(LogPCG, Verbose, TEXT("No mapping from EPixelFormat to EPCGRenderTargetFormat is possible for %s, falling back to EPCGRenderTargetFormat::RGBA16f."), *UEnum::GetValueAsString<EPixelFormat>(InPixelFormat));
+
+		return EPCGRenderTargetFormat::RGBA16f;
+	}
+
 #if WITH_EDITOR
 	void ConvertObjectPathToShaderFilePath(FString& InOutPath)
 	{
diff --git a/Engine/Plugins/PCG/Source/PCG/Private/Compute/PCGComputeGraph.cpp b/Engine/Plugins/PCG/Source/PCG/Private/Compute/PCGComputeGraph.cpp
--- a/Engine/Plugins/PCG/Source/PCG/Private/Compute/PCGComputeGraph.cpp
+++ b/Engine/Plugins/PCG/Source/PCG/Private/Compute/PCGComputeGraph.cpp
@@ -213,86 +213,96 @@ void UPCGComputeGraph::DebugLogDataDescriptions(const UPCGDataBinding* InBinding
 		for (int32 DataIndex = 0; DataIndex < InDataDesc->GetDataDescriptions().Num(); ++DataIndex)
 		{
 			const FPCGDataDesc& DataDesc = InDataDesc->GetDataDescriptions()[DataIndex];
+			const FPCGDataTypeIdentifier& Type = DataDesc.GetType();
 
 			UE_LOG(LogPCG, Display, TEXT("\t\tData Index %d"), DataIndex);
-			UE_LOG(LogPCG, Display, TEXT("\t\t\tType: %s"), *DataDesc.GetType().ToString());
+			UE_LOG(LogPCG, Display, TEXT("\t\t\tType: %s"), *Type.ToString());
 			UE_LOG(LogPCG, Display, TEXT("\t\t\tNum Elements: %s (%s)"), *DataDesc.GetElementCount().ToString(), *UEnum::GetValueAsString(DataDesc.GetElementDimension()));
-			UE_LOG(LogPCG, Display, TEXT("\t\t\tAttributes (%d)"), DataDesc.GetAttributeDescriptions().Num());
 
-			TSharedPtr<const FPCGProxyForGPUDataCollection> UpstreamDataCollectionInfo;
-
-			if (const FName* VirtualLabel = GetVirtualPinLabelFromKernelInputPin(InKernelPin))
+			if (Type.IsChildOf(FPCGDataTypeInfoBaseTexture2D::AsId()))
+			{
+				UE_LOG(LogPCG, Display, TEXT("\t\t\tFormat: %s"), *UEnum::GetValueAsString(DataDesc.GetRenderTargetFormat()));
+				UE_LOG(LogPCG, Display, TEXT("\t\t\tFilter: %s"), *UEnum::GetValueAsString(DataDesc.GetTextureFilter()));
+			}
+			else
 			{
-				TArray<FPCGTaggedData> PinInputs = InBinding->GetInputDataCollection().GetInputsByPin(*VirtualLabel);
+				UE_LOG(LogPCG, Display, TEXT("\t\t\tAttributes (%d)"), DataDesc.GetAttributeDescriptions().Num());
 
-				if (ensure(PinInputs.IsValidIndex(DataIndex)))
+				TSharedPtr<const FPCGProxyForGPUDataCollection> UpstreamDataCollectionInfo;
+
+				if (const FName* VirtualLabel = GetVirtualPinLabelFromKernelInputPin(InKernelPin))
 				{
-					if (const UPCGProxyForGPUData* Proxy = Cast<UPCGProxyForGPUData>(PinInputs[DataIndex].Data))
+					TArray<FPCGTaggedData> PinInputs = InBinding->GetInputDataCollection().GetInputsByPin(*VirtualLabel);
+
+					if (ensure(PinInputs.IsValidIndex(DataIndex)))
 					{
-						UpstreamDataCollectionInfo = Proxy->GetGPUInfo();
+						if (const UPCGProxyForGPUData* Proxy = Cast<UPCGProxyForGPUData>(PinInputs[DataIndex].Data))
+						{
+							UpstreamDataCollectionInfo = Proxy->GetGPUInfo();
+						}
 					}
 				}
-			}
-
-			for (const FPCGKernelAttributeDesc& AttributeDesc : DataDesc.GetAttributeDescriptions())
-			{
-				const bool bAllocated = DataDesc.IsAttributeAllocated(AttributeDesc);
 
-				// If the data comes from an upstream / source buffer then look up its info.
-				FString UniqueStringKeys;
-				for (int32 GraphStringKey : AttributeDesc.GetUniqueStringKeys())
+				for (const FPCGKernelAttributeDesc& AttributeDesc : DataDesc.GetAttributeDescriptions())
 				{
-					UniqueStringKeys += TEXT(", ");
-					UniqueStringKeys += FString::Format(TEXT("{0}"), { GraphStringKey });
+					const bool bAllocated = DataDesc.IsAttributeAllocated(AttributeDesc);
 
-					if (const int32* UpstreamStringKey = UpstreamDataCollectionInfo ? InBinding->GetGraphToSourceBufferStringKey().Find({ UpstreamDataCollectionInfo->GetBuffer(), GraphStringKey }) : nullptr)
+					// If the data comes from an upstream / source buffer then look up its info.
+					FString UniqueStringKeys;
+					for (int32 GraphStringKey : AttributeDesc.GetUniqueStringKeys())
 					{
-						UniqueStringKeys += FString::Format(TEXT(" ({0})"), { *UpstreamStringKey });
+						UniqueStringKeys += TEXT(", ");
+						UniqueStringKeys += FString::Format(TEXT("{0}"), { GraphStringKey });
+
+						if (const int32* UpstreamStringKey = UpstreamDataCollectionInfo ? InBinding->GetGraphToSourceBufferStringKey().Find({ UpstreamDataCollectionInfo->GetBuffer(), GraphStringKey }) : nullptr)
+						{
+							UniqueStringKeys += FString::Format(TEXT(" ({0})"), { *UpstreamStringKey });
+						}
 					}
-				}
 
-				// If attribute comes from an upstream GPU buffer, display the attribute ID in the upstream buffer.
-				int32 SourceBufferAttributeId = INDEX_NONE;
+					// If attribute comes from an upstream GPU buffer, display the attribute ID in the upstream buffer.
+					int32 SourceBufferAttributeId = INDEX_NONE;
 
-				if (AttributeDesc.GetAttributeId() >= PCGDataCollectionPackingConstants::NUM_RESERVED_ATTRS && UpstreamDataCollectionInfo)
-				{
-					const int32 GraphAttributeIndex = PCGComputeHelpers::GetMetadataAttributeIndexFromAttributeId(AttributeDesc.GetAttributeId());
-					const int32* FoundIndex = InBinding->GetGraphAttributeToSourceBufferAttributeIndex().Find({ GraphAttributeIndex, UpstreamDataCollectionInfo->GetBuffer() });
+					if (AttributeDesc.GetAttributeId() >= PCGDataCollectionPackingConstants::NUM_RESERVED_ATTRS && UpstreamDataCollectionInfo)
+					{
+						const int32 GraphAttributeIndex = PCGComputeHelpers::GetMetadataAttributeIndexFromAttributeId(AttributeDesc.GetAttributeId());
+						const int32* FoundIndex = InBinding->GetGraphAttributeToSourceBufferAttributeIndex().Find({ GraphAttributeIndex, UpstreamDataCollectionInfo->GetBuffer() });
 
-					SourceBufferAttributeId = PCGComputeHelpers::GetAttributeIdFromMetadataAttributeIndex(FoundIndex ? *FoundIndex : GraphAttributeIndex);
-				}
+						SourceBufferAttributeId = PCGComputeHelpers::GetAttributeIdFromMetadataAttributeIndex(FoundIndex ? *FoundIndex : GraphAttributeIndex);
+					}
 
-				if (SourceBufferAttributeId == INDEX_NONE)
-				{
-					UE_LOG(LogPCG, Display, TEXT("\t\t\t\tID: %d\t\tName: %s\t\tType: %s\t\tAllocated: %s\t\tUniqueStringKeys%s"),
-						AttributeDesc.GetAttributeId(),
-						*AttributeDesc.GetAttributeKey().GetIdentifier().ToString(),
-						*PCGKernelAttributeTypeEnum->GetDisplayNameTextByValue(static_cast<int64>(AttributeDesc.GetAttributeKey().GetType())).ToString(),
-						bAllocated ? TEXT("true") : TEXT("false"),
-						UniqueStringKeys.IsEmpty() ? TEXT(": ") : *UniqueStringKeys);
+					if (SourceBufferAttributeId == INDEX_NONE)
+					{
+						UE_LOG(LogPCG, Display, TEXT("\t\t\t\tID: %d\t\tName: %s\t\tType: %s\t\tAllocated: %s\t\tUniqueStringKeys%s"),
+							AttributeDesc.GetAttributeId(),
+							*AttributeDesc.GetAttributeKey().GetIdentifier().ToString(),
+							*PCGKernelAttributeTypeEnum->GetDisplayNameTextByValue(static_cast<int64>(AttributeDesc.GetAttributeKey().GetType())).ToString(),
+							bAllocated ? TEXT("true") : TEXT("false"),
+							UniqueStringKeys.IsEmpty() ? TEXT(": ") : *UniqueStringKeys);
+					}
+					else
+					{
+						UE_LOG(LogPCG, Display, TEXT("\t\t\t\tID: %d\t\tName: %s\t\tType: %s\t\tAllocated: %s\t\tUniqueStringKeys%s\t\tSource Buffer ID: %d"),
+							AttributeDesc.GetAttributeId(),
+							*AttributeDesc.GetAttributeKey().GetIdentifier().ToString(),
+							*PCGKernelAttributeTypeEnum->GetDisplayNameTextByValue(static_cast<int64>(AttributeDesc.GetAttributeKey().GetType())).ToString(),
+							bAllocated ? TEXT("true") : TEXT("false"),
+							UniqueStringKeys.IsEmpty() ? TEXT(": ") : *UniqueStringKeys,
+							SourceBufferAttributeId);
+					}
 				}
-				else
+
+				UE_LOG(LogPCG, Display, TEXT("\t\t\tTags (%d)"), DataDesc.GetTagStringKeys().Num());
+
+				FString TagStringKeys;
+				for (int32 TagStringKey : DataDesc.GetTagStringKeys())
 				{
-					UE_LOG(LogPCG, Display, TEXT("\t\t\t\tID: %d\t\tName: %s\t\tType: %s\t\tAllocated: %s\t\tUniqueStringKeys%s\t\tSource Buffer ID: %d"),
-						AttributeDesc.GetAttributeId(),
-						*AttributeDesc.GetAttributeKey().GetIdentifier().ToString(),
-						*PCGKernelAttributeTypeEnum->GetDisplayNameTextByValue(static_cast<int64>(AttributeDesc.GetAttributeKey().GetType())).ToString(),
-						bAllocated ? TEXT("true") : TEXT("false"),
-						UniqueStringKeys.IsEmpty() ? TEXT(": ") : *UniqueStringKeys,
-						SourceBufferAttributeId);
+					TagStringKeys += TEXT(", ");
+					TagStringKeys += FString::Format(TEXT("{0}"), { TagStringKey });
 				}
-			}
 
-			UE_LOG(LogPCG, Display, TEXT("\t\t\tTags (%d)"), DataDesc.GetTagStringKeys().Num());
-
-			FString TagStringKeys;
-			for (int32 TagStringKey : DataDesc.GetTagStringKeys())
-			{
-				TagStringKeys += TEXT(", ");
-				TagStringKeys += FString::Format(TEXT("{0}"), { TagStringKey });
+				UE_LOG(LogPCG, Display, TEXT("\t\t\tTag String Keys%s"), *TagStringKeys);
 			}
-
-			UE_LOG(LogPCG, Display, TEXT("\t\t\tTag String Keys%s"), *TagStringKeys);
 		}
 	};
 
diff --git a/Engine/Plugins/PCG/Source/PCG/Private/Compute/PCGComputeKernel.cpp b/Engine/Plugins/PCG/Source/PCG/Private/Compute/PCGComputeKernel.cpp
--- a/Engine/Plugins/PCG/Source/PCG/Private/Compute/PCGComputeKernel.cpp
+++ b/Engine/Plugins/PCG/Source/PCG/Private/Compute/PCGComputeKernel.cpp
@@ -522,6 +522,30 @@ void UPCGComputeKernel::ComputeDataDescFromPinProperties(const FPCGPinProperties
 		{
 			Desc.ScaleElementCount(Multiplier);
 		}
+
+		if (OutputPinProps.AllowedTypes.IsChildOf(EPCGDataType::BaseTexture))
+		{
+			// This code assumes texture output pins are single data.
+			const FPCGDataDesc* InputDesc = (!InputDescs.IsEmpty() && InputDescs[0] && !InputDescs[0]->GetDataDescriptions().IsEmpty()) ? &InputDescs[0]->GetDataDescriptions()[0] : nullptr;
+
+			if (OutputPinProps.PropertiesGPU.TextureFormatInheritanceMode == EPCGPropertyInheritanceMode::FromInput && InputDesc)
+			{
+				OutPinDesc->SetRenderTargetFormatForAllData(InputDesc->GetRenderTargetFormat());
+			}
+			else
+			{
+				OutPinDesc->SetRenderTargetFormatForAllData(OutputPinProps.PropertiesGPU.TextureFormat);
+			}
+
+			if (OutputPinProps.PropertiesGPU.TextureFilterInheritanceMode == EPCGPropertyInheritanceMode::FromInput && InputDesc)
+			{
+				OutPinDesc->SetTextureFilterForAllData(InputDesc->GetTextureFilter());
+			}
+			else
+			{
+				OutPinDesc->SetTextureFilterForAllData(OutputPinProps.PropertiesGPU.TextureFilter);
+			}
+		}
 	}
 	else if (Props.InitializationMode == EPCGPinInitMode::Custom)
 	{
@@ -542,7 +566,12 @@ void UPCGComputeKernel::ComputeDataDescFromPinProperties(const FPCGPinProperties
 				// Not supported yet.
 				checkNoEntry();
 			}
+		}
 
+		if (OutputPinProps.AllowedTypes.IsChildOf(EPCGDataType::BaseTexture))
+		{
+			OutPinDesc->SetRenderTargetFormatForAllData(OutputPinProps.PropertiesGPU.TextureFormat);
+			OutPinDesc->SetTextureFilterForAllData(OutputPinProps.PropertiesGPU.TextureFilter);
 		}
 	}
 	else
diff --git a/Engine/Plugins/PCG/Source/PCG/Private/Compute/PCGDataDescription.cpp b/Engine/Plugins/PCG/Source/PCG/Private/Compute/PCGDataDescription.cpp
--- a/Engine/Plugins/PCG/Source/PCG/Private/Compute/PCGDataDescription.cpp
+++ b/Engine/Plugins/PCG/Source/PCG/Private/Compute/PCGDataDescription.cpp
@@ -259,6 +259,12 @@ FPCGDataDesc::FPCGDataDesc(const FPCGTaggedData& InTaggedData, const UPCGDataBin
 	ElementCount = PCGComputeHelpers::GetElementCount(InTaggedData.Data);
 	ElementDimension = PCGComputeHelpers::GetElementDimension(InTaggedData.Data);
 
+	if (const UPCGBaseTextureData* TextureData = Cast<UPCGBaseTextureData>(InTaggedData.Data))
+	{
+		SetTextureFilter(TextureData->Filter);
+		SetRenderTargetFormat(PCGComputeHelpers::GetPCGRenderTargetFormatFromPixelFormat(TextureData->GetTextureFormat()));
+	}
+
 	TagStringKeys.Reserve(InTaggedData.Tags.Num());
 
 	for (const FString& Tag : InTaggedData.Tags)
@@ -769,4 +775,20 @@ int FPCGDataCollectionDesc::GetNumStringKeyValues(int32 InAttributeId) const
 	return UniqueStringKeys.Num();
 }
 
+void FPCGDataCollectionDesc::SetRenderTargetFormatForAllData(EPCGRenderTargetFormat InRenderTargetFormat)
+{
+	for (FPCGDataDesc& DataDesc : DataDescs)
+	{
+		DataDesc.SetRenderTargetFormat(InRenderTargetFormat);
+	}
+}
+
+void FPCGDataCollectionDesc::SetTextureFilterForAllData(EPCGTextureFilter InTextureFilter)
+{
+	for (FPCGDataDesc& DataDesc : DataDescs)
+	{
+		DataDesc.SetTextureFilter(InTextureFilter);
+	}
+}
+
 #undef LOCTEXT_NAMESPACE
diff --git a/Engine/Plugins/PCG/Source/PCG/Private/Data/PCGRenderTargetData.cpp b/Engine/Plugins/PCG/Source/PCG/Private/Data/PCGRenderTargetData.cpp
--- a/Engine/Plugins/PCG/Source/PCG/Private/Data/PCGRenderTargetData.cpp
+++ b/Engine/Plugins/PCG/Source/PCG/Private/Data/PCGRenderTargetData.cpp
@@ -23,14 +23,15 @@ void UPCGRenderTargetData::Initialize(UTextureRenderTarget2D* InRenderTarget, co
 	{
 		Width = RenderTarget->SizeX;
 		Height = RenderTarget->SizeY;
+		Format = RenderTarget->GetFormat();
 
 		if (!bSkipReadbackToCPU)
 		{
 			if (FTextureRenderTargetResource* RTResource = RenderTarget->GameThread_GetRenderTargetResource())
 			{
 				TRACE_CPUPROFILER_EVENT_SCOPE(UPCGRenderTargetData::Initialize::ReadData);
-				ETextureRenderTargetFormat Format = RenderTarget->RenderTargetFormat;
-				if ((Format == RTF_RGBA16f) || (Format == RTF_RGBA32f) || (Format == RTF_RGBA8) || (Format == RTF_RGB10A2))
+				ETextureRenderTargetFormat RenderTargetFormat = RenderTarget->RenderTargetFormat;
+				if ((RenderTargetFormat == RTF_RGBA16f) || (RenderTargetFormat == RTF_RGBA32f) || (RenderTargetFormat == RTF_RGBA8) || (RenderTargetFormat == RTF_RGB10A2))
 				{
 					const FIntRect Rect = FIntRect(0, 0, RenderTarget->SizeX, RenderTarget->SizeY);
 					const FReadSurfaceDataFlags ReadPixelFlags(RCM_MinMax);
@@ -38,7 +39,7 @@ void UPCGRenderTargetData::Initialize(UTextureRenderTarget2D* InRenderTarget, co
 				}
 				else
 				{
-					UE_LOG(LogPCG, Warning, TEXT("Could not CPU-initialize UPCGRenderTargetData with render target format '%s'."), *UEnum::GetValueAsString(Format));
+					UE_LOG(LogPCG, Warning, TEXT("Could not CPU-initialize UPCGRenderTargetData with render target format '%s'."), *UEnum::GetValueAsString(RenderTargetFormat));
 				}
 			}
 		}
diff --git a/Engine/Plugins/PCG/Source/PCG/Private/Data/PCGTextureData.cpp b/Engine/Plugins/PCG/Source/PCG/Private/Data/PCGTextureData.cpp
--- a/Engine/Plugins/PCG/Source/PCG/Private/Data/PCGTextureData.cpp
+++ b/Engine/Plugins/PCG/Source/PCG/Private/Data/PCGTextureData.cpp
@@ -492,6 +492,7 @@ void UPCGTextureData::InitializeInternal(UTexture* InTexture, uint32 InTextureIn
 	Width = 0;
 	Height = 0;
 	bSkipReadbackToCPU = bInSkipReadbackToCPU;
+	Format = GetTextureRHI() ? GetTextureRHI()->GetDesc().Format : PF_Unknown;
 
 	Bounds = FBox(EForceInit::ForceInit);
 	Bounds += FVector(-1.0f, -1.0f, 0.0f);
@@ -645,6 +646,7 @@ bool UPCGTextureData::Initialize(TRefCountPtr<IPooledRenderTarget> InTextureHand
 	TextureIndex = InTextureIndex;
 	Transform = InTransform;
 	bSkipReadbackToCPU = bInSkipReadbackToCPU;
+	Format = InTextureHandle.IsValid() ? InTextureHandle->GetDesc().Format : PF_Unknown;
 
 	Bounds = FBox(EForceInit::ForceInit);
 	Bounds += FVector(-1.0f, -1.0f, 0.0f);
diff --git a/Engine/Plugins/PCG/Source/PCG/Public/Compute/PCGComputeCommon.h b/Engine/Plugins/PCG/Source/PCG/Public/Compute/PCGComputeCommon.h
--- a/Engine/Plugins/PCG/Source/PCG/Public/Compute/PCGComputeCommon.h
+++ b/Engine/Plugins/PCG/Source/PCG/Public/Compute/PCGComputeCommon.h
@@ -5,21 +5,23 @@
 #include "PCGCommon.h"
 #include "Data/Registry/PCGDataTypeIdentifier.h"
 
+#include "Engine/TextureRenderTarget2D.h"
+
 #include "PCGComputeCommon.generated.h"
 
+class FRHIShaderResourceView;
+class UObject;
 class UPCGComputeDataInterface;
 class UPCGComputeKernel;
 class UPCGData;
 class UPCGDataBinding;
 class UPCGNode;
-class UObject;
 class UPCGSettings;
-struct FPCGGPUCompilationContext;
+enum EPixelFormat : uint8;
 struct FPCGContext;
+struct FPCGGPUCompilationContext;
 struct FPCGPinProperties;
 
-class FRHIShaderResourceView;
-
 #define PCG_KERNEL_LOGGING_ENABLED (!(UE_BUILD_SHIPPING || UE_BUILD_TEST) || USE_LOGGING_IN_SHIPPING)
 
 #if PCG_KERNEL_LOGGING_ENABLED
@@ -55,6 +57,38 @@ enum class EPCGElementDimension : uint8
 	Four
 };
 
+/** Mirror of ETextureRenderTargetFormat that exposes only formats supported in PCG data descriptions. */
+UENUM()
+enum class EPCGRenderTargetFormat : uint8
+{
+	/** R channel, 8 bit per channel fixed point, range [0, 1]. */
+	R8 = RTF_R8,
+
+	/** RG channels, 8 bit per channel fixed point, range [0, 1]. */
+	RG8 = ETextureRenderTargetFormat::RTF_RG8,
+
+	/** RGBA channels, 8 bit per channel fixed point, range [0, 1]. */
+	RGBA8 = ETextureRenderTargetFormat::RTF_RGBA8,
+
+	/** R channel, 16 bit per channel floating point, range [-65504, 65504]. */
+	R16f = ETextureRenderTargetFormat::RTF_R16f,
+
+	/** RG channels, 16 bit per channel floating point, range [-65504, 65504]. */
+	RG16f = ETextureRenderTargetFormat::RTF_RG16f,
+
+	/** RGBA channels, 16 bit per channel floating point, range [-65504, 65504]. */
+	RGBA16f = ETextureRenderTargetFormat::RTF_RGBA16f,
+
+	/** R channel, 32 bit per channel floating point, range [-3.402823 x 10^38, 3.402823 x 10^38]. */
+	R32f = ETextureRenderTargetFormat::RTF_R32f,
+
+	/** RG channels, 32 bit per channel floating point, range [-3.402823 x 10^38, 3.402823 x 10^38]. */
+	RG32f = ETextureRenderTargetFormat::RTF_RG32f,
+
+	/** RGBA channels, 32 bit per channel floating point, range [-3.402823 x 10^38, 3.402823 x 10^38]. */
+	RGBA32f = ETextureRenderTargetFormat::RTF_RGBA32f,
+};
+
 namespace PCGComputeConstants
 {
 	constexpr int32 MAX_PRIMITIVE_COMPONENTS_PER_SPAWNER = 256;
@@ -108,6 +142,9 @@ namespace PCGComputeHelpers
 	/** Produces the data interface name of a data label resolver. */
 	FString GetDataLabelResolverName(FName InPinLabel);
 
+	EPixelFormat GetPixelFormatFromPCGRenderTargetFormat(EPCGRenderTargetFormat InRenderTargetFormat);
+	EPCGRenderTargetFormat GetPCGRenderTargetFormatFromPixelFormat(EPixelFormat InPixelFormat);
+
 #if WITH_EDITOR
 	void ConvertObjectPathToShaderFilePath(FString& InOutPath);
 
diff --git a/Engine/Plugins/PCG/Source/PCG/Public/Compute/PCGDataDescription.h b/Engine/Plugins/PCG/Source/PCG/Public/Compute/PCGDataDescription.h
--- a/Engine/Plugins/PCG/Source/PCG/Public/Compute/PCGDataDescription.h
+++ b/Engine/Plugins/PCG/Source/PCG/Public/Compute/PCGDataDescription.h
@@ -5,6 +5,7 @@
 #include "PCGCommon.h"
 #include "PCGData.h"
 #include "Compute/PCGComputeCommon.h"
+#include "Data/PCGTextureData.h"
 
 #include "PCGDataDescription.generated.h"
 
@@ -195,6 +196,12 @@ public:
 	void SetTagStringKeys(TConstArrayView<int32> InTagStringKeys) { TagStringKeys = InTagStringKeys; }
 	void AddUniqueTagStringKey(int32 InTagStringKey) { TagStringKeys.AddUnique(InTagStringKey); }
 
+	EPCGRenderTargetFormat GetRenderTargetFormat() const { return RenderTargetFormat; }
+	void SetRenderTargetFormat(EPCGRenderTargetFormat InRenderTargetFormat) { RenderTargetFormat = InRenderTargetFormat; }
+
+	EPCGTextureFilter GetTextureFilter() const { return TextureFilter; }
+	void SetTextureFilter(EPCGTextureFilter InTextureFilter) { TextureFilter = InTextureFilter; }
+
 private:
 	void InitializeAttributeDescs(const UPCGData* InData, const UPCGDataBinding* InBinding);
 
@@ -208,6 +215,11 @@ private:
 	EPCGPointNativeProperties AllocatedPointProperties = EPCGPointNativeProperties::None;
 
 	TArray<int32, TInlineAllocator<4>> TagStringKeys;
+
+	// Texture settings
+	/** Texture format used to create a render target UAV from this data. If this description represents a CPU texture data, then this format will be a narrowed conversion from its EPixelFormat. */
+	EPCGRenderTargetFormat RenderTargetFormat = EPCGRenderTargetFormat::RGBA16f;
+	EPCGTextureFilter TextureFilter = EPCGTextureFilter::Bilinear;
 };
 
 /** Data description for a data collection (FPCGDataCollection). */
@@ -246,6 +258,10 @@ struct FPCGDataCollectionDesc
 
 	PCG_API int GetNumStringKeyValues(int32 InAttributeId) const;
 
+	PCG_API void SetRenderTargetFormatForAllData(EPCGRenderTargetFormat InRenderTargetFormat);
+
+	PCG_API void SetTextureFilterForAllData(EPCGTextureFilter InTextureFilter);
+
 private:
 	FPCGDataCollectionDesc() = default;
 
diff --git a/Engine/Plugins/PCG/Source/PCG/Public/Compute/PCGPinPropertiesGPU.h b/Engine/Plugins/PCG/Source/PCG/Public/Compute/PCGPinPropertiesGPU.h
--- a/Engine/Plugins/PCG/Source/PCG/Public/Compute/PCGPinPropertiesGPU.h
+++ b/Engine/Plugins/PCG/Source/PCG/Public/Compute/PCGPinPropertiesGPU.h
@@ -4,6 +4,7 @@
 
 #include "PCGPin.h"
 #include "Compute/PCGDataDescription.h"
+#include "Data/PCGTextureData.h"
 
 #include "PCGPinPropertiesGPU.generated.h"
 
@@ -13,24 +14,25 @@
 UENUM()
 enum class EPCGPinInitMode : uint8
 {
-	FromInputPins UMETA(Tooltip = "Initialize structure of data (data count, element count, attribute definitions, etc) from one or more input pins."),
+	FromInputPins UMETA(DisplayName = "FromInput", Tooltip = "Initialize structure of data (data count, element count, attribute definitions, etc) from one or more input pins."),
 	Custom UMETA(Tooltip = "Initialize structure of data (data count, element count, attribute definitions, etc) from fixed values."),
 };
 
+// @todo_pcg: Data/Element count mode could be consolidated with the generic EPCGPropertyInheritanceMode below.
 /** Method for computing data count. */
 UENUM()
 enum class EPCGDataCountMode : uint8
 {
-	FromInputData,
-	Fixed,
+	FromInputData UMETA(DisplayName = "FromInput"),
+	Fixed UMETA(DisplayName = "Custom"),
 };
 
 /** Method for computing element count. */
 UENUM()
 enum class EPCGElementCountMode : uint8
 {
-	FromInputData,
-	Fixed,
+	FromInputData UMETA(DisplayName = "FromInput"),
+	Fixed UMETA(DisplayName = "Custom"),
 };
 
 /** Method for combining two or more data counts. */
@@ -59,6 +61,14 @@ enum class EPCGAttributeInheritanceMode : uint8
 	//CopyAttributeSetupAndValues UMETA(Tooltip = "Take attribute names, types, and initial values from initialization pins. Pins declared first will take priority during conflicts."),
 };
 
+/** Method for inheriting properties from input pins. */
+UENUM()
+enum class EPCGPropertyInheritanceMode : uint8
+{
+	FromInput,
+	Custom
+};
+
 /** Helper struct to nest GPU pin properties inside a UI category. */
 USTRUCT(BlueprintType)
 struct FPCGPinPropertiesGPUStruct
@@ -97,12 +107,26 @@ struct FPCGPinPropertiesGPUStruct
 	UPROPERTY(BlueprintReadWrite, EditAnywhere, Category = Settings, DisplayName = "Num Elements", meta = (EditCondition = "bAllowEditInitMode && (InitializationMode == EPCGPinInitMode::Custom || ElementCountMode == EPCGElementCountMode::Fixed) && !bShowTexturePinSettings", EditConditionHides, ClampMin = 1))
 	int ElementCount = 1;
 
-	// @todo_pcg: Show texture data type info, e.g. Dimension, Format, etc.
-
 	/** Fixed number of elements to create in each output data. */
 	UPROPERTY(BlueprintReadWrite, EditAnywhere, Category = Settings, DisplayName = "Num Elements", meta = (EditCondition = "bAllowEditInitMode && (InitializationMode == EPCGPinInitMode::Custom || ElementCountMode == EPCGElementCountMode::Fixed) && bShowTexturePinSettings", EditConditionHides, ClampMin = 1))
 	FIntPoint NumElements2D = FIntPoint(64, 64);
 
+	/** Method for inheriting texture format. */
+	UPROPERTY(BlueprintReadWrite, EditAnywhere, Category = Settings, meta = (EditCondition = "bAllowEditInitMode && bShowTexturePinSettings && InitializationMode == EPCGPinInitMode::FromInputPins", EditConditionHides, HideEditConditionToggle))
+	EPCGPropertyInheritanceMode TextureFormatInheritanceMode = EPCGPropertyInheritanceMode::FromInput;
+
+	/** Format for all texture data produced on this pin. */
+	UPROPERTY(BlueprintReadWrite, EditAnywhere, Category = Settings, meta = (EditCondition = "bAllowEditInitMode && (InitializationMode == EPCGPinInitMode::Custom || TextureFormatInheritanceMode == EPCGPropertyInheritanceMode::Custom) && bShowTexturePinSettings", EditConditionHides, HideEditConditionToggle))
+	EPCGRenderTargetFormat TextureFormat = EPCGRenderTargetFormat::RGBA16f;
+
+	/** Method for inheriting texture filter. */
+	UPROPERTY(BlueprintReadWrite, EditAnywhere, Category = Settings, meta = (EditCondition = "bAllowEditInitMode && bShowTexturePinSettings && InitializationMode == EPCGPinInitMode::FromInputPins", EditConditionHides, HideEditConditionToggle))
+	EPCGPropertyInheritanceMode TextureFilterInheritanceMode = EPCGPropertyInheritanceMode::FromInput;
+
+	/** Method used to determine the value for a sample based on the value of nearby texels. */
+	UPROPERTY(BlueprintReadWrite, EditAnywhere, Category = Settings, meta = (EditCondition = "bAllowEditInitMode && (InitializationMode == EPCGPinInitMode::Custom || TextureFilterInheritanceMode == EPCGPropertyInheritanceMode::Custom) && bShowTexturePinSettings", EditConditionHides, HideEditConditionToggle))
+	EPCGTextureFilter TextureFilter = EPCGTextureFilter::Bilinear;
+
 	/** Scalar on the fixed number of elements to create in each output data. */
 	UPROPERTY(BlueprintReadWrite, EditAnywhere, Category = Settings, meta = (EditCondition = "bAllowEditInitMode && (InitializationMode == EPCGPinInitMode::Custom || ElementCountMode == EPCGElementCountMode::FromInputData)", EditConditionHides, ClampMin = 1))
 	int ElementCountMultiplier = 1;
diff --git a/Engine/Plugins/PCG/Source/PCG/Public/Data/PCGTextureData.h b/Engine/Plugins/PCG/Source/PCG/Public/Data/PCGTextureData.h
--- a/Engine/Plugins/PCG/Source/PCG/Public/Data/PCGTextureData.h
+++ b/Engine/Plugins/PCG/Source/PCG/Public/Data/PCGTextureData.h
@@ -4,6 +4,7 @@
 
 #include "Data/PCGSurfaceData.h"
 
+#include "PixelFormat.h"
 #include "RendererInterface.h"
 #include "RHI.h"
 
@@ -117,6 +118,7 @@ public:
 	virtual TRefCountPtr<IPooledRenderTarget> GetRefCountedTexture() const { return nullptr; }
 	virtual int GetTextureSlice() const { return 0; }
 	virtual FIntPoint GetTextureSize() const { return FIntPoint(Width, Height); }
+	EPixelFormat GetTextureFormat() const { return Format; }
 
 protected:
 	UE_API const UPCGBasePointData* CreateBasePointData(FPCGContext* Context, TSubclassOf<UPCGBasePointData> PointDataClass) const;
@@ -185,6 +187,10 @@ protected:
 	UPROPERTY(BlueprintReadOnly, VisibleAnywhere, Category = SpatialData)
 	int32 Width = 0;
 
+	/** Format of the underlying texture resource. Note that for render targets this is distinct from the render target format. */
+	UPROPERTY(BlueprintReadOnly, VisibleAnywhere, Category = SpatialData)
+	TEnumAsByte<EPixelFormat> Format = EPixelFormat::PF_Unknown;
+
 	UPROPERTY()
 	bool bSkipReadbackToCPU = false;
 
