[PCG] Remove dynamic branch on sampler from texture data interface

Helps to alleviate excessive combined sampler counts which break mobile cooks

#rnx
#Synced-CL 48600440

diff --git a/Engine/Plugins/PCG/Shaders/Private/PCGTextureDataInterface.ush b/Engine/Plugins/PCG/Shaders/Private/PCGTextureDataInterface.ush
--- a/Engine/Plugins/PCG/Shaders/Private/PCGTextureDataInterface.ush
+++ b/Engine/Plugins/PCG/Shaders/Private/PCGTextureDataInterface.ush
@@ -32,8 +32,7 @@ PCG_TEXTURE_UAV_DECLARATION(0);
 
 #undef PCG_TEXTURE_UAV_DECLARATION
 
-SamplerState {DataInterfaceName}_Sampler_Linear;
-SamplerState {DataInterfaceName}_Sampler_Point;
+SamplerState {DataInterfaceName}_Sampler;
 
 float4x4 {DataInterfaceName}_TextureInverseTransforms[PCG_MAX_NUM_SRV_BINDINGS];
 
@@ -42,7 +41,7 @@ float4x4 {DataInterfaceName}_TextureInverseTransforms[PCG_MAX_NUM_SRV_BINDINGS];
  *   - uint32 - BindingIndex
  *   - uint32 - SliceIndex (aka TextureIndex if bound to a TextureArray, 0 otherwise.)
  *   - uint32 - Dimension
- *   - uint32 - Sampler
+ *   - uint32 - UNUSED
  *   - uint32 - ResolutionX
  *   - uint32 - ResolutionY
  *   - float  - TexelSizeWorldX (cm)
@@ -56,7 +55,7 @@ uint {DataInterfaceName}_NumTextureInfos;
 
 // #################### INTERNAL HELPERS ##########################
 
-void GetTextureInfoInternal_{DataInterfaceName}(uint InDataIndex, out uint OutBindingIndex, out uint OutSliceIndex, out uint OutTextureDimension, out uint OutPointSample)
+void GetTextureInfoInternal_{DataInterfaceName}(uint InDataIndex, out uint OutBindingIndex, out uint OutSliceIndex, out uint OutTextureDimension)
 {
 	if (InDataIndex >= 0u && InDataIndex < {DataInterfaceName}_NumTextureInfos)
 	{
@@ -64,14 +63,12 @@ void GetTextureInfoInternal_{DataInterfaceName}(uint InDataIndex, out uint OutBi
 		OutBindingIndex = TextureInfo.x;
 		OutSliceIndex = TextureInfo.y;
 		OutTextureDimension = TextureInfo.z;
-		OutPointSample = TextureInfo.w;
 	}
 	else
 	{
 		OutBindingIndex = 0u;
 		OutSliceIndex = 0u;
 		OutTextureDimension = 0u;
-		OutPointSample = 0u;
 	}
 }
 
@@ -160,25 +157,11 @@ float2 GetTexelSizeWorld_{DataInterfaceName}(uint InDataIndex)
 	{ \
 		if (TextureDimension == PCG_TextureDimension_Texture2D) \
 		{ \
-			if (PointSample == 0) \
-			{ \
-				return {DataInterfaceName}_Texture_##Index.SampleLevel({DataInterfaceName}_Sampler_Linear, InTextureUV, 0); \
-			} \
-			else \
-			{ \
-				return {DataInterfaceName}_Texture_##Index.SampleLevel({DataInterfaceName}_Sampler_Point, InTextureUV, 0); \
-			} \
+			return {DataInterfaceName}_Texture_##Index.SampleLevel({DataInterfaceName}_Sampler, InTextureUV, 0); \
 		} \
 		else if (TextureDimension == PCG_TextureDimension_Texture2DArray) \
 		{ \
-			if (PointSample == 0) \
-			{ \
-				return {DataInterfaceName}_TextureArray_##Index.SampleLevel({DataInterfaceName}_Sampler_Linear, float3(InTextureUV, SliceIndex), 0); \
-			} \
-			else \
-			{ \
-				return {DataInterfaceName}_TextureArray_##Index.SampleLevel({DataInterfaceName}_Sampler_Point, float3(InTextureUV, SliceIndex), 0); \
-			} \
+			return {DataInterfaceName}_TextureArray_##Index.SampleLevel({DataInterfaceName}_Sampler, float3(InTextureUV, SliceIndex), 0); \
 		} \
 		else \
 		{ \
@@ -192,8 +175,7 @@ float4 Sample_{DataInterfaceName}(uint InDataIndex, float2 InTextureUV)
 	uint BindingIndex = 0;
 	uint SliceIndex = 0;
 	uint TextureDimension = 0;
-	uint PointSample = 0;
-	GetTextureInfoInternal_{DataInterfaceName}(InDataIndex, BindingIndex, SliceIndex, TextureDimension, PointSample);
+	GetTextureInfoInternal_{DataInterfaceName}(InDataIndex, BindingIndex, SliceIndex, TextureDimension);
 
 #if PCG_ENABLE_MULTIPLE_TEXTURE_OBJECTS
 	PCG_CONDITIONAL_TEXTURE_SAMPLE(0)
@@ -211,25 +193,11 @@ float4 Sample_{DataInterfaceName}(uint InDataIndex, float2 InTextureUV)
 #else
 	if (TextureDimension == PCG_TextureDimension_Texture2D)
 	{
-		if (PointSample == 0)
-		{
-			return {DataInterfaceName}_Texture_0.SampleLevel({DataInterfaceName}_Sampler_Linear, InTextureUV, 0);
-		}
-		else
-		{
-			return {DataInterfaceName}_Texture_0.SampleLevel({DataInterfaceName}_Sampler_Point, InTextureUV, 0);
-		}
+		return {DataInterfaceName}_Texture_0.SampleLevel({DataInterfaceName}_Sampler, InTextureUV, 0);
 	}
 	else if (TextureDimension == PCG_TextureDimension_Texture2DArray)
 	{
-		if (PointSample == 0)
-		{
-			return {DataInterfaceName}_TextureArray_0.SampleLevel({DataInterfaceName}_Sampler_Linear, float3(InTextureUV, SliceIndex), 0);
-		}
-		else
-		{
-			return {DataInterfaceName}_TextureArray_0.SampleLevel({DataInterfaceName}_Sampler_Point, float3(InTextureUV, SliceIndex), 0);
-		}
+		return {DataInterfaceName}_TextureArray_0.SampleLevel({DataInterfaceName}_Sampler, float3(InTextureUV, SliceIndex), 0);
 	}
 	else
 	{
@@ -246,25 +214,11 @@ float4 Sample_{DataInterfaceName}(uint InDataIndex, float2 InTextureUV)
 		const float2 TextureUV = GetTexCoords_{DataInterfaceName}(Index, WorldPos); \
 		if (TextureDimension == PCG_TextureDimension_Texture2D) \
 		{ \
-			if (PointSample == 0) \
-			{ \
-				return {DataInterfaceName}_Texture_##Index.SampleLevel({DataInterfaceName}_Sampler_Linear, TextureUV, 0); \
-			} \
-			else \
-			{ \
-				return {DataInterfaceName}_Texture_##Index.SampleLevel({DataInterfaceName}_Sampler_Point, TextureUV, 0); \
-			} \
+			return {DataInterfaceName}_Texture_##Index.SampleLevel({DataInterfaceName}_Sampler, TextureUV, 0); \
 		} \
 		else if (TextureDimension == PCG_TextureDimension_Texture2DArray) \
 		{ \
-			if (PointSample == 0) \
-			{ \
-				return {DataInterfaceName}_TextureArray_##Index.SampleLevel({DataInterfaceName}_Sampler_Linear, float3(TextureUV, SliceIndex), 0); \
-			} \
-			else \
-			{ \
-				return {DataInterfaceName}_TextureArray_##Index.SampleLevel({DataInterfaceName}_Sampler_Point, float3(TextureUV, SliceIndex), 0); \
-			} \
+			return {DataInterfaceName}_TextureArray_##Index.SampleLevel({DataInterfaceName}_Sampler, float3(TextureUV, SliceIndex), 0); \
 		} \
 		else \
 		{ \
@@ -278,8 +232,7 @@ float4 SampleWorldPos_{DataInterfaceName}(uint InDataIndex, float3 WorldPos)
 	uint BindingIndex = 0;
 	uint SliceIndex = 0;
 	uint TextureDimension = 0;
-	uint PointSample = 0;
-	GetTextureInfoInternal_{DataInterfaceName}(InDataIndex, BindingIndex, SliceIndex, TextureDimension, PointSample);
+	GetTextureInfoInternal_{DataInterfaceName}(InDataIndex, BindingIndex, SliceIndex, TextureDimension);
 
 #if PCG_ENABLE_MULTIPLE_TEXTURE_OBJECTS
 	PCG_CONDITIONAL_TEXTURE_SAMPLE_WORLD_POS(0)
@@ -299,25 +252,11 @@ float4 SampleWorldPos_{DataInterfaceName}(uint InDataIndex, float3 WorldPos)
 
 	if (TextureDimension == PCG_TextureDimension_Texture2D)
 	{
-		if (PointSample == 0)
-		{
-			return {DataInterfaceName}_Texture_0.SampleLevel({DataInterfaceName}_Sampler_Linear, TextureUV, 0);
-		}
-		else
-		{
-			return {DataInterfaceName}_Texture_0.SampleLevel({DataInterfaceName}_Sampler_Point, TextureUV, 0);
-		}
+		return {DataInterfaceName}_Texture_0.SampleLevel({DataInterfaceName}_Sampler, TextureUV, 0);
 	}
 	else if (TextureDimension == PCG_TextureDimension_Texture2DArray)
 	{
-		if (PointSample == 0)
-		{
-			return {DataInterfaceName}_TextureArray_0.SampleLevel({DataInterfaceName}_Sampler_Linear, float3(TextureUV, SliceIndex), 0);
-		}
-		else
-		{
-			return {DataInterfaceName}_TextureArray_0.SampleLevel({DataInterfaceName}_Sampler_Point, float3(TextureUV, SliceIndex), 0);
-		}
+		return {DataInterfaceName}_TextureArray_0.SampleLevel({DataInterfaceName}_Sampler, float3(TextureUV, SliceIndex), 0);
 	}
 	else
 	{
@@ -351,8 +290,7 @@ float4 Load_{DataInterfaceName}(uint InDataIndex, uint2 InElementIndex)
 	uint BindingIndex = 0;
 	uint SliceIndex = 0;
 	uint TextureDimension = 0;
-	uint PointSample = 0;
-	GetTextureInfoInternal_{DataInterfaceName}(InDataIndex, BindingIndex, SliceIndex, TextureDimension, PointSample);
+	GetTextureInfoInternal_{DataInterfaceName}(InDataIndex, BindingIndex, SliceIndex, TextureDimension);
 
 #if PCG_ENABLE_MULTIPLE_TEXTURE_OBJECTS
 	PCG_CONDITIONAL_TEXTURE_LOAD(0)
@@ -399,8 +337,7 @@ void Store_{DataInterfaceName}(uint InDataIndex, uint2 InElementIndex, float4 In
 	uint BindingIndex = 0;
 	uint SliceIndex = 0;
 	uint TextureDimension = 0;
-	uint PointSample = 0;
-	GetTextureInfoInternal_{DataInterfaceName}(InDataIndex, BindingIndex, SliceIndex, TextureDimension, PointSample);
+	GetTextureInfoInternal_{DataInterfaceName}(InDataIndex, BindingIndex, SliceIndex, TextureDimension);
 
 	if (TextureDimension != PCG_TextureDimension_Texture2D)
 	{
diff --git a/Engine/Plugins/PCG/Source/PCG/Private/Compute/DataInterfaces/PCGTextureDataInterface.cpp b/Engine/Plugins/PCG/Source/PCG/Private/Compute/DataInterfaces/PCGTextureDataInterface.cpp
--- a/Engine/Plugins/PCG/Source/PCG/Private/Compute/DataInterfaces/PCGTextureDataInterface.cpp
+++ b/Engine/Plugins/PCG/Source/PCG/Private/Compute/DataInterfaces/PCGTextureDataInterface.cpp
@@ -193,8 +193,7 @@ BEGIN_SHADER_PARAMETER_STRUCT(FPCGTextureDataInterfaceParameters,)
 	SHADER_PARAMETER_RDG_TEXTURE_SRV_ARRAY(Texture2D, Texture, [PCGTextureDataInterfaceConstants::MAX_NUM_SRV_BINDINGS])
 	SHADER_PARAMETER_RDG_TEXTURE_SRV_ARRAY(Texture2DArray, TextureArray, [PCGTextureDataInterfaceConstants::MAX_NUM_SRV_BINDINGS])
 	SHADER_PARAMETER_RDG_TEXTURE_UAV_ARRAY(RWTexture2D<float4>, TextureUAV, [PCGTextureDataInterfaceConstants::MAX_NUM_UAV_BINDINGS])
-	SHADER_PARAMETER_SAMPLER(SamplerState, Sampler_Linear)
-	SHADER_PARAMETER_SAMPLER(SamplerState, Sampler_Point)
+	SHADER_PARAMETER_SAMPLER(SamplerState, Sampler)
 	SHADER_PARAMETER_ARRAY(FMatrix44f, TextureInverseTransforms, [PCGTextureDataInterfaceConstants::MAX_NUM_SRV_BINDINGS])
 	SHADER_PARAMETER(uint32, NumTextureInfos)
 	SHADER_PARAMETER_RDG_BUFFER_SRV(StructuredBuffer<FUintVector4>, TextureInfos)
@@ -290,6 +289,20 @@ bool UPCGTextureDataProvider::PrepareForExecute_GameThread(UPCGDataBinding* InBi
 		BuildInfosFromDataDescription(InBinding);
 	}
 
+	// Validate that all binding infos have the same sampler - otherwise warn that only the first will be used.
+	if (!BindingInfos.IsEmpty())
+	{
+		const EPCGTextureFilter Filter = BindingInfos[0].Filter;
+
+		for (int BindingIndex = 1; BindingIndex < BindingInfos.Num(); ++BindingIndex)
+		{
+			if (BindingInfos[BindingIndex].Filter != Filter)
+			{
+				UE_LOG(LogPCG, Warning, TEXT("Encountered multiple unique texture filters on texture data interface: '%s'. Only the first will be used."), *GetOutputPinLabel().ToString());
+			}
+		}
+	}
+
 	return true;
 }
 
@@ -447,8 +460,9 @@ void FPCGTextureDataProviderProxy::GatherDispatchData(FDispatchData const& InDis
 	{
 		FParameters& Parameters = ParameterArray[InvocationIndex];
 
-		Parameters.Sampler_Linear = TStaticSamplerState<SF_Bilinear, AM_Clamp, AM_Clamp, AM_Clamp>::GetRHI();
-		Parameters.Sampler_Point = TStaticSamplerState<SF_Point, AM_Clamp, AM_Clamp, AM_Clamp>::GetRHI();
+		Parameters.Sampler = (!BindingInfos.IsEmpty() && BindingInfos[0].Filter == EPCGTextureFilter::Point)
+			? TStaticSamplerState<SF_Point, AM_Clamp, AM_Clamp, AM_Clamp>::GetRHI()
+			: TStaticSamplerState<SF_Bilinear, AM_Clamp, AM_Clamp, AM_Clamp>::GetRHI();
 		Parameters.NumTextureInfos = TextureInfos.Num();
 		Parameters.TextureInfos = TextureInfosBufferSRV;
 
@@ -665,7 +679,7 @@ void FPCGTextureDataProviderProxy::PackTextureInfos(FRDGBuilder& GraphBuilder)
 	//   - uint32 - BindingIndex
 	//   - uint32 - SliceIndex (aka TextureIndex if bound to a TextureArray, 0 otherwise.)
 	//   - uint32 - Dimension
-	//   - uint32 - Sampler
+	//   - uint32 - UNUSED
 	//   - uint32 - ResolutionX
 	//   - uint32 - ResolutionY
 	//   - float  - TexelSizeWorldX (cm)
@@ -692,7 +706,7 @@ void FPCGTextureDataProviderProxy::PackTextureInfos(FRDGBuilder& GraphBuilder)
 			const float TexelWorldSizeX = BindingInfo.Size.X != 0 ? (TextureWorldScale.X * 2.0) / BindingInfo.Size.X : 0.0f;
 			const float TexelWorldSizeY = BindingInfo.Size.Y != 0 ? (TextureWorldScale.Y * 2.0) / BindingInfo.Size.Y : 0.0f;
 
-			TextureInfosPackedData[TextureInfoIndex * 2 + 0] = FUintVector4(TextureInfo.BindingIndex, TextureInfo.SliceIndex, static_cast<uint32>(BindingInfo.Dimension), BindingInfo.Filter == EPCGTextureFilter::Point);
+			TextureInfosPackedData[TextureInfoIndex * 2 + 0] = FUintVector4(TextureInfo.BindingIndex, TextureInfo.SliceIndex, static_cast<uint32>(BindingInfo.Dimension), /*UNUSEd*/0u);
 			TextureInfosPackedData[TextureInfoIndex * 2 + 1] = FUintVector4(BindingInfo.Size.X, BindingInfo.Size.Y, FMath::AsUInt(TexelWorldSizeX), FMath::AsUInt(TexelWorldSizeY));
 		}
 
