[PCG] Fix component not cleaning up when bActivated disabled on BP PCG actor.

#rnx
#Synced-CL 48649347

diff --git a/Engine/Plugins/PCG/Source/PCG/Private/RuntimeGen/PCGRuntimeGenScheduler.cpp b/Engine/Plugins/PCG/Source/PCG/Private/RuntimeGen/PCGRuntimeGenScheduler.cpp
--- a/Engine/Plugins/PCG/Source/PCG/Private/RuntimeGen/PCGRuntimeGenScheduler.cpp
+++ b/Engine/Plugins/PCG/Source/PCG/Private/RuntimeGen/PCGRuntimeGenScheduler.cpp
@@ -1500,7 +1500,7 @@ void FPCGRuntimeGenScheduler::OnObjectsReplaced(const TMap<UObject*, UObject*>&
 		{
 			for (FGridGenerationKey& GenerationKey : GeneratedComponents)
 			{
-				if (GenerationKey.GetOriginalComponent() == OldComponent)
+				if (GenerationKey.GetOriginalComponentEvenIfGarbage() == OldComponent)
 				{
 					GenerationKey.SetOriginalComponent(NewComponent);
 				}
diff --git a/Engine/Plugins/PCG/Source/PCG/Public/RuntimeGen/PCGRuntimeGenScheduler.h b/Engine/Plugins/PCG/Source/PCG/Public/RuntimeGen/PCGRuntimeGenScheduler.h
--- a/Engine/Plugins/PCG/Source/PCG/Public/RuntimeGen/PCGRuntimeGenScheduler.h
+++ b/Engine/Plugins/PCG/Source/PCG/Public/RuntimeGen/PCGRuntimeGenScheduler.h
@@ -114,6 +114,7 @@ protected:
 		FIntVector GetGridCoords() const { return GridCoords; }
 
 		UPCGComponent* GetOriginalComponent() const { return OriginalComponent.ResolveObjectPtr(); }
+		UPCGComponent* GetOriginalComponentEvenIfGarbage() const { return OriginalComponent.ResolveObjectPtrEvenIfGarbage(); }
 		void SetOriginalComponent(UPCGComponent* InComponent) { OriginalComponent = InComponent; }
 
 		IPCGGraphExecutionSource* GetCachedLocalComponent() const { return CachedLocalComponent.Get(); }
