[PCG] Add safeguards when writing Objects/Enums in properties

* There is no guarantee the soft path is of the right class (since we do SoftObjectPath->TObjectPtr)
* Same for enums, no guarantee that the value is a valid value for the enum (since we do int64 -> Enum)

#rnx
#Synced-CL 48550537

diff --git a/Engine/Plugins/PCG/Source/PCG/Private/Metadata/Accessors/PCGPropertyAccessor.cpp b/Engine/Plugins/PCG/Source/PCG/Private/Metadata/Accessors/PCGPropertyAccessor.cpp
--- a/Engine/Plugins/PCG/Source/PCG/Private/Metadata/Accessors/PCGPropertyAccessor.cpp
+++ b/Engine/Plugins/PCG/Source/PCG/Private/Metadata/Accessors/PCGPropertyAccessor.cpp
@@ -26,18 +26,43 @@ const UStruct* IPCGPropertyChain::GetTopPropertyStruct() const
 	}
 }
 
+FPCGEnumPropertyAccessor::FPCGEnumPropertyAccessor(const FEnumProperty* InProperty, TArray<const FProperty*>&& ExtraProperties)
+	: Super(/*bInReadOnly=*/ false)
+	, IPCGPropertyChain(InProperty, std::move(ExtraProperties))
+	, Property(InProperty)
+{
+	check(Property);
+
+	const UEnum* Enum = Property->GetEnum();
+	bHasMaxValue = Enum && Enum->ContainsExistingMax();
+	EnumMaxValue = bHasMaxValue ? Enum->GetMaxEnumValue() : -1;
+}
+
+
 bool FPCGEnumPropertyAccessor::GetRangeImpl(TArrayView<int64> OutValues, int32 Index, const IPCGAttributeAccessorKeys& Keys) const
 {
 	return PCGPropertyAccessor::IterateGet(GetPropertyChain(), OutValues, Index, Keys, [this](const void* PropertyAddressData) -> Type
-		{
-			return Property->GetUnderlyingProperty()->GetSignedIntPropertyValue(PropertyAddressData);
-		});
+	{
+		return Property->GetUnderlyingProperty()->GetSignedIntPropertyValue(PropertyAddressData);
+	});
 }
 
 bool FPCGEnumPropertyAccessor::SetRangeImpl(TArrayView<const int64> InValues, int32 Index, IPCGAttributeAccessorKeys& Keys, EPCGAttributeAccessorFlags)
 {
-	return PCGPropertyAccessor::IterateSet(GetPropertyChain(), InValues, Index, Keys, [this](void* PropertyAddressData, const int64& Value) -> void
+	bool bWriteOK = true;
+	const UEnum* Enum = Property->GetEnum();
+
+	const bool bIterateOK = PCGPropertyAccessor::IterateSet(GetPropertyChain(), InValues, Index, Keys, [this, &bWriteOK, Enum](void* PropertyAddressData, const int64& Value) -> void
+	{
+		if (Enum && (!Enum->IsValidEnumValue(Value) || (bHasMaxValue && Value == EnumMaxValue)))
 		{
-			Property->GetUnderlyingProperty()->SetIntPropertyValue(PropertyAddressData, Value);
-		});
+			UE_LOG(LogPCG, Error, TEXT("[FPCGEnumPropertyAccessor::SetRange] Try to write into an Enum \"%s\" a value \"%lld\" that is not a valid value for this enum."), *Enum->GetName(), Value);
+			bWriteOK = false;
+			return;
+		}
+		
+		Property->GetUnderlyingProperty()->SetIntPropertyValue(PropertyAddressData, Value);
+	});
+
+	return bIterateOK && bWriteOK;
 }
\ No newline at end of file
diff --git a/Engine/Plugins/PCG/Source/PCG/Public/Metadata/Accessors/PCGPropertyAccessor.h b/Engine/Plugins/PCG/Source/PCG/Public/Metadata/Accessors/PCGPropertyAccessor.h
--- a/Engine/Plugins/PCG/Source/PCG/Public/Metadata/Accessors/PCGPropertyAccessor.h
+++ b/Engine/Plugins/PCG/Source/PCG/Public/Metadata/Accessors/PCGPropertyAccessor.h
@@ -4,6 +4,7 @@
 
 #include "IPCGAttributeAccessorTpl.h"
 
+#include "PCGModule.h"
 #include "Metadata/PCGAttributePropertySelector.h"
 
 #include "Containers/UnrealString.h" // IWYU pragma: keep
@@ -186,6 +187,7 @@ private:
 * Templated accessor class for enum properties. Will wrap around an enum property.
 * Will always convert to int64 for PCG
 * Do not instanciate it manually, use PCGAttributeAccessorHelpers::CreatePropertyAccessor.
+* Note that setting the property value will be slower, because we need to make sure that the value is a valid value for the enum.
 * Key supported: Generic object
 */
 class FPCGEnumPropertyAccessor : public IPCGAttributeAccessorT<FPCGEnumPropertyAccessor>, IPCGPropertyChain
@@ -194,19 +196,15 @@ public:
 	using Type = int64;
 	using Super = IPCGAttributeAccessorT<FPCGEnumPropertyAccessor>;
 
-	FPCGEnumPropertyAccessor(const FEnumProperty* InProperty, TArray<const FProperty*>&& ExtraProperties = {})
-		: Super(/*bInReadOnly=*/ false)
-		, IPCGPropertyChain(InProperty, std::move(ExtraProperties))
-		, Property(InProperty)
-	{
-		check(Property);
-	}
+	FPCGEnumPropertyAccessor(const FEnumProperty* InProperty, TArray<const FProperty*>&& ExtraProperties = {});
 
 	bool GetRangeImpl(TArrayView<int64> OutValues, int32 Index, const IPCGAttributeAccessorKeys& Keys) const;
 	bool SetRangeImpl(TArrayView<const int64> InValues, int32 Index, IPCGAttributeAccessorKeys& Keys, EPCGAttributeAccessorFlags);
 
 private:
 	const FEnumProperty* Property = nullptr;
+	bool bHasMaxValue = false;
+	int64 EnumMaxValue = -1;
 };
 
 /**
@@ -418,9 +416,11 @@ private:
 * Templated accessor class for object/class ptr properties.
 * Do not instantiate it manually, use PCGAttributeAccessorHelpers::CreatePropertyAccessor.
 * Will always convert to FSoftObjectPath/FSoftClassPath for PCG
+* Note that setting the property value will be slower, because we need to make sure that the Soft Object/Class Path is of the right type.
+* It will also force the object to load, as ObjectPtrs are expected to be loaded.
 * Key supported: Generic object
 */
-template <typename PropertyType>
+template <typename PropertyType> requires std::is_base_of_v<FObjectProperty, PropertyType>
 class FPCGPropertyObjectPtrAccessor : public IPCGAttributeAccessorT<FPCGPropertyObjectPtrAccessor<PropertyType>>, IPCGPropertyChain
 {
 public:
@@ -432,8 +432,6 @@ public:
 		, IPCGPropertyChain(InProperty, std::move(ExtraProperties))
 		, Property(InProperty)
 	{
-		// Making sure it is the right properties.
-		static_assert(std::is_same_v<PropertyType, FClassProperty> || std::is_same_v<PropertyType, FObjectProperty>);
 		check(Property);
 	}
 
@@ -447,10 +445,33 @@ public:
 
 	bool SetRangeImpl(TArrayView<const Type> InValues, int32 Index, IPCGAttributeAccessorKeys& Keys, EPCGAttributeAccessorFlags)
 	{
-		return PCGPropertyAccessor::IterateSet(GetPropertyChain(), InValues, Index, Keys, [this](void* PropertyAddressData, const Type& Value) -> void
+		bool bWriteOK = true;
+
+		const bool bIterateOK = PCGPropertyAccessor::IterateSet(GetPropertyChain(), InValues, Index, Keys, [this, &bWriteOK](void* PropertyAddressData, const Type& Value) -> void
 		{
-			Property->SetPropertyValue(PropertyAddressData, Value.TryLoad());
+			UObject* Object = Value.TryLoad();
+			if (Object && Property->PropertyClass && !Object->IsA(Property->PropertyClass))
+			{
+				UE_LOG(LogPCG, Error, TEXT("[FPCGPropertyObjectPtrAccessor::SetRange] Try to write into an Object property of class \"%s\" with an object of class \"%s\" (%s). This is invalid."), *Property->PropertyClass->GetName(), *Object->GetClass()->GetName(), *Value.ToString());
+				bWriteOK = false;
+				return;
+			}
+
+			if constexpr (std::is_base_of_v<FClassProperty, PropertyType>)
+			{
+				const UClass* Class = Cast<const UClass>(Object);
+				if (Class && Property->MetaClass && !Class->IsChildOf(Property->MetaClass))
+				{
+					UE_LOG(LogPCG, Error, TEXT("[FPCGPropertyObjectPtrAccessor::SetRange] Try to write into an Class property of metaclass \"%s\" with a class of metaclass \"%s\" (%s). This is invalid."), *Property->MetaClass->GetName(), *Class->GetName(), *Value.ToString());
+					bWriteOK = false;
+					return;
+				}
+			}
+
+			Property->SetPropertyValue(PropertyAddressData, Object);
 		});
+
+		return bIterateOK && bWriteOK;
 	}
 
 private:
