[PCG] Apply tags to GPU proxy data.

Previously the tags for GPU proxy data were stored in the GPU data description. Moving them to the normal tags means that tags can be manipulated on the proxies without triggering readbacks to CPU.

Fix ALV failing to display data from GPU proxies (showing empty ALV and 'GPU Proxy' in data dropdown).

Enable GPU-resident data on Filter Data By Index.

[RN] [PCG] Enable tag operations on GPU data without readback to CPU.
#Synced-CL 48603294

diff --git a/Engine/Plugins/PCG/Source/PCG/Private/Compute/DataInterfaces/PCGDataCollectionDataInterface.cpp b/Engine/Plugins/PCG/Source/PCG/Private/Compute/DataInterfaces/PCGDataCollectionDataInterface.cpp
--- a/Engine/Plugins/PCG/Source/PCG/Private/Compute/DataInterfaces/PCGDataCollectionDataInterface.cpp
+++ b/Engine/Plugins/PCG/Source/PCG/Private/Compute/DataInterfaces/PCGDataCollectionDataInterface.cpp
@@ -831,7 +831,7 @@ void FPCGDataCollectionDataProviderProxy::AllocateResources(FRDGBuilder& GraphBu
 						Proxy->Initialize(DataCollectionOnGPU, DataIndex);
 
 						// TODO - binding is doing a lot of work. Could store a context handle in the data provider instead?
-						Binding->ReceiveDataFromGPU_GameThread(Proxy, DataProvider->GetProducerSettings(), ExportMode, OutputPinLabel, OutputPinLabelAlias);
+						Binding->ReceiveDataFromGPU_GameThread(Proxy, DataProvider->GetProducerSettings(), ExportMode, OutputPinLabel, OutputPinLabelAlias, PinDesc->GetDataDescriptions()[DataIndex]);
 					}
 
 					DataProvider->OnDataExported_GameThread().Broadcast();
diff --git a/Engine/Plugins/PCG/Source/PCG/Private/Compute/DataInterfaces/PCGRawBufferDataInterface.cpp b/Engine/Plugins/PCG/Source/PCG/Private/Compute/DataInterfaces/PCGRawBufferDataInterface.cpp
--- a/Engine/Plugins/PCG/Source/PCG/Private/Compute/DataInterfaces/PCGRawBufferDataInterface.cpp
+++ b/Engine/Plugins/PCG/Source/PCG/Private/Compute/DataInterfaces/PCGRawBufferDataInterface.cpp
@@ -189,16 +189,19 @@ FComputeDataProviderRenderProxy* UPCGRawBufferDataProvider::GetRenderProxy()
 
 	ensure(ReadbackState);
 
-	Params.AsyncReadbackCallback_RenderThread = [ReadbackState=ReadbackState, DataProviderWeakPtr=ThisWeak, GenerationCount=GetGenerationCounter(), ExportMode=GetExportMode(), OutputPinLabel = GetOutputPinLabel(), OutputPinLabelAlias = GetOutputPinLabelAlias(), BindingWeakPtr = TWeakObjectPtr<UPCGDataBinding>(GetDataBinding())](const void* InData, int InNumBytes)
+	TSharedPtr<const FPCGDataCollectionDesc> PinDesc = GetPinDescription();
+	ensure(PinDesc);
+
+	Params.AsyncReadbackCallback_RenderThread = [ReadbackState=ReadbackState, DataProviderWeakPtr=ThisWeak, GenerationCount=GetGenerationCounter(), ExportMode=GetExportMode(), OutputPinLabel=GetOutputPinLabel(), OutputPinLabelAlias = GetOutputPinLabelAlias(), BindingWeakPtr=TWeakObjectPtr<UPCGDataBinding>(GetDataBinding()), PinDesc](const void* InData, int InNumBytes)
 	{
-		if (!InData || !ReadbackState)
+		if (!InData || !ReadbackState || !PinDesc || PinDesc->GetDataDescriptions().IsEmpty())
 		{
 			return;
 		}
 
 		ReadbackState->Data.Append(static_cast<const uint32*>(InData), InNumBytes / sizeof(uint32));
 
-		ExecuteOnGameThread(UE_SOURCE_LOCATION, [ReadbackState, DataProviderWeakPtr, GenerationCount, ExportMode, OutputPinLabel, OutputPinLabelAlias, BindingWeakPtr]()
+		ExecuteOnGameThread(UE_SOURCE_LOCATION, [ReadbackState, DataProviderWeakPtr, GenerationCount, ExportMode, OutputPinLabel, OutputPinLabelAlias, BindingWeakPtr, PinDesc]()
 		{
 			UPCGDataBinding* Binding = BindingWeakPtr.Get();
 			if (!Binding)
@@ -222,7 +225,7 @@ FComputeDataProviderRenderProxy* UPCGRawBufferDataProvider::GetRenderProxy()
 			UPCGRawBufferData* RawBufferData = NewObject<UPCGRawBufferData>();
 			RawBufferData->SetData(MoveTemp(ReadbackState->Data));
 
-			Binding->ReceiveDataFromGPU_GameThread(RawBufferData, DataProvider->GetProducerSettings(), ExportMode, OutputPinLabel, OutputPinLabelAlias);
+			Binding->ReceiveDataFromGPU_GameThread(RawBufferData, DataProvider->GetProducerSettings(), ExportMode, OutputPinLabel, OutputPinLabelAlias, PinDesc->GetDataDescriptions()[0]);
 
 			DataProvider->OnDataExported_GameThread().Broadcast();
 		});
diff --git a/Engine/Plugins/PCG/Source/PCG/Private/Compute/DataInterfaces/PCGTextureDataInterface.cpp b/Engine/Plugins/PCG/Source/PCG/Private/Compute/DataInterfaces/PCGTextureDataInterface.cpp
--- a/Engine/Plugins/PCG/Source/PCG/Private/Compute/DataInterfaces/PCGTextureDataInterface.cpp
+++ b/Engine/Plugins/PCG/Source/PCG/Private/Compute/DataInterfaces/PCGTextureDataInterface.cpp
@@ -654,17 +654,8 @@ void FPCGTextureDataProviderProxy::ExportTextureUAVs(const TArray<TRefCountPtr<I
 				ExportedData->TexelSize = TexelSize.GetMin();
 				ExportedData->Initialize(ExportedTextures[TextureDataIndex], TextureInfo.SliceIndex, TextureTransform, /*bInSkipReadbackToCPU=*/true);
 
-				TSet<FString> Tags;
-				for (const int32 TagStringKey : PinDesc->GetDataDescriptions()[TextureDataIndex].GetTagStringKeys())
-				{
-					if (ensure(StringTable.IsValidIndex(TagStringKey)))
-					{
-						Tags.Add(StringTable[TagStringKey]);
-					}
-				}
-
 				// @todo_pcg: Binding is doing a lot of work. Could store a context handle in the data provider instead?
-				Binding->ReceiveDataFromGPU_GameThread(ExportedData, DataProvider->GetProducerSettings(), ExportMode, OutputPinLabel, OutputPinLabelAlias, Tags);
+				Binding->ReceiveDataFromGPU_GameThread(ExportedData, DataProvider->GetProducerSettings(), ExportMode, OutputPinLabel, OutputPinLabelAlias, PinDesc->GetDataDescriptions()[TextureDataIndex]);
 			}
 		}
 
diff --git a/Engine/Plugins/PCG/Source/PCG/Private/Compute/Elements/PCGComputeGraphElement.cpp b/Engine/Plugins/PCG/Source/PCG/Private/Compute/Elements/PCGComputeGraphElement.cpp
--- a/Engine/Plugins/PCG/Source/PCG/Private/Compute/Elements/PCGComputeGraphElement.cpp
+++ b/Engine/Plugins/PCG/Source/PCG/Private/Compute/Elements/PCGComputeGraphElement.cpp
@@ -852,7 +852,7 @@ void FPCGComputeGraphElement::StoreDataForInspection(FPCGComputeGraphContext* In
 				UPCGProxyForGPUData::FReadbackResult Readback = ProxyData->GetCPUData(InContext);
 				FPCGTaggedData TaggedData = Readback.TaggedData;
 				TaggedData.Pin = DataToInspect.PinLabel;
-				TaggedData.Tags.Append(DataToInspect.AdditionalTags);
+				TaggedData.Tags.Append(DataToInspect.Tags);
 				DataCollection.TaggedData.Add(TaggedData);
 			}
 			else if (const UPCGTextureData* TextureData = Cast<UPCGTextureData>(DataToInspect.Data.Get()))
@@ -862,7 +862,7 @@ void FPCGComputeGraphElement::StoreDataForInspection(FPCGComputeGraphContext* In
 				FPCGTaggedData TaggedData;
 				TaggedData.Data = bReadbackTextureDataOnInspect ? DataToInspect.DataPendingInit : DataToInspect.Data;
 				TaggedData.Pin = DataToInspect.PinLabel;
-				TaggedData.Tags = DataToInspect.AdditionalTags;
+				TaggedData.Tags = DataToInspect.Tags;
 				DataCollection.TaggedData.Add(TaggedData);
 			}
 		}
diff --git a/Engine/Plugins/PCG/Source/PCG/Private/Compute/PCGComputeGraph.cpp b/Engine/Plugins/PCG/Source/PCG/Private/Compute/PCGComputeGraph.cpp
--- a/Engine/Plugins/PCG/Source/PCG/Private/Compute/PCGComputeGraph.cpp
+++ b/Engine/Plugins/PCG/Source/PCG/Private/Compute/PCGComputeGraph.cpp
@@ -453,14 +453,14 @@ bool UPCGComputeGraph::ComputeTaggedDataPinDesc(const FPCGTaggedData& InTaggedDa
 					}
 				}
 			}
+		}
 
-			// Remap tag string keys.
-			for (int32& TagStringKey : OutDescription.GetTagStringKeysMutable())
+		for (const FString& Tag : InTaggedData.Tags)
+		{
+			const int32 Key = InBinding->GetStringTable().IndexOfByKey(Tag);
+			if (ensure(Key != INDEX_NONE))
 			{
-				if (const int32* GraphStringKey = InBinding->GetSourceBufferToGraphStringKey().Find({ DataGPU->GetBuffer(), TagStringKey }))
-				{
-					TagStringKey = *GraphStringKey;
-				}
+				OutDescription.AddUniqueTagStringKey(Key);
 			}
 		}
 	}
diff --git a/Engine/Plugins/PCG/Source/PCG/Private/Compute/PCGDataBinding.cpp b/Engine/Plugins/PCG/Source/PCG/Private/Compute/PCGDataBinding.cpp
--- a/Engine/Plugins/PCG/Source/PCG/Private/Compute/PCGDataBinding.cpp
+++ b/Engine/Plugins/PCG/Source/PCG/Private/Compute/PCGDataBinding.cpp
@@ -204,43 +204,29 @@ TSharedPtr<const FPCGDataCollectionDesc> UPCGDataBinding::GetCachedKernelPinData
 	return *FoundDescription;
 }
 
-void UPCGDataBinding::ReceiveDataFromGPU_GameThread(
-	UPCGData* InData,
-	const UPCGSettings* InProducerSettings,
-	EPCGExportMode InExportMode,
-	FName InPinLabel,
-	FName InPinLabelAlias)
-{
-	ReceiveDataFromGPU_GameThread(InData, InProducerSettings, InExportMode, InPinLabel, InPinLabelAlias, /*InAdditionalTags=*/{});
-}
-
 void UPCGDataBinding::ReceiveDataFromGPU_GameThread(
 	UPCGData* InData,
 	const UPCGSettings* InProducerSettings,
 	EPCGExportMode InExportMode,
 	FName InPinLabel,
 	FName InPinLabelAlias,
-	const TSet<FString>& InAdditionalTags)
+	const FPCGDataDesc& InDataDesc)
 {
 	check(IsInGameThread());
 
-	if (!!(InExportMode & EPCGExportMode::ComputeGraphOutput))
+	TSet<FString> Tags;
+	for (int32 TagStringKey : InDataDesc.GetTagStringKeys())
 	{
-		FPCGTaggedData& TaggedData = OutputDataCollection.TaggedData.Emplace_GetRef();
-		TaggedData.Pin = InPinLabelAlias;
-		TaggedData.Data = InData;
-		TaggedData.Tags = InAdditionalTags;
-	}
-
-	if (!!(InExportMode & EPCGExportMode::DebugVisualization))
-	{
-		DataToDebug.Add(FPCGDataToDebug{ .Data = InData, .ProducerSettings = InProducerSettings, .PinLabel = InPinLabel, .PinLabelAlias = InPinLabelAlias, .AdditionalTags = InAdditionalTags });
+		if (GetStringTable().IsValidIndex(TagStringKey) && TagStringKey > 0)
+		{
+			Tags.Add(GetStringTable()[TagStringKey]);
+		}
 	}
 
-	if (!!(InExportMode & EPCGExportMode::Inspection))
-	{
-		DataToInspect.Add(FPCGDataToDebug{ .Data = InData, .ProducerSettings = InProducerSettings, .PinLabel = InPinLabel, .PinLabelAlias = InPinLabelAlias, .AdditionalTags = InAdditionalTags });
-	}
+	// When deprecated ReceiveDataFromGPU_GameThread functions removed, inline here.
+	PRAGMA_DISABLE_DEPRECATION_WARNINGS
+	ReceiveDataFromGPU_GameThread(InData, InProducerSettings, InExportMode, InPinLabel, InPinLabelAlias, Tags);
+	PRAGMA_ENABLE_DEPRECATION_WARNINGS
 }
 
 const TMap<UPCGDataBinding::FSourceBufferAndAttributeIndex, int32>& UPCGDataBinding::GetSourceBufferAttributeToGraphAttributeIndex() const
@@ -619,3 +605,33 @@ void UPCGDataBinding::DebugLogDataDescriptions() const
 	}
 #endif
 }
+
+// Deprecated
+void UPCGDataBinding::ReceiveDataFromGPU_GameThread(UPCGData* InData, const UPCGSettings* InProducerSettings, EPCGExportMode InExportMode, FName InPinLabel, FName InPinLabelAlias)
+{
+	PRAGMA_DISABLE_DEPRECATION_WARNINGS
+	ReceiveDataFromGPU_GameThread(InData, InProducerSettings, InExportMode, InPinLabel, InPinLabelAlias, /*AdditionalTags*/TSet<FString>());
+	PRAGMA_ENABLE_DEPRECATION_WARNINGS
+}
+
+// Deprecated
+void UPCGDataBinding::ReceiveDataFromGPU_GameThread(UPCGData* InData, const UPCGSettings* InProducerSettings, EPCGExportMode InExportMode, FName InPinLabel, FName InPinLabelAlias, const TSet<FString>& AdditionalTags)
+{
+	if (!!(InExportMode & EPCGExportMode::ComputeGraphOutput))
+	{
+		FPCGTaggedData& TaggedData = OutputDataCollection.TaggedData.Emplace_GetRef();
+		TaggedData.Pin = InPinLabelAlias;
+		TaggedData.Data = InData;
+		TaggedData.Tags = AdditionalTags;
+	}
+
+	if (!!(InExportMode & EPCGExportMode::DebugVisualization))
+	{
+		DataToDebug.Add(FPCGDataToDebug{ .Data = InData, .ProducerSettings = InProducerSettings, .PinLabel = InPinLabel, .PinLabelAlias = InPinLabelAlias, .Tags = AdditionalTags });
+	}
+
+	if (!!(InExportMode & EPCGExportMode::Inspection))
+	{
+		DataToInspect.Add(FPCGDataToDebug{ .Data = InData, .ProducerSettings = InProducerSettings, .PinLabel = InPinLabel, .PinLabelAlias = InPinLabelAlias, .Tags = AdditionalTags });
+	}
+}
diff --git a/Engine/Plugins/PCG/Source/PCG/Private/Compute/Packing/PCGDataCollectionPacking.cpp b/Engine/Plugins/PCG/Source/PCG/Private/Compute/Packing/PCGDataCollectionPacking.cpp
--- a/Engine/Plugins/PCG/Source/PCG/Private/Compute/Packing/PCGDataCollectionPacking.cpp
+++ b/Engine/Plugins/PCG/Source/PCG/Private/Compute/Packing/PCGDataCollectionPacking.cpp
@@ -983,14 +983,6 @@ namespace PCGDataCollectionPackingHelpers
 				OutTaggedData.Data = OutPointData;
 				OutTaggedData.Pin = InPin;
 	
-				for (const int32 TagStringKey : DataDesc.GetTagStringKeys())
-				{
-					if (ensure(InStringTable.IsValidIndex(TagStringKey)))
-					{
-						OutTaggedData.Tags.Add(InStringTable[TagStringKey]);
-					}
-				}
-	
 				for (const FPCGKernelAttributeDesc& AttributeDesc : AttributeDescs)
 				{
 					if (AttributeDesc.GetAttributeId() >= PCGDataCollectionPackingConstants::NUM_RESERVED_ATTRS)
diff --git a/Engine/Plugins/PCG/Source/PCG/Private/PCGElement.cpp b/Engine/Plugins/PCG/Source/PCG/Private/PCGElement.cpp
--- a/Engine/Plugins/PCG/Source/PCG/Private/PCGElement.cpp
+++ b/Engine/Plugins/PCG/Source/PCG/Private/PCGElement.cpp
@@ -478,7 +478,6 @@ bool IPCGElement::ConvertInputsIfNeeded(FPCGContext* Context) const
 					// Readback done, replace the proxy with the result data.
 					ensure(Result.TaggedData.Data);
 					TaggedData.Data = Result.TaggedData.Data;
-					TaggedData.Tags.Append(Result.TaggedData.Tags);
 
 					// Will update referenced data objects.
 					Context->bInputDataModified = true;
diff --git a/Engine/Plugins/PCG/Source/PCG/Public/Compute/PCGDataBinding.h b/Engine/Plugins/PCG/Source/PCG/Public/Compute/PCGDataBinding.h
--- a/Engine/Plugins/PCG/Source/PCG/Public/Compute/PCGDataBinding.h
+++ b/Engine/Plugins/PCG/Source/PCG/Public/Compute/PCGDataBinding.h
@@ -60,8 +60,10 @@ struct FPCGDataToDebug
 	UPROPERTY()
 	FName PinLabelAlias;
 
-	// @todo_pcg: This is a big hack to support tags on texture data proxies. Should be replaced with a proper abstraction.
-	/** Allow appending additional tags to the CPU tagged data. */
+	UPROPERTY()
+	TSet<FString> Tags;
+
+	UE_DEPRECATED(5.8, "Use Tags field instead")
 	UPROPERTY()
 	TSet<FString> AdditionalTags;
 };
@@ -105,8 +107,7 @@ public:
 	PCG_API TSharedPtr<const FPCGDataCollectionDesc> GetCachedKernelPinDataDesc(const UPCGComputeKernel* InKernel, FName InPinLabel, bool bIsInput) const;
 	TSharedPtr<const FPCGDataCollectionDesc> GetCachedKernelPinDataDesc(int32 InGraphBindingIndex) const;
 
-	PCG_API void ReceiveDataFromGPU_GameThread(UPCGData* InData, const UPCGSettings* InProducerSettings, EPCGExportMode InExportMode, FName InPinLabel, FName InPinLabelAlias);
-	PCG_API void ReceiveDataFromGPU_GameThread(UPCGData* InData, const UPCGSettings* InProducerSettings, EPCGExportMode InExportMode, FName InPinLabel, FName InPinLabelAlias, const TSet<FString>& AdditionalTags);
+	PCG_API void ReceiveDataFromGPU_GameThread(UPCGData* InData, const UPCGSettings* InProducerSettings, EPCGExportMode InExportMode, FName InPinLabel, FName InPinLabelAlias, const FPCGDataDesc& InDataDesc);
 
 	using FSourceBufferAndAttributeIndex = TPair<TSharedPtr<const FPCGProxyForGPUDataCollection>, /*Source Metadata Attribute Index*/int32>;
 	using FGraphAttributeIndexAndSourceBuffer = TPair</*Graph Metadata Attribute Index*/int32, TRefCountPtr<FRDGPooledBuffer>>;
@@ -231,4 +232,11 @@ private:
 
 	/** Cache of overridden settings/values per-kernel, built during compute graph execution. */
 	FPCGKernelParamsCache KernelParamsCache;
+
+	// Deprecated section
+public:
+	UE_DEPRECATED(5.8, "Use the version that takes a data description instead.")
+	PCG_API void ReceiveDataFromGPU_GameThread(UPCGData* InData, const UPCGSettings* InProducerSettings, EPCGExportMode InExportMode, FName InPinLabel, FName InPinLabelAlias);
+	UE_DEPRECATED(5.8, "Use the version that takes a data description instead.")
+	PCG_API void ReceiveDataFromGPU_GameThread(UPCGData* InData, const UPCGSettings* InProducerSettings, EPCGExportMode InExportMode, FName InPinLabel, FName InPinLabelAlias, const TSet<FString>& AdditionalTags);
 };
diff --git a/Engine/Plugins/PCG/Source/PCG/Public/Elements/Data/PCGReplaceDataByTags.h b/Engine/Plugins/PCG/Source/PCG/Public/Elements/Data/PCGReplaceDataByTags.h
--- a/Engine/Plugins/PCG/Source/PCG/Public/Elements/Data/PCGReplaceDataByTags.h
+++ b/Engine/Plugins/PCG/Source/PCG/Public/Elements/Data/PCGReplaceDataByTags.h
@@ -11,7 +11,7 @@ class FPCGReplaceDataByTagElement : public IPCGElement
 protected:
 	virtual bool ExecuteInternal(FPCGContext* Context) const override;
 	virtual bool SupportsBasePointDataInputs(FPCGContext* InContext) const override { return true; }
-	// @todo_pcg: unclear if we support gpu data because of tags
+	virtual bool SupportsGPUResidentData(FPCGContext* InContext) const override { return true; }
 };
 
 /**
diff --git a/Engine/Plugins/PCG/Source/PCG/Public/Elements/PCGAddTag.h b/Engine/Plugins/PCG/Source/PCG/Public/Elements/PCGAddTag.h
--- a/Engine/Plugins/PCG/Source/PCG/Public/Elements/PCGAddTag.h
+++ b/Engine/Plugins/PCG/Source/PCG/Public/Elements/PCGAddTag.h
@@ -60,4 +60,5 @@ protected:
 	virtual bool ExecuteInternal(FPCGContext* Context) const override;
 	virtual EPCGElementExecutionLoopMode ExecutionLoopMode(const UPCGSettings* Settings) const override { return EPCGElementExecutionLoopMode::SinglePrimaryPin; }
 	virtual bool SupportsBasePointDataInputs(FPCGContext* InContext) const override { return true; }
-};
\ No newline at end of file
+	virtual bool SupportsGPUResidentData(FPCGContext* InContext) const override { return true; }
+};
diff --git a/Engine/Plugins/PCG/Source/PCG/Public/Elements/PCGConvertToAttributeSet.h b/Engine/Plugins/PCG/Source/PCG/Public/Elements/PCGConvertToAttributeSet.h
--- a/Engine/Plugins/PCG/Source/PCG/Public/Elements/PCGConvertToAttributeSet.h
+++ b/Engine/Plugins/PCG/Source/PCG/Public/Elements/PCGConvertToAttributeSet.h
@@ -79,4 +79,5 @@ protected:
 	virtual bool ExecuteInternal(FPCGContext* Context) const override;
 	virtual EPCGElementExecutionLoopMode ExecutionLoopMode(const UPCGSettings* Settings) const override { return EPCGElementExecutionLoopMode::SinglePrimaryPin; }
 	virtual bool SupportsBasePointDataInputs(FPCGContext* InContext) const override { return true; }
-};
\ No newline at end of file
+	virtual bool SupportsGPUResidentData(FPCGContext* InContext) const override { return true; }
+};
diff --git a/Engine/Plugins/PCG/Source/PCG/Public/Elements/PCGDeleteTags.h b/Engine/Plugins/PCG/Source/PCG/Public/Elements/PCGDeleteTags.h
--- a/Engine/Plugins/PCG/Source/PCG/Public/Elements/PCGDeleteTags.h
+++ b/Engine/Plugins/PCG/Source/PCG/Public/Elements/PCGDeleteTags.h
@@ -63,4 +63,5 @@ protected:
 	virtual bool ExecuteInternal(FPCGContext* Context) const override;
 	virtual EPCGElementExecutionLoopMode ExecutionLoopMode(const UPCGSettings* Settings) const override { return EPCGElementExecutionLoopMode::SinglePrimaryPin; }
 	virtual bool SupportsBasePointDataInputs(FPCGContext* InContext) const override { return true; }
-};
\ No newline at end of file
+	virtual bool SupportsGPUResidentData(FPCGContext* InContext) const override { return true; }
+};
diff --git a/Engine/Plugins/PCG/Source/PCG/Public/Elements/PCGFilterByIndex.h b/Engine/Plugins/PCG/Source/PCG/Public/Elements/PCGFilterByIndex.h
--- a/Engine/Plugins/PCG/Source/PCG/Public/Elements/PCGFilterByIndex.h
+++ b/Engine/Plugins/PCG/Source/PCG/Public/Elements/PCGFilterByIndex.h
@@ -46,4 +46,5 @@ class FPCGFilterByIndexElement : public IPCGElement
 protected:
 	virtual bool ExecuteInternal(FPCGContext* Context) const override;
 	virtual bool SupportsBasePointDataInputs(FPCGContext* InContext) const override { return true; }
+	virtual bool SupportsGPUResidentData(FPCGContext* InContext) const override { return true; }
 };
diff --git a/Engine/Plugins/PCG/Source/PCG/Public/Elements/PCGGetDataInfo.h b/Engine/Plugins/PCG/Source/PCG/Public/Elements/PCGGetDataInfo.h
--- a/Engine/Plugins/PCG/Source/PCG/Public/Elements/PCGGetDataInfo.h
+++ b/Engine/Plugins/PCG/Source/PCG/Public/Elements/PCGGetDataInfo.h
@@ -38,6 +38,7 @@ protected:
 	virtual bool ExecuteInternal(FPCGContext* Context) const override;
 	virtual EPCGElementExecutionLoopMode ExecutionLoopMode(const UPCGSettings* Settings) const override { return EPCGElementExecutionLoopMode::SinglePrimaryPin; }
 	virtual bool SupportsBasePointDataInputs(FPCGContext* InContext) const override { return true; }
+	virtual bool SupportsGPUResidentData(FPCGContext* InContext) const override { return true; }
 };
 
 /** Extracts the list of attribute on data to an attribute set with one entry per attribute. */
diff --git a/Engine/Plugins/PCG/Source/PCG/Public/Elements/PCGReplaceTags.h b/Engine/Plugins/PCG/Source/PCG/Public/Elements/PCGReplaceTags.h
--- a/Engine/Plugins/PCG/Source/PCG/Public/Elements/PCGReplaceTags.h
+++ b/Engine/Plugins/PCG/Source/PCG/Public/Elements/PCGReplaceTags.h
@@ -56,4 +56,5 @@ protected:
 	virtual bool ExecuteInternal(FPCGContext* Context) const override;
 	virtual EPCGElementExecutionLoopMode ExecutionLoopMode(const UPCGSettings* Settings) const override { return EPCGElementExecutionLoopMode::SinglePrimaryPin; }
 	virtual bool SupportsBasePointDataInputs(FPCGContext* InContext) const override { return true; }
-};
\ No newline at end of file
+	virtual bool SupportsGPUResidentData(FPCGContext* InContext) const override { return true; }
+};
diff --git a/Engine/Plugins/PCG/Source/PCG/Public/Elements/PCGSortTags.h b/Engine/Plugins/PCG/Source/PCG/Public/Elements/PCGSortTags.h
--- a/Engine/Plugins/PCG/Source/PCG/Public/Elements/PCGSortTags.h
+++ b/Engine/Plugins/PCG/Source/PCG/Public/Elements/PCGSortTags.h
@@ -46,4 +46,5 @@ class FPCGSortTagsElement : public IPCGElement
 protected:
 	virtual bool ExecuteInternal(FPCGContext* Context) const override;
 	virtual bool SupportsBasePointDataInputs(FPCGContext* InContext) const override { return true; }
+	virtual bool SupportsGPUResidentData(FPCGContext* InContext) const override { return true; }
 };
diff --git a/Engine/Plugins/PCG/Source/PCGEditor/Private/Widgets/SPCGEditorGraphAttributeListView.cpp b/Engine/Plugins/PCG/Source/PCGEditor/Private/Widgets/SPCGEditorGraphAttributeListView.cpp
--- a/Engine/Plugins/PCG/Source/PCGEditor/Private/Widgets/SPCGEditorGraphAttributeListView.cpp
+++ b/Engine/Plugins/PCG/Source/PCGEditor/Private/Widgets/SPCGEditorGraphAttributeListView.cpp
@@ -645,18 +645,26 @@ void SPCGEditorGraphAttributeListView::OnTick()
 {
 	if (bNeedsRefresh)
 	{
-		bNeedsRefresh = false;
-
 		bool bKeepSelection = !bPCGEditorGraphNodeChanged;
 		bPCGEditorGraphNodeChanged = false;
 
 		bool bSelectionChanged = false;
 		RefreshPinComboBox(bKeepSelection, bSelectionChanged);
 
+		// Read back any GPU data after the pin selection is set.
+		if (!ReadbackGPUData())
+		{
+			return;
+		}
+
 		bKeepSelection &= !bSelectionChanged;
 
-		RefreshDataComboBox(bKeepSelection);
+		RefreshDataComboBox(bKeepSelection || bKeepDataComboBoxSelection);
 		RefreshAttributeList();
+
+		bKeepDataComboBoxSelection = false;
+
+		bNeedsRefresh = false;
 	}
 
 	if (bViewportNeedsRefresh)
@@ -709,6 +717,25 @@ void SPCGEditorGraphAttributeListView::OnTick()
 	}
 }
 
+bool SPCGEditorGraphAttributeListView::ReadbackGPUData() const
+{
+	bool bAllDataReady = true;
+
+	if (const FPCGDataCollection* InspectionData = GetInspectionData())
+	{
+		for (int32 DataIndex = 0; DataIndex < InspectionData->TaggedData.Num(); ++DataIndex)
+		{
+			if (const UPCGProxyForGPUData* Proxy = Cast<UPCGProxyForGPUData>(InspectionData->TaggedData[DataIndex].Data))
+			{
+				// Poll all readbacks to progress them (don't early out if one is not ready).
+				bAllDataReady &= Proxy->GetCPUData(/*InContext=*/nullptr).bComplete;
+			}
+		}
+	}
+
+	return bAllDataReady;
+}
+
 void SPCGEditorGraphAttributeListView::RequestRefresh()
 {
 	bNeedsRefresh = true;
@@ -1017,7 +1044,7 @@ void SPCGEditorGraphAttributeListView::RefreshDataComboBox(bool bKeepSelection)
 		DataComboBoxItemsSelectedIndex = GetSelectedDataIndex();
 	}
 	
-	if(!bKeepSelection)
+	if (!bKeepSelection)
 	{
 		DataComboBoxItemsSelectedIndex = INDEX_NONE;
 	}
@@ -1033,13 +1060,23 @@ void SPCGEditorGraphAttributeListView::RefreshDataComboBox(bool bKeepSelection)
 		return;
 	}
 
-	for(int32 TaggedDataIndex = 0; TaggedDataIndex < InspectionData->TaggedData.Num(); ++TaggedDataIndex)
+	for (int32 TaggedDataIndex = 0; TaggedDataIndex < InspectionData->TaggedData.Num(); ++TaggedDataIndex)
 	{
 		const FPCGTaggedData& TaggedData = InspectionData->TaggedData[TaggedDataIndex];
+		const UPCGData* Data = TaggedData.Data;
+
+		if (const UPCGProxyForGPUData* ProxyData = Cast<UPCGProxyForGPUData>(Data))
+		{
+			UPCGProxyForGPUData::FReadbackResult Readback = ProxyData->GetCPUData(nullptr);
+			if (ensure(Readback.bComplete)) // Readbacks should have been performed at beginning of refresh
+			{
+				Data = Readback.TaggedData.Data;
+			}
+		}
 
 		FString ItemName = FString::Format(
 			TEXT("[{0}] {1}"),
-			{ FText::AsNumber(TaggedDataIndex).ToString(), (TaggedData.Data ? TaggedData.Data->GetClass()->GetDisplayNameText().ToString(): TEXT("No Data")) });
+			{ FText::AsNumber(TaggedDataIndex).ToString(), (Data ? Data->GetClass()->GetDisplayNameText().ToString(): TEXT("No Data")) });
 
 		if (!TaggedData.Tags.IsEmpty())
 		{
@@ -1102,6 +1139,15 @@ void SPCGEditorGraphAttributeListView::RefreshDomainComboBox(bool bKeepSelection
 	const FPCGTaggedData& TaggedData = InspectionData->TaggedData[DataIndex];
 	const UPCGData* PCGData = TaggedData.Data;
 
+	if (const UPCGProxyForGPUData* ProxyData = Cast<UPCGProxyForGPUData>(PCGData))
+	{
+		UPCGProxyForGPUData::FReadbackResult Readback = ProxyData->GetCPUData(nullptr);
+		if (ensure(Readback.bComplete)) // Readbacks should have been performed at beginning of refresh
+		{
+			PCGData = Readback.TaggedData.Data;
+		}
+	}
+
 	if (!PCGData)
 	{
 		return;
@@ -1455,12 +1501,8 @@ void SPCGEditorGraphAttributeListView::OnSelectionChangedPin(TSharedPtr<FPinComb
 		PCGEditorPtr.Pin()->GetInspectionDataManager().SetupInspectionEntry(Params);
 	}
 
-	RefreshDataComboBox(/*bKeepSelection=*/InSelectInfo == ESelectInfo::Direct);
-
-	if (InSelectInfo != ESelectInfo::Direct)
-	{
-		RefreshAttributeList();
-	}
+	bKeepDataComboBoxSelection = InSelectInfo == ESelectInfo::Direct;
+	RequestRefresh();
 }
 
 TSharedRef<SWidget> SPCGEditorGraphAttributeListView::OnGeneratePinWidget(TSharedPtr<FPinComboBoxItem> InItem) const
diff --git a/Engine/Plugins/PCG/Source/PCGEditor/Private/Widgets/SPCGEditorGraphAttributeListView.h b/Engine/Plugins/PCG/Source/PCGEditor/Private/Widgets/SPCGEditorGraphAttributeListView.h
--- a/Engine/Plugins/PCG/Source/PCGEditor/Private/Widgets/SPCGEditorGraphAttributeListView.h
+++ b/Engine/Plugins/PCG/Source/PCGEditor/Private/Widgets/SPCGEditorGraphAttributeListView.h
@@ -174,6 +174,9 @@ private:
 	void OnTick();
 	TSharedRef<SHeaderRow> CreateHeaderRowWidget() const;
 
+	/** Read back any GPU data in the inspection data. Returns true when read backs complete. */
+	bool ReadbackGPUData() const;
+
 	void OnInspectedStackChanged(const FPCGStack& InPCGStack);
 
 	void ResizeColumnToMaxWidth(FName InColumnId);
@@ -334,6 +337,8 @@ private:
 	bool bNeedsRefresh : 1 = false;
 	bool bViewportNeedsRefresh : 1 = false;
 
+	bool bKeepDataComboBoxSelection = true;
+
 	/** True if this property view is currently locked (i.e. The objects being observed are not changed automatically due to user selection)*/
 	bool bIsLocked : 1 = false;
 
