[PCG] Fix ISM Component reuse in Preview mode

#Synced-CL 48606104

diff --git a/Engine/Plugins/PCG/Source/PCG/Private/Helpers/PCGActorHelpers.cpp b/Engine/Plugins/PCG/Source/PCG/Private/Helpers/PCGActorHelpers.cpp
--- a/Engine/Plugins/PCG/Source/PCG/Private/Helpers/PCGActorHelpers.cpp
+++ b/Engine/Plugins/PCG/Source/PCG/Private/Helpers/PCGActorHelpers.cpp
@@ -148,7 +148,10 @@ UPCGManagedISMComponent* UPCGActorHelpers::GetOrCreateManagedISMC(AActor* InTarg
 		UPCGManagedISMComponent* MatchingResource = nullptr;
 		if(InParams.SettingsCrc.IsValid())
 		{
-			InSourceComponent->ForEachManagedResource([&MatchingResource, &InParams, &InTargetActor, &Descriptor](UPCGManagedResource* InResource)
+			// If we explicitely are trying to find a RF_Transient through the params or if the PCGComponent itself is in preview mode
+			const bool bTransient = InParams.bTransient || (InSourceComponent->GetEditingMode() == EPCGEditorDirtyMode::Preview);
+
+			InSourceComponent->ForEachManagedResource([&MatchingResource, &InParams, &InTargetActor, &Descriptor, bTransient](UPCGManagedResource* InResource)
 			{
 				// Early out if already found a match
 				if (MatchingResource)
@@ -171,7 +174,7 @@ UPCGManagedISMComponent* UPCGActorHelpers::GetOrCreateManagedISMC(AActor* InTarg
 							ISMC->GetOwner() == InTargetActor &&
 							ISMC->NumCustomDataFloats == InParams.NumCustomDataFloats &&
 							Resource->GetDescriptor() == Descriptor &&
-							ISMC->HasAnyFlags(RF_Transient) == InParams.bTransient &&
+							ISMC->HasAnyFlags(RF_Transient) == bTransient &&
 							Algo::Compare(ISMC->GetDefaultCustomPrimitiveData().Data, InParams.CustomPrimitiveData, [](const float LHS, const float RHS) { return FMath::IsNearlyEqual(LHS, RHS); }))
 						{
 							MatchingResource = Resource;
@@ -889,4 +892,4 @@ TArray<APCGIsolatedActor*> UPCGActorHelpers::GetIsolatedActorsOriginatingFrom(co
 	}
 
 	return IsolatedActors;
-}
\ No newline at end of file
+}
