[PCG] Adapt query tool for automated usage
- Bring Isolate functionality from Editor mode into PCG module

#rnx

#Synced-CL 48605867

diff --git a/Engine/Plugins/PCG/Source/PCG/Private/Helpers/PCGActorHelpers.cpp b/Engine/Plugins/PCG/Source/PCG/Private/Helpers/PCGActorHelpers.cpp
--- a/Engine/Plugins/PCG/Source/PCG/Private/Helpers/PCGActorHelpers.cpp
+++ b/Engine/Plugins/PCG/Source/PCG/Private/Helpers/PCGActorHelpers.cpp
@@ -893,3 +893,54 @@ TArray<APCGIsolatedActor*> UPCGActorHelpers::GetIsolatedActorsOriginatingFrom(co
 
 	return IsolatedActors;
 }
+
+AActor* UPCGActorHelpers::IsolateFromActor(AActor* Actor, const TArray<FName>& RequiredTags)
+{
+	UPCGComponent* PCGComponent = Actor ? Actor->GetComponentByClass<UPCGComponent>() : nullptr;
+
+	if (!PCGComponent)
+	{
+		return nullptr;
+	}
+
+	// Will store data temporarily (output from the ClearPCGLink) until we move it to the created actor
+	FPCGDataCollection TempCollection;
+
+	FPCGMoveResourceParams Params
+	{
+		.TemplateTargetClass = APCGIsolatedActor::StaticClass(),
+		.RequiredTags = RequiredTags,
+		.TargetDataCollection = &TempCollection,
+	};
+
+	AActor* IsolatedActor = PCGComponent->ClearPCGLink(Params);
+
+	if (IsolatedActor)
+	{
+		// @todo_pcg: Push tags onto actor?
+		// Remove all PCG Generated tags, so we don't ignore those elsewhere
+		IsolatedActor->ForEachComponent(/*bIncludeFromChildActors=*/true, [](UActorComponent* Component)
+		{
+			Component->ComponentTags.Remove(PCGHelpers::DefaultPCGTag);
+		});
+
+		if (APCGIsolatedActor* IsolatedPCGActor = Cast<APCGIsolatedActor>(IsolatedActor))
+		{
+			IsolatedPCGActor->Origin = Actor;
+
+			// Move collection to isolated actor
+			IsolatedPCGActor->Data = MoveTemp(TempCollection);
+
+			// Re-outer data
+			for (FPCGTaggedData& TaggedData : IsolatedPCGActor->Data.TaggedData)
+			{
+				TaggedData.Data->VisitDataNetwork([IsolatedPCGActor](const UPCGData* InData)
+				{
+					const_cast<UPCGData*>(InData)->Rename(nullptr, IsolatedPCGActor, REN_DontCreateRedirectors);
+				});
+			}
+		}
+	}
+
+	return IsolatedActor;
+}
diff --git a/Engine/Plugins/PCG/Source/PCG/Public/Helpers/PCGActorHelpers.h b/Engine/Plugins/PCG/Source/PCG/Public/Helpers/PCGActorHelpers.h
--- a/Engine/Plugins/PCG/Source/PCG/Public/Helpers/PCGActorHelpers.h
+++ b/Engine/Plugins/PCG/Source/PCG/Public/Helpers/PCGActorHelpers.h
@@ -261,6 +261,9 @@ PRAGMA_ENABLE_DEPRECATION_WARNINGS
 
 	/** Returns all isolated actors that originate from the given actor. */
 	static TArray<APCGIsolatedActor*> GetIsolatedActorsOriginatingFrom(const AActor* InActor);
+	
+	/** Extracts (like Clear PCG Link) content from the target actor, assuming it has a PCG component, into a PCG Isolated Actor. */
+	static UE_API AActor* IsolateFromActor(AActor* Actor, const TArray<FName>& RequiredTags);
 };
 
 #undef UE_API
diff --git a/Engine/Plugins/PCG/Source/PCGEditor/Private/EditorMode/PCGEdMode.cpp b/Engine/Plugins/PCG/Source/PCGEditor/Private/EditorMode/PCGEdMode.cpp
--- a/Engine/Plugins/PCG/Source/PCGEditor/Private/EditorMode/PCGEdMode.cpp
+++ b/Engine/Plugins/PCG/Source/PCGEditor/Private/EditorMode/PCGEdMode.cpp
@@ -127,7 +127,7 @@ namespace PCGEditorMode
 				Tags.Add(FName(Args[ArgIndex]));
 			}
 			
-			PCGIsolateTool::Isolate(Actor, Tags);
+			UPCGActorHelpers::IsolateFromActor(Actor, Tags);
 		}));
 }
 
diff --git a/Engine/Plugins/PCG/Source/PCGEditor/Private/EditorMode/Tools/Query/PCGQueryTool.cpp b/Engine/Plugins/PCG/Source/PCGEditor/Private/EditorMode/Tools/Query/PCGQueryTool.cpp
--- a/Engine/Plugins/PCG/Source/PCGEditor/Private/EditorMode/Tools/Query/PCGQueryTool.cpp
+++ b/Engine/Plugins/PCG/Source/PCGEditor/Private/EditorMode/Tools/Query/PCGQueryTool.cpp
@@ -36,65 +36,9 @@ namespace PCGIsolateTool
 	{
 		const FName ToolTag = TEXT("IsolateTool");
 	}
-
-	AActor* Isolate(AActor* Actor, const TArray<FName>& Tags)
-	{
-		UPCGComponent* PCGComponent = Actor ? Actor->GetComponentByClass<UPCGComponent>() : nullptr;
-
-		if (!Actor || !PCGComponent)
-		{
-			return nullptr;
-		}
-
-		// Will store data temporarily (output from the ClearPCGLink) until we move it to the created actor
-		FPCGDataCollection TempCollection;
-
-		FPCGMoveResourceParams Params
-		{
-			.TemplateTargetClass = APCGIsolatedActor::StaticClass(),
-			.RequiredTags = Tags,
-			.TargetDataCollection = &TempCollection,
-		};
-
-		AActor* IsolatedActor = PCGComponent->ClearPCGLink(Params);
-
-		if (IsolatedActor)
-		{
-			// @todo_pcg: Push tags onto actor?
-			// Remove all PCG Generated tags, so we don't ignore those elsewhere
-			IsolatedActor->ForEachComponent(/*bIncludeFromChildActors=*/true, [](UActorComponent* Component)
-			{
-				Component->ComponentTags.Remove(PCGHelpers::DefaultPCGTag);
-			});
-
-			if (APCGIsolatedActor* IsolatedPCGActor = Cast<APCGIsolatedActor>(IsolatedActor))
-			{
-				IsolatedPCGActor->Origin = Actor;
-
-				// Move collection to isolated actor
-				IsolatedPCGActor->Data = MoveTemp(TempCollection);
-
-				// Re-outer data
-				for (FPCGTaggedData& TaggedData : IsolatedPCGActor->Data.TaggedData)
-				{
-					TaggedData.Data->VisitDataNetwork([IsolatedPCGActor](const UPCGData* InData)
-					{
-						const_cast<UPCGData*>(InData)->Rename(nullptr, IsolatedPCGActor, REN_DontCreateRedirectors);
-					});
-				}
-			}
-		}
-
-		return IsolatedActor;
-	}
 }
 
-FName UPCGQueryTool::StaticGetToolTag()
-{
-	return PCGQueryTool::Constants::ToolTag;
-}
-
-void UPCGQueryTool::Setup()
+void UPCGBaseQueryTool::Setup()
 {
 	Super::Setup();
 
@@ -103,7 +47,7 @@ void UPCGQueryTool::Setup()
 	AddInputBehavior(ClickBehavior);
 }
 
-void UPCGQueryTool::Shutdown(const EToolShutdownType ShutdownType)
+void UPCGBaseQueryTool::Shutdown(const EToolShutdownType ShutdownType)
 {
 	// We have to call this before the super function as it would remove the property sets before we can route the shutdown function to the pcg settings
 	UE::PCG::EditorMode::Tool::Shutdown(this, ShutdownType);
@@ -111,7 +55,7 @@ void UPCGQueryTool::Shutdown(const EToolShutdownType ShutdownType)
 	Super::Shutdown(ShutdownType);
 }
 
-HHitProxy* UPCGQueryTool::GetHitProxy(const FInputDeviceRay& ClickPosition) const
+HHitProxy* UPCGBaseQueryTool::GetHitProxy(const FInputDeviceRay& ClickPosition) const
 {
 	FViewport* FocusedViewport = GetToolManager()->GetContextQueriesAPI()->GetFocusedViewport();
 	if (FocusedViewport)
@@ -125,7 +69,7 @@ HHitProxy* UPCGQueryTool::GetHitProxy(const FInputDeviceRay& ClickPosition) cons
 	return nullptr;
 }
 
-UActorComponent* UPCGQueryTool::GetHitComponent(HHitProxy* HitProxy) const
+UActorComponent* UPCGBaseQueryTool::GetHitComponent(HHitProxy* HitProxy) const
 {
 	if (!HitProxy)
 	{
@@ -145,12 +89,12 @@ UActorComponent* UPCGQueryTool::GetHitComponent(HHitProxy* HitProxy) const
 	return nullptr;
 }
 
-UActorComponent* UPCGQueryTool::GetHitComponent(const FInputDeviceRay& ClickPosition) const
+UActorComponent* UPCGBaseQueryTool::GetHitComponent(const FInputDeviceRay& ClickPosition) const
 {
 	return GetHitComponent(GetHitProxy(ClickPosition));
 }
 
-FInputRayHit UPCGQueryTool::IsHitByClick(const FInputDeviceRay& ClickPosition)
+FInputRayHit UPCGBaseQueryTool::IsHitByClick(const FInputDeviceRay& ClickPosition)
 {
 	using namespace UE::PCG::EditorMode;
 
@@ -166,7 +110,7 @@ FInputRayHit UPCGQueryTool::IsHitByClick(const FInputDeviceRay& ClickPosition)
 	return {};
 }
 
-void UPCGQueryTool::OnClicked(const FInputDeviceRay& ClickPosition)
+void UPCGBaseQueryTool::OnClicked(const FInputDeviceRay& ClickPosition)
 {
 	using namespace UE::PCG::EditorMode;
 
@@ -184,7 +128,7 @@ void UPCGQueryTool::OnClicked(const FInputDeviceRay& ClickPosition)
 	}
 }
 
-void UPCGQueryTool::VisitHierarchy(UActorComponent* HitComponent, int32 InstanceIndex, TFunctionRef<void(UActorComponent*, AActor*, int32)> Func)
+void UPCGBaseQueryTool::VisitHierarchy(UActorComponent* HitComponent, int32 InstanceIndex, TFunctionRef<void(UActorComponent*, AActor*, int32)> Func)
 {
 	if (!HitComponent)
 	{
@@ -221,9 +165,14 @@ void UPCGQueryTool::VisitHierarchy(UActorComponent* HitComponent, int32 Instance
 	}
 }
 
+FName UPCGQueryTool::StaticGetToolTag()
+{
+	return PCGQueryTool::Constants::ToolTag;
+}
+
 void UPCGQueryTool::OnHitProxy(HHitProxy* HitProxy, UActorComponent* HitComponent, int32 InstanceIndex)
 {
-	auto OutputToLog = [](UActorComponent* CurrentComponent, AActor* CurrentActor, int32 InstanceIndex)
+	auto OutputToLog = [&FoundActorLabel = FoundActorLabel, &SelectedTags = SelectedTags](UActorComponent* CurrentComponent, AActor* CurrentActor, int32 InstanceIndex)
 	{
 		if(CurrentComponent)
 		{
@@ -239,30 +188,38 @@ void UPCGQueryTool::OnHitProxy(HHitProxy* HitProxy, UActorComponent* HitComponen
 				UE_LOG(LogPCGEditor, Warning, TEXT("Tags: "));
 				for (FName Tag : CurrentComponent->ComponentTags)
 				{
+					SelectedTags.Add(Tag.ToString());
 					UE_LOG(LogPCGEditor, Warning, TEXT("      %s"), *Tag.ToString());
 				}
 			}
 		}
 		else if (CurrentActor)
 		{
+			if (FoundActorLabel.IsEmpty())
+			{
+				FoundActorLabel = CurrentActor->GetActorLabel();
+			}
 			UE_LOG(LogPCGEditor, Warning, TEXT("Actor: '%s'"), *CurrentActor->GetActorLabel());
 			if (!CurrentActor->Tags.IsEmpty())
 			{
 				UE_LOG(LogPCGEditor, Warning, TEXT("Tags: "));
 				for (FName Tag : CurrentActor->Tags)
 				{
+					SelectedTags.Add(Tag.ToString());
 					UE_LOG(LogPCGEditor, Warning, TEXT("      %s"), *Tag.ToString());
 				}
 			}
 		}
 	};
 
+	SelectedTags.Empty();
+	FoundActorLabel = FString();
 	VisitHierarchy(HitComponent, InstanceIndex, OutputToLog);
 }
 
 UInteractiveTool* UPCGQueryToolBuilder::BuildTool(const FToolBuilderState& SceneState) const
 {
-	UPCGQueryTool* NewTool = NewObject<UPCGQueryTool>(SceneState.ToolManager, ToolClass);
+	UPCGBaseQueryTool* NewTool = NewObject<UPCGBaseQueryTool>(SceneState.ToolManager, ToolClass);
 
 	NewTool->SetToolInfo(
 	{
@@ -286,7 +243,7 @@ bool UPCGQueryToolBuilder::CanBuildTool(const FToolBuilderState& SceneState) con
 	return !!ToolClass;
 }
 
-void UPCGQueryToolBuilder::SetToolClass(TSubclassOf<UPCGQueryTool> InClass)
+void UPCGQueryToolBuilder::SetToolClass(TSubclassOf<UPCGBaseQueryTool> InClass)
 {
 	ToolClass = InClass;
 }
@@ -295,7 +252,7 @@ void UPCGInteractiveToolSettings_Isolate::Apply(UInteractiveTool* OwningTool)
 {
 	if (AActor* Actor = SelectedActor.Get())
 	{
-		PCGIsolateTool::Isolate(Actor, SelectedTags);
+		UPCGActorHelpers::IsolateFromActor(Actor, SelectedTags);
 	}
 }
 
diff --git a/Engine/Plugins/PCG/Source/PCGEditor/Private/EditorMode/Tools/Query/PCGQueryTool.h b/Engine/Plugins/PCG/Source/PCGEditor/Private/EditorMode/Tools/Query/PCGQueryTool.h
--- a/Engine/Plugins/PCG/Source/PCGEditor/Private/EditorMode/Tools/Query/PCGQueryTool.h
+++ b/Engine/Plugins/PCG/Source/PCGEditor/Private/EditorMode/Tools/Query/PCGQueryTool.h
@@ -16,18 +16,13 @@ class HHitProxy;
 class UActorComponent;
 class USingleClickInputBehavior;
 
-namespace PCGIsolateTool
-{
-	/** Extracts (like Clear PCG Link) content from the target actor, assuming it has a PCG component, into a PCG Isolated Actor. */
-	AActor* Isolate(AActor* Actor, const TArray<FName>& RequiredTags);
-}
 
 /**
 * Simple editor tool to gather information from a picked object in the world. 
 * By default logs visited hierarchy tags (from touched component, to actor, ...) but can be reimplemented in a derived class.
 */
 UCLASS()
-class UPCGQueryTool : public UPCGInteractiveTool, public IClickBehaviorTarget
+class UPCGBaseQueryTool : public UPCGInteractiveTool, public IClickBehaviorTarget
 {
 	GENERATED_BODY()
 
@@ -35,9 +30,6 @@ public:
 	//~Begin UInteractiveTool interface
 	virtual void Setup() override;
 	virtual void Shutdown(EToolShutdownType ShutdownType) override;
-	virtual bool HasCancel() const override { return false; }
-	virtual bool HasAccept() const override { return false; }
-	virtual bool CanAccept() const override { return false; }
 	//~End UInteractiveTool interface
 
 	//~Begin UPCGInteractiveTool interface
@@ -49,19 +41,50 @@ public:
 	virtual void OnClicked(const FInputDeviceRay& ClickPosition) override;
 	//~End IClickBehaviorTarget interface
 
-	static FName StaticGetToolTag();
-
 protected:
 	HHitProxy* GetHitProxy(const FInputDeviceRay& ClickPosition) const;
 	UActorComponent* GetHitComponent(HHitProxy* HitProxy) const;
 	UActorComponent* GetHitComponent(const FInputDeviceRay& ClickPosition) const;
-	virtual void OnHitProxy(HHitProxy* HitProxy, UActorComponent* HitComponent, int InstanceIndex);
+	virtual void OnHitProxy(HHitProxy* HitProxy, UActorComponent* HitComponent, int InstanceIndex) {}
 	void VisitHierarchy(UActorComponent* HitComponent, int32 InstanceIndex, TFunctionRef<void(UActorComponent*, AActor*, int32)> Func);
 
 private:
 	UPROPERTY(Instanced)
 	TObjectPtr<USingleClickInputBehavior> ClickBehavior = nullptr;
 };
+/**
+* Simple editor tool to gather information from a picked object in the world. 
+* By default logs visited hierarchy tags (from touched component, to actor, ...) but can be reimplemented in a derived class.
+*/
+UCLASS()
+class UPCGQueryTool : public UPCGBaseQueryTool
+{
+	GENERATED_BODY()
+
+public:
+	//~Begin UInteractiveTool interface
+	virtual bool HasCancel() const override { return true; }
+	virtual bool HasAccept() const override { return true; }
+	virtual bool CanAccept() const override { return !SelectedTags.IsEmpty(); }
+	//~End UInteractiveTool interface
+
+	//~Begin UPCGInteractiveTool interface
+	virtual bool RequiresSettings() const override { return false; }
+	//~End UPCGInteractiveTool interface
+
+	static FName StaticGetToolTag();
+	
+	UPROPERTY()
+	TArray<FString> SelectedTags;
+
+	UPROPERTY()
+	FString FoundActorLabel;
+
+protected:
+	virtual void OnHitProxy(HHitProxy* HitProxy, UActorComponent* HitComponent, int InstanceIndex);
+};
+
+
 
 UCLASS(Transient)
 class UPCGQueryToolBuilder : public UPCGInteractiveToolBuilder
@@ -72,11 +95,11 @@ public:
 	virtual bool CanBuildTool(const FToolBuilderState& SceneState) const override;
 	virtual UInteractiveTool* BuildTool(const FToolBuilderState& SceneState) const override;
 
-	void SetToolClass(TSubclassOf<UPCGQueryTool> InClass);
+	void SetToolClass(TSubclassOf<UPCGBaseQueryTool> InClass);
 
 private:
 	UPROPERTY()
-	TSubclassOf<UPCGQueryTool> ToolClass = nullptr;
+	TSubclassOf<UPCGBaseQueryTool> ToolClass = nullptr;
 };
 
 UCLASS()
@@ -108,7 +131,7 @@ public:
 * Tags are gathered through picking (from the Query tool), and listed so the user can choose a combination.
 */
 UCLASS()
-class UPCGIsolateTool : public UPCGQueryTool
+class UPCGIsolateTool : public UPCGBaseQueryTool
 {
 	GENERATED_BODY()
 
