[PCG] Pass runtime gen in to context and make grabbing the frustum conditional on it being used, because this was the majority of the gen source gather/tick cost.

#rnx

#Synced-CL 48603294

diff --git a/Engine/Plugins/PCG/Source/PCG/Private/RuntimeGen/GenSources/PCGGenSourceBase.cpp b/Engine/Plugins/PCG/Source/PCG/Private/RuntimeGen/GenSources/PCGGenSourceBase.cpp
--- a/Engine/Plugins/PCG/Source/PCG/Private/RuntimeGen/GenSources/PCGGenSourceBase.cpp
+++ b/Engine/Plugins/PCG/Source/PCG/Private/RuntimeGen/GenSources/PCGGenSourceBase.cpp
@@ -48,6 +48,14 @@ static bool AreConvexVolumesEquivalent(const TOptional<FConvexVolume>& InA, cons
 	return true;
 }
 
+void IPCGGenSourceBase::Tick(const FPCGRuntimeGenContext& InContext)
+{
+	// Backwards compatibility.
+	PRAGMA_DISABLE_DEPRECATION_WARNINGS
+	Tick();
+	PRAGMA_ENABLE_DEPRECATION_WARNINGS
+}
+
 bool IPCGGenSourceBase::IsEquivalent(IPCGGenSourceBase* InOther, const FPCGRuntimeGenContext& InContext) const
 {
 	if (this == InOther)
diff --git a/Engine/Plugins/PCG/Source/PCG/Private/RuntimeGen/GenSources/PCGGenSourceEditorCamera.cpp b/Engine/Plugins/PCG/Source/PCG/Private/RuntimeGen/GenSources/PCGGenSourceEditorCamera.cpp
--- a/Engine/Plugins/PCG/Source/PCG/Private/RuntimeGen/GenSources/PCGGenSourceEditorCamera.cpp
+++ b/Engine/Plugins/PCG/Source/PCG/Private/RuntimeGen/GenSources/PCGGenSourceEditorCamera.cpp
@@ -2,6 +2,8 @@
 
 #include "RuntimeGen/GenSources/PCGGenSourceEditorCamera.h"
 
+#include "RuntimeGen/PCGRuntimeGenContext.h"
+
 #include "SceneView.h"
 
 #if WITH_EDITOR
@@ -10,12 +12,12 @@
 
 #include UE_INLINE_GENERATED_CPP_BY_NAME(PCGGenSourceEditorCamera)
 
-void UPCGGenSourceEditorCamera::Tick()
+void UPCGGenSourceEditorCamera::Tick(const FPCGRuntimeGenContext& InContext)
 {
 #if WITH_EDITOR
 	ViewFrustum = TOptional<FConvexVolume>();
 
-	if (EditorViewportClient)
+	if (InContext.bAnySourcesUseFrustumCulling && EditorViewportClient)
 	{
 		const UWorld* World = EditorViewportClient->GetWorld();
 		const FViewport* Viewport = EditorViewportClient->Viewport;
@@ -71,4 +73,4 @@ TOptional<FConvexVolume> UPCGGenSourceEditorCamera::GetViewFrustum(bool bIs2DGri
 #else
 	return TOptional<FConvexVolume>();
 #endif	
-}
\ No newline at end of file
+}
diff --git a/Engine/Plugins/PCG/Source/PCG/Private/RuntimeGen/GenSources/PCGGenSourcePlayer.cpp b/Engine/Plugins/PCG/Source/PCG/Private/RuntimeGen/GenSources/PCGGenSourcePlayer.cpp
--- a/Engine/Plugins/PCG/Source/PCG/Private/RuntimeGen/GenSources/PCGGenSourcePlayer.cpp
+++ b/Engine/Plugins/PCG/Source/PCG/Private/RuntimeGen/GenSources/PCGGenSourcePlayer.cpp
@@ -2,6 +2,8 @@
 
 #include "RuntimeGen/GenSources/PCGGenSourcePlayer.h"
 
+#include "RuntimeGen/PCGRuntimeGenContext.h"
+
 #include "SceneView.h"
 #include "Engine/GameViewportClient.h"
 #include "Engine/LocalPlayer.h"
@@ -10,7 +12,7 @@
 
 #include UE_INLINE_GENERATED_CPP_BY_NAME(PCGGenSourcePlayer)
 
-void UPCGGenSourcePlayer::Tick()
+void UPCGGenSourcePlayer::Tick(const FPCGRuntimeGenContext& InContext)
 {
 	ViewFrustum = TOptional<FConvexVolume>();
 
@@ -18,16 +20,19 @@ void UPCGGenSourcePlayer::Tick()
 
 	if (LocalPlayer)
 	{
-		if (LocalPlayer->ViewportClient)
+		if (InContext.bAnySourcesUseFrustumCulling)
 		{
-			FSceneViewProjectionData ProjectionData;
-
-			if (LocalPlayer->GetProjectionData(LocalPlayer->ViewportClient->Viewport, ProjectionData))
+			if (LocalPlayer->ViewportClient)
 			{
-				FConvexVolume ConvexVolume;
-				GetViewFrustumBounds(ConvexVolume, ProjectionData.ComputeViewProjectionMatrix(), /*bUseNearPlane=*/true, /*bUseFarPlane=*/true);
+				FSceneViewProjectionData ProjectionData;
+
+				if (LocalPlayer->GetProjectionData(LocalPlayer->ViewportClient->Viewport, ProjectionData))
+				{
+					FConvexVolume ConvexVolume;
+					GetViewFrustumBounds(ConvexVolume, ProjectionData.ComputeViewProjectionMatrix(), /*bUseNearPlane=*/true, /*bUseFarPlane=*/true);
 
-				ViewFrustum = ConvexVolume;
+					ViewFrustum = ConvexVolume;
+				}
 			}
 		}
 
@@ -37,6 +42,11 @@ void UPCGGenSourcePlayer::Tick()
 		CachedPosition = CameraLocation;
 		CachedDirection = CameraRotation.Vector();
 	}
+	else
+	{
+		CachedPosition.Reset();
+		CachedDirection.Reset();
+	}
 }
 
 TOptional<FVector> UPCGGenSourcePlayer::GetPosition() const
diff --git a/Engine/Plugins/PCG/Source/PCG/Private/RuntimeGen/PCGGenSourceManager.cpp b/Engine/Plugins/PCG/Source/PCG/Private/RuntimeGen/PCGGenSourceManager.cpp
--- a/Engine/Plugins/PCG/Source/PCG/Private/RuntimeGen/PCGGenSourceManager.cpp
+++ b/Engine/Plugins/PCG/Source/PCG/Private/RuntimeGen/PCGGenSourceManager.cpp
@@ -37,11 +37,9 @@ namespace PCGGenSourceManager
 		TEXT("Enable the RuntimeGeneration system on servers. (0 disabled, 1 local generation sources only (listen server), 2 all generation sources"));
 }
 
-FPCGGenSourceManager::FPCGGenSourceManager(const UWorld* InWorld)
+FPCGGenSourceManager::FPCGGenSourceManager(const FPCGRuntimeGenContext& InContext)
+	: Context(InContext)
 {
-	// We capture the World so we can differentiate between editor-world and PIE-world Generation Sources.
-	World = InWorld;
-
 #if WITH_EDITOR
 	EditorCameraGenSource = NewObject<UPCGGenSourceEditorCamera>();
 #endif
@@ -71,7 +69,7 @@ TSet<IPCGGenSourceBase*> FPCGGenSourceManager::GetAllGenSources(const APCGWorldA
 		{
 			if (bDirty)
 			{
-				GenSource->Tick();
+				GenSource->Tick(Context);
 			}
 
 			AllGenSources.Add(GenSource.GetInterface());
@@ -84,7 +82,7 @@ TSet<IPCGGenSourceBase*> FPCGGenSourceManager::GetAllGenSources(const APCGWorldA
 		{
 			if (bDirty)
 			{
-				NamedGenSource.Value->Tick();
+				NamedGenSource.Value->Tick(Context);
 			}
 
 			AllGenSources.Add(NamedGenSource.Value.GetInterface());
@@ -100,7 +98,7 @@ TSet<IPCGGenSourceBase*> FPCGGenSourceManager::GetAllGenSources(const APCGWorldA
 	{
 		if (bDirty)
 		{
-			EditorCameraGenSource->Tick();
+			EditorCameraGenSource->Tick(Context);
 		}
 
 		AllGenSources.Add(EditorCameraGenSource);
@@ -132,7 +130,7 @@ bool FPCGGenSourceManager::RegisterGenSource(IPCGGenSourceBase* InGenSource, FNa
 {
 	if (UPCGGenSourceComponent* GenSourceComponent = Cast<UPCGGenSourceComponent>(InGenSource))
 	{
-		if (GenSourceComponent->GetWorld() != World)
+		if (GenSourceComponent->GetWorld() != Context.World)
 		{
 			return false;
 		}
@@ -161,7 +159,7 @@ bool FPCGGenSourceManager::UnregisterGenSource(const IPCGGenSourceBase* InGenSou
 {
 	if (const UPCGGenSourceComponent* GenSourceComponent = Cast<UPCGGenSourceComponent>(InGenSource))
 	{
-		if (GenSourceComponent->GetWorld() != World)
+		if (GenSourceComponent->GetWorld() != Context.World)
 		{
 			return false;
 		}
@@ -197,13 +195,13 @@ void FPCGGenSourceManager::UpdatePerTickGenSources(const APCGWorldActor* InPCGWo
 	// Gather Player controller sources
 	TSet<APlayerController*> CurrentControllers;
 
-	const bool bIsServer = PCGHelpers::IsListenServer(World) || PCGHelpers::IsDedicatedServer(World);
+	const bool bIsServer = PCGHelpers::IsListenServer(Context.World) || PCGHelpers::IsDedicatedServer(Context.World);
 	const int32 ServerMode = PCGGenSourceManager::CVarServerRuntimeGenerationMode.GetValueOnAnyThread();
 
 	if (InPCGWorldActor->bTreatPlayerControllersAsGenerationSources && (!bIsServer || ServerMode > 0))
 	{
 		// Update existing sources or add new ones
-		for (auto It = World->GetPlayerControllerIterator(); It; ++It)
+		for (auto It = Context.World->GetPlayerControllerIterator(); It; ++It)
 		{
 			APlayerController* PlayerController = It->Get();
 			if (!PlayerController)
@@ -252,7 +250,7 @@ void FPCGGenSourceManager::UpdatePerTickGenSources(const APCGWorldActor* InPCGWo
 	EditorCameraGenSource->EditorViewportClient = nullptr;
 
 	// Update the active editor viewport client for the EditorCameraGenSource, only if requested by the world actor, in-editor, and the viewport is visible.
-	if (InPCGWorldActor->bTreatEditorViewportAsGenerationSource && !World->IsGameWorld() && GEditor)
+	if (InPCGWorldActor->bTreatEditorViewportAsGenerationSource && !Context.World->IsGameWorld() && GEditor)
 	{
 		if (FViewport* Viewport = GEditor->GetActiveViewport())
 		{
@@ -277,7 +275,7 @@ PRAGMA_DISABLE_DEPRECATION_WARNINGS
 		{
 			if (bDirty)
 			{
-				WPGenSource->Tick();
+				WPGenSource->Tick(Context);
 			}
 
 			OutGenSources.Add(WPGenSource);
@@ -354,3 +352,10 @@ PRAGMA_ENABLE_DEPRECATION_WARNINGS
 		KeyValuePair.Value.AddReferencedObjects(Collector);
 	}
 }
+
+// Deprecated (5.8)
+FPCGGenSourceManager::FPCGGenSourceManager(const UWorld* InWorld)
+	: Context(FallBackContext)
+	, FallBackContext(FPCGRuntimeGenContext{ .World = const_cast<UWorld*>(InWorld), .bAnyRuntimeGenSourcesExist = true, .bAnySourcesUseFrustumCulling = true, .bAnySourcesUse2DGrids = true, .bAnySourcesUse3DGrids = true })
+{
+}
diff --git a/Engine/Plugins/PCG/Source/PCG/Private/RuntimeGen/PCGRuntimeGenScheduler.cpp b/Engine/Plugins/PCG/Source/PCG/Private/RuntimeGen/PCGRuntimeGenScheduler.cpp
--- a/Engine/Plugins/PCG/Source/PCG/Private/RuntimeGen/PCGRuntimeGenScheduler.cpp
+++ b/Engine/Plugins/PCG/Source/PCG/Private/RuntimeGen/PCGRuntimeGenScheduler.cpp
@@ -180,7 +180,7 @@ namespace PCGRuntimeGenSchedulerHelpers
 						TEXT("    PA Pool: %d / %d available\n")
 						TEXT("    VT Preloads: %d\n")
 						TEXT("    Unique generation sources:\n%s"),
-						PCGHelpers::IsServer(InRuntimeGenScheduler.World) ? TEXT(" (Server)") : TEXT(""),
+						PCGHelpers::IsServer(InRuntimeGenScheduler.Context.World) ? TEXT(" (Server)") : TEXT(""),
 						ElapsedMs,
 						GenerateCallCounter,
 						CleanupCallCounter,
@@ -215,12 +215,13 @@ FPCGRuntimeGenScheduler::FPCGRuntimeGenScheduler(UWorld* InWorld, FPCGTrackingMa
 {
 	check(InWorld && InTrackingManager);
 
-	World = InWorld;
-	Subsystem = UPCGSubsystem::GetInstance(World);
-	WorldPartitionSubsystem = World ? UWorld::GetSubsystem<UWorldPartitionSubsystem>(World) : nullptr;
+	Context.World = InWorld;
+
+	Subsystem = UPCGSubsystem::GetInstance(Context.World);
+	WorldPartitionSubsystem = Context.World ? UWorld::GetSubsystem<UWorldPartitionSubsystem>(Context.World) : nullptr;
 
 	TrackingManager = InTrackingManager;
-	GenSourceManager = new FPCGGenSourceManager(InWorld);
+	GenSourceManager = new FPCGGenSourceManager(Context);
 	bPoolingWasEnabledLastFrame = PCGRuntimeGenSchedulerHelpers::CVarRuntimeGenerationEnablePooling.GetValueOnAnyThread();
 	BasePoolSizeLastFrame = PCGRuntimeGenSchedulerHelpers::CVarRuntimeGenerationBasePoolSize.GetValueOnAnyThread();
 	FramesUntilGeneration = PCGRuntimeGenSchedulerHelpers::CVarFramesBeforeFirstGenerate.GetValueOnGameThread();
@@ -236,7 +237,7 @@ FPCGRuntimeGenScheduler::FPCGRuntimeGenScheduler(UWorld* InWorld, FPCGTrackingMa
 	TArray<APCGPartitionActor*> ExistingPoolPartitionActors;
 	int32 NumPartitionActors = 0;
 
-	UPCGActorHelpers::ForEachActorInLevel<APCGPartitionActor>(World->PersistentLevel, [&NumPartitionActors, &ExistingPoolPartitionActors](AActor* InActor)
+	UPCGActorHelpers::ForEachActorInLevel<APCGPartitionActor>(Context.World->PersistentLevel, [&NumPartitionActors, &ExistingPoolPartitionActors](AActor* InActor)
 	{
 		if (APCGPartitionActor* InPartitionActor = Cast<APCGPartitionActor>(InActor); InPartitionActor && InPartitionActor->HasAnyFlags(RF_Transient) && InPartitionActor->IsRuntimeGenerated())
 		{
@@ -367,7 +368,7 @@ void FPCGRuntimeGenScheduler::Tick(APCGWorldActor* InPCGWorldActor, double InEnd
 
 bool FPCGRuntimeGenScheduler::ShouldTick()
 {
-	check(World && TrackingManager);
+	check(Context.World && TrackingManager);
 
 	if (!PCGRuntimeGenSchedulerHelpers::CVarRuntimeGenerationEnable.GetValueOnAnyThread())
 	{
@@ -375,7 +376,7 @@ bool FPCGRuntimeGenScheduler::ShouldTick()
 	}
 
 	// Disable tick of editor scheduling if in runtime or PIE.
-	if (PCGHelpers::IsRuntimeOrPIE() && !World->IsGameWorld())
+	if (PCGHelpers::IsRuntimeOrPIE() && !Context.World->IsGameWorld())
 	{
 		return false;
 	}
@@ -383,7 +384,7 @@ bool FPCGRuntimeGenScheduler::ShouldTick()
 #if WITH_EDITOR
 	// If we're in an editor world, stop updating preview if the editor window/viewport is not active (follows
 	// same behaviour as other things).
-	if (!World->IsGameWorld())
+	if (!Context.World->IsGameWorld())
 	{
 		FViewport* Viewport = GEditor ? GEditor->GetActiveViewport() : nullptr;
 		FEditorViewportClient* ViewportClient = Viewport ? static_cast<FEditorViewportClient*>(Viewport->GetClient()) : nullptr;
@@ -617,7 +618,7 @@ void FPCGRuntimeGenScheduler::TickQueueComponentsForGeneration(
 					{
 						TOptional<FVector> Direction = GenSource->GetDirection();
 
-						DrawDebugString(World, Direction ? (GenSourcePosition + 100.0f * Direction.GetValue()) : GenSourcePosition, GenSourceObject->GetName(), /*TestBaseActor=*/nullptr, FColor::Red, /*Duration=*/0.0f);
+						DrawDebugString(Context.World, Direction ? (GenSourcePosition + 100.0f * Direction.GetValue()) : GenSourcePosition, GenSourceObject->GetName(), /*TestBaseActor=*/nullptr, FColor::Red, /*Duration=*/0.0f);
 					}
 
 					for (uint32 GridSize : GridSizes)
@@ -626,7 +627,7 @@ void FPCGRuntimeGenScheduler::TickQueueComponentsForGeneration(
 
 						const double GenerationRadius = OriginalComponent->GetGenerationRadiusFromGrid(PCGHiGenGrid::GridSizeToGrid(GridSize));
 
-						DrawDebugSphere(World, GenSourcePosition, GenerationRadius, 64, FColor::Red, /*bPersistentLines=*/false, /*LifeTime=*/0.0f);
+						DrawDebugSphere(Context.World, GenSourcePosition, GenerationRadius, 64, FColor::Red, /*bPersistentLines=*/false, /*LifeTime=*/0.0f);
 					}
 				}
 #endif // UE_ENABLE_DEBUG_DRAWING
@@ -1249,7 +1250,7 @@ void FPCGRuntimeGenScheduler::TickRequestVirtualTexturePriming(const TSet<IPCGGe
 
 		for (UObject* FoundComponent : FoundComponents)
 		{
-			if (FoundComponent && FoundComponent->GetWorld() == World)
+			if (FoundComponent && FoundComponent->GetWorld() == Context.World)
 			{
 				URuntimeVirtualTextureComponent* VirtualTextureComponent = Cast<URuntimeVirtualTextureComponent>(FoundComponent);
 				URuntimeVirtualTexture* VirtualTexture = VirtualTextureComponent ? VirtualTextureComponent->GetVirtualTexture() : nullptr;
@@ -1343,10 +1344,10 @@ void FPCGRuntimeGenScheduler::TickRequestVirtualTexturePriming(const TSet<IPCGGe
 
 						if (PCGRuntimeGenSchedulerHelpers::CVarDebugDrawTexturePrimingBounds.GetValueOnAnyThread())
 						{
-							check(World);
+							check(Context.World);
 
 							DrawDebugCylinder(
-								World,
+								Context.World,
 								FVector(GenSourcePosition.X, GenSourcePosition.Y, OriginalComponentBounds.Min.Z),
 								FVector(GenSourcePosition.X, GenSourcePosition.Y, OriginalComponentBounds.Max.Z),
 								/*Radius=*/PrimingBounds.W,
@@ -1400,7 +1401,7 @@ void FPCGRuntimeGenScheduler::TickCVars(const APCGWorldActor* InPCGWorldActor)
 #if WITH_EDITOR
 	const bool bTreatEditorViewportAsGenSource = InPCGWorldActor->bTreatEditorViewportAsGenerationSource;
 	
-	if (bTreatEditorViewportAsGenSource != bTreatEditorViewportAsGenSourcePreviousFrame && World && World->IsEditorWorld())
+	if (bTreatEditorViewportAsGenSource != bTreatEditorViewportAsGenSourcePreviousFrame && Context.World && Context.World->IsEditorWorld())
 	{
 		UE_LOG(LogPCG, Verbose, TEXT("Flushed change detection status after bTreatEditorViewportAsGenerationSource value change."));
 		ChangeDetector.FlushCachedState();
@@ -1618,7 +1619,7 @@ void FPCGRuntimeGenScheduler::CleanupLocalComponents(const APCGWorldActor* InPCG
 
 void FPCGRuntimeGenScheduler::OnLevelStreamingStateChanged(UWorld* InWorld, const ULevelStreaming* InLevelStreaming, ULevel* InLevelIfLoaded, ELevelStreamingState InPreviousState, ELevelStreamingState InNewState)
 {
-	if (World == InWorld && (InPreviousState == ELevelStreamingState::LoadedVisible || InNewState == ELevelStreamingState::LoadedVisible))
+	if (Context.World == InWorld && (InPreviousState == ELevelStreamingState::LoadedVisible || InNewState == ELevelStreamingState::LoadedVisible))
 	{
 		// todo_pcg: Fine-grained invalidation based on bounds overlap tests did not trivially work on first attempt, improve debug tools/vis and retry.
 		CachedStreamingQueryResults.Empty(CachedStreamingQueryResults.Num());
@@ -1694,7 +1695,7 @@ void FPCGRuntimeGenScheduler::CleanupLocalComponent(APCGPartitionActor* Partitio
 #if WITH_EDITOR
 			PartitionActor->Rename(nullptr, PartitionActor->GetOuter(), REN_NonTransactional | REN_DoNotDirty);
 #endif
-			World->DestroyActor(PartitionActor);
+			Context.World->DestroyActor(PartitionActor);
 		}
 	}
 }
@@ -2003,7 +2004,7 @@ APCGPartitionActor* FPCGRuntimeGenScheduler::GetPartitionActorFromPool(const FPC
 
 	check(TrackingManager);
 
-	if (!World)
+	if (!Context.World)
 	{
 		UE_LOG(LogPCG, Error, TEXT("[GetPartitionActorFromPool] World is null."));
 		return nullptr;
@@ -2075,13 +2076,13 @@ void FPCGRuntimeGenScheduler::AddPartitionActorPoolCount(int32 Count)
 	SpawnParams.ObjectFlags &= ~RF_Transactional;
 
 	// Important to override the level to make sure we spawn in the persistent level and not the editor's current editing level
-	SpawnParams.OverrideLevel = World->PersistentLevel;
+	SpawnParams.OverrideLevel = Context.World->PersistentLevel;
 
 	// Create RuntimeGen PA pool.
 	for (int32 I = 0; I < Count; ++I)
 	{
 		// TODO: do these actors get networked automatically? do we want that or not?
-		APCGPartitionActor* NewActor = World->SpawnActor<APCGPartitionActor>(SpawnParams);
+		APCGPartitionActor* NewActor = Context.World->SpawnActor<APCGPartitionActor>(SpawnParams);
 		check(NewActor);
 		NewActor->SetToRuntimeGenerated();
 		PartitionActorPool.Add(NewActor);
@@ -2098,7 +2099,7 @@ void FPCGRuntimeGenScheduler::ResetPartitionActorPoolToSize(uint32 NewPoolSize)
 #if WITH_EDITOR
 		PartitionActor->Rename(nullptr, PartitionActor->GetOuter(), REN_NonTransactional | REN_DoNotDirty);
 #endif
-		World->DestroyActor(PartitionActor);
+		Context.World->DestroyActor(PartitionActor);
 	}
 
 	PartitionActorPool.Empty();
diff --git a/Engine/Plugins/PCG/Source/PCG/Public/RuntimeGen/GenSources/PCGGenSourceBase.h b/Engine/Plugins/PCG/Source/PCG/Public/RuntimeGen/GenSources/PCGGenSourceBase.h
--- a/Engine/Plugins/PCG/Source/PCG/Public/RuntimeGen/GenSources/PCGGenSourceBase.h
+++ b/Engine/Plugins/PCG/Source/PCG/Public/RuntimeGen/GenSources/PCGGenSourceBase.h
@@ -25,7 +25,7 @@ class IPCGGenSourceBase
 
 public:
 	/** Update the generation source so that it can cache data that is queried often (e.g. view frustum). Should be called every tick on any active generation sources. */
-	virtual void Tick() {}
+	PCG_API virtual void Tick(const FPCGRuntimeGenContext& InContext);
 
 	/** Returns the world space position of this gen source. */
 	virtual TOptional<FVector> GetPosition() const PURE_VIRTUAL(UPCGGenSourceBase::GetPosition, return TOptional<FVector>(););
@@ -43,4 +43,9 @@ public:
 	PCG_API virtual bool IsEquivalent(IPCGGenSourceBase* InOther, const FPCGRuntimeGenContext& InContext) const;
 
 	virtual FString GetDebugName() const { return TEXT("AnonymousGenSource"); }
+
+	// Deprecated section
+public:
+	UE_DEPRECATED(5.8, "Implement version that takes runtime generation context instead.")
+	virtual void Tick() {}
 };
diff --git a/Engine/Plugins/PCG/Source/PCG/Public/RuntimeGen/GenSources/PCGGenSourceEditorCamera.h b/Engine/Plugins/PCG/Source/PCG/Public/RuntimeGen/GenSources/PCGGenSourceEditorCamera.h
--- a/Engine/Plugins/PCG/Source/PCG/Public/RuntimeGen/GenSources/PCGGenSourceEditorCamera.h
+++ b/Engine/Plugins/PCG/Source/PCG/Public/RuntimeGen/GenSources/PCGGenSourceEditorCamera.h
@@ -21,7 +21,7 @@ class UPCGGenSourceEditorCamera : public UObject, public IPCGGenSourceBase
 
 public:
 	//~ Begin IPCGGenSourceInterface
-	PCG_API virtual void Tick() override;
+	PCG_API virtual void Tick(const FPCGRuntimeGenContext& InContext) override;
 	PCG_API virtual TOptional<FVector> GetPosition() const override;
 	PCG_API virtual TOptional<FVector> GetDirection() const override;
 	PCG_API virtual TOptional<FConvexVolume> GetViewFrustum(bool bIs2DGrid) const override;
diff --git a/Engine/Plugins/PCG/Source/PCG/Public/RuntimeGen/GenSources/PCGGenSourcePlayer.h b/Engine/Plugins/PCG/Source/PCG/Public/RuntimeGen/GenSources/PCGGenSourcePlayer.h
--- a/Engine/Plugins/PCG/Source/PCG/Public/RuntimeGen/GenSources/PCGGenSourcePlayer.h
+++ b/Engine/Plugins/PCG/Source/PCG/Public/RuntimeGen/GenSources/PCGGenSourcePlayer.h
@@ -18,7 +18,7 @@ class UPCGGenSourcePlayer : public UObject, public IPCGGenSourceBase
 
 public:
 	//~ Begin IPCGGenSourceInterface
-	PCG_API virtual void Tick() override;
+	PCG_API virtual void Tick(const FPCGRuntimeGenContext& InContext) override;
 	PCG_API virtual TOptional<FVector> GetPosition() const override;
 	PCG_API virtual TOptional<FVector> GetDirection() const override;
 	virtual TOptional<FConvexVolume> GetViewFrustum(bool bIs2DGrid) const override { return ViewFrustum; }
diff --git a/Engine/Plugins/PCG/Source/PCG/Public/RuntimeGen/PCGGenSourceManager.h b/Engine/Plugins/PCG/Source/PCG/Public/RuntimeGen/PCGGenSourceManager.h
--- a/Engine/Plugins/PCG/Source/PCG/Public/RuntimeGen/PCGGenSourceManager.h
+++ b/Engine/Plugins/PCG/Source/PCG/Public/RuntimeGen/PCGGenSourceManager.h
@@ -2,6 +2,7 @@
 
 #pragma once
 
+#include "RuntimeGen/PCGRuntimeGenContext.h"
 #include "RuntimeGen/GenSources/PCGGenSourceBase.h"
 
 class AController;
@@ -12,7 +13,6 @@ class UPCGGenSourceComponent;
 class UPCGGenSourceEditorCamera;
 class UPCGGenSourceWPStreamingSource;
 class UWorld;
-struct FPCGRuntimeGenContext;
 
 /**
  * The runtime Generation Source Manager tracks generation sources in the world for use by the Runtime Generation Scheduler.
@@ -20,7 +20,7 @@ struct FPCGRuntimeGenContext;
 class FPCGGenSourceManager
 {
 public:
-	FPCGGenSourceManager(const UWorld* InWorld);
+	FPCGGenSourceManager(const FPCGRuntimeGenContext& InContext);
 	~FPCGGenSourceManager();
 
 	/** Marks the GenSourceManager as dirty so that the next call to 'GetAllGenSources()' will update tracked generation sources. */
@@ -45,7 +45,6 @@ public:
 	void AddReferencedObjects(FReferenceCollector& Collector);
 
 protected:
-
 	/** Updates tracked generation sources that should be refreshed per tick. */
 	void UpdatePerTickGenSources(const APCGWorldActor* InPCGWorldActor);
 
@@ -70,6 +69,14 @@ PRAGMA_DISABLE_DEPRECATION_WARNINGS
 	bool bLoggedWorldPartitionGenSourceWarning = false;
 PRAGMA_ENABLE_DEPRECATION_WARNINGS
 
-	const UWorld* World = nullptr;
 	bool bDirty = false;
+	const FPCGRuntimeGenContext& Context;
+
+	// Deprecated section
+public:
+	UE_DEPRECATED(5.8, "Use constructor that takes runtime gen context.")
+	FPCGGenSourceManager(const UWorld* InWorld);
+
+private:
+	FPCGRuntimeGenContext FallBackContext;
 };
diff --git a/Engine/Plugins/PCG/Source/PCG/Public/RuntimeGen/PCGRuntimeGenContext.h b/Engine/Plugins/PCG/Source/PCG/Public/RuntimeGen/PCGRuntimeGenContext.h
--- a/Engine/Plugins/PCG/Source/PCG/Public/RuntimeGen/PCGRuntimeGenContext.h
+++ b/Engine/Plugins/PCG/Source/PCG/Public/RuntimeGen/PCGRuntimeGenContext.h
@@ -3,10 +3,12 @@
 #pragma once
 
 class APCGWorldActor;
+class UWorld;
 
 /** Runtime generation scheduler context information. */
 struct FPCGRuntimeGenContext
 {
+	UWorld* World = nullptr;
 	const APCGWorldActor* PCGWorldActor = nullptr;
 
 	bool bAnyRuntimeGenSourcesExist = false;
diff --git a/Engine/Plugins/PCG/Source/PCG/Public/RuntimeGen/PCGRuntimeGenScheduler.h b/Engine/Plugins/PCG/Source/PCG/Public/RuntimeGen/PCGRuntimeGenScheduler.h
--- a/Engine/Plugins/PCG/Source/PCG/Public/RuntimeGen/PCGRuntimeGenScheduler.h
+++ b/Engine/Plugins/PCG/Source/PCG/Public/RuntimeGen/PCGRuntimeGenScheduler.h
@@ -251,7 +251,6 @@ private:
 	FPCGTrackingManager* TrackingManager = nullptr;
 	FPCGGenSourceManager* GenSourceManager = nullptr;
 	UPCGSubsystem* Subsystem = nullptr;
-	UWorld* World = nullptr;
 	const UWorldPartitionSubsystem* WorldPartitionSubsystem = nullptr;
 
 	bool bPoolingWasEnabledLastFrame = true;
