[PCG] split editor permissions from specific feature flags

#rnx
#Synced-CL 47946549

diff --git a/Engine/Plugins/PCG/Source/PCG/Private/PCGComponent.cpp b/Engine/Plugins/PCG/Source/PCG/Private/PCGComponent.cpp
--- a/Engine/Plugins/PCG/Source/PCG/Private/PCGComponent.cpp
+++ b/Engine/Plugins/PCG/Source/PCG/Private/PCGComponent.cpp
@@ -48,7 +48,6 @@
 #include "EditorActorFolders.h"
 #include "ScopedTransaction.h"
 #include "Editor/EditorEngine.h"
-#include "Editor/IPCGEditorModule.h"
 #include "Editor/Transactor.h"
 #include "UObject/ObjectSaveContext.h"
 #include "WorldPartition/WorldPartition.h"
@@ -2022,8 +2021,7 @@ void UPCGComponent::BeginDestroy()
 bool UPCGComponent::IsEditorOnly() const
 {
 #if WITH_EDITOR
-	// Only allow non-editor components when running with all permissions
-	if (IPCGEditorModule* EditorModule = IPCGEditorModule::Get(); EditorModule && EditorModule->GetPermissionMode() != EPCGEditorPermissionMode::All)
+	if (FPCGModule::IsEditorOnlyGeneration())
 	{
 		return true;
 	}
diff --git a/Engine/Plugins/PCG/Source/PCG/Private/PCGGraph.cpp b/Engine/Plugins/PCG/Source/PCG/Private/PCGGraph.cpp
--- a/Engine/Plugins/PCG/Source/PCG/Private/PCGGraph.cpp
+++ b/Engine/Plugins/PCG/Source/PCG/Private/PCGGraph.cpp
@@ -593,6 +593,13 @@ void UPCGGraph::PostLoad()
 
 bool UPCGGraph::IsEditorOnly() const
 {
+#if WITH_EDITOR
+	if (FPCGModule::IsEditorOnlyGeneration())
+	{
+		return true;
+	}
+#endif
+
 	bool bIsCurrentlyEditorOnly = (Super::IsEditorOnly() || bIsEditorOnly);
 
 	if (!bIsCurrentlyEditorOnly)
diff --git a/Engine/Plugins/PCG/Source/PCG/Private/PCGModule.cpp b/Engine/Plugins/PCG/Source/PCG/Private/PCGModule.cpp
--- a/Engine/Plugins/PCG/Source/PCG/Private/PCGModule.cpp
+++ b/Engine/Plugins/PCG/Source/PCG/Private/PCGModule.cpp
@@ -301,6 +301,26 @@ void FPCGModule::DeregisterNativeElementDeterminismTests()
 {
 	PCGDeterminismTests::FNativeTestRegistry::DeregisterTestFunction(UPCGDifferenceSettings::StaticClass());
 }
+
+void FPCGModule::SetCPUOnlyExecution(bool bEnabled)
+{
+	FPCGModule::GetPCGModuleChecked().bIsCPUOnlyExecution = bEnabled;
+}
+
+bool FPCGModule::IsCPUOnlyExecution()
+{
+	return FPCGModule::GetPCGModuleChecked().bIsCPUOnlyExecution;
+}
+
+void FPCGModule::SetEditorOnlyGeneration(bool bEnabled)
+{
+	FPCGModule::GetPCGModuleChecked().bIsEditorOnlyGeneration = bEnabled;
+}
+
+bool FPCGModule::IsEditorOnlyGeneration()
+{
+	return FPCGModule::GetPCGModuleChecked().bIsEditorOnlyGeneration;
+}
 #endif // WITH_EDITOR
 
 IMPLEMENT_MODULE(FPCGModule, PCG);
diff --git a/Engine/Plugins/PCG/Source/PCG/Private/PCGSettings.cpp b/Engine/Plugins/PCG/Source/PCG/Private/PCGSettings.cpp
--- a/Engine/Plugins/PCG/Source/PCG/Private/PCGSettings.cpp
+++ b/Engine/Plugins/PCG/Source/PCG/Private/PCGSettings.cpp
@@ -29,7 +29,6 @@
 
 #if WITH_EDITOR
 #include "Editor.h"
-#include "Editor/IPCGEditorModule.h"
 #include "Hash/PCGSettingsHashContext.h"
 #endif
 
@@ -340,7 +339,7 @@ void UPCGSettings::PostLoad()
 		ExecutionMode_DEPRECATED = EPCGSettingsExecutionMode::Enabled;
 	}
 
-	if (IPCGEditorModule* EditorModule = IPCGEditorModule::Get(); EditorModule && EditorModule->GetPermissionMode() != EPCGEditorPermissionMode::All)
+	if (FPCGModule::IsCPUOnlyExecution())
 	{
 		bExecuteOnGPU = false;
 	}
@@ -1070,7 +1069,7 @@ TArray<FPCGSettingsOverridableParam> UPCGSettings::GatherOverridableParams() con
 
 bool UPCGSettings::DisplayExecuteOnGPUSettingCondition() const
 {
-	if (IPCGEditorModule* EditorModule = IPCGEditorModule::Get(); EditorModule && EditorModule->GetPermissionMode() != EPCGEditorPermissionMode::All)
+	if (FPCGModule::IsCPUOnlyExecution())
 	{
 		return false;
 	}
@@ -1082,8 +1081,7 @@ bool UPCGSettings::DisplayExecuteOnGPUSettingCondition() const
 bool UPCGSettings::ShouldExecuteOnGPU() const
 {
 #if WITH_EDITOR
-	// In case property was changed and then permissions were lowered
-	if (IPCGEditorModule* EditorModule = IPCGEditorModule::Get(); EditorModule && EditorModule->GetPermissionMode() != EPCGEditorPermissionMode::All)
+	if (FPCGModule::IsCPUOnlyExecution())
 	{
 		return false;
 	}
diff --git a/Engine/Plugins/PCG/Source/PCG/Public/PCGModule.h b/Engine/Plugins/PCG/Source/PCG/Public/PCGModule.h
--- a/Engine/Plugins/PCG/Source/PCG/Public/PCGModule.h
+++ b/Engine/Plugins/PCG/Source/PCG/Public/PCGModule.h
@@ -47,6 +47,14 @@ public:
 	void PreExit();
 	void OnPostInitEngine();
 
+#if WITH_EDITOR
+	PCG_API static void SetCPUOnlyExecution(bool bEnabled);
+	PCG_API static bool IsCPUOnlyExecution();
+
+	PCG_API static void SetEditorOnlyGeneration(bool bEnabled);
+	PCG_API static bool IsEditorOnlyGeneration();
+#endif
+
 	PCG_API static FPCGModule& GetPCGModuleChecked();
 	PCG_API static FPCGModule* GetPCGModule();
 	static const FPCGGetDataFunctionRegistry& ConstGetDataFunctionRegistry() { return GetPCGModuleChecked().GetDataFunctionRegistry; }
@@ -120,5 +128,8 @@ private:
 	void NotifyGraphChanged(UPCGGraph* InGraph, EPCGChangeType ChangeType) { OnGraphChangedDelegate.Broadcast(InGraph, ChangeType); }
 
 	static FPCGGraphChangedDelegate OnGraphChangedDelegate;
+
+	bool bIsCPUOnlyExecution = false;
+	bool bIsEditorOnlyGeneration = false;
 #endif // WITH_EDITOR
 };
