[PCG] DataFromActor/ActorSelector support for IPCGGraphExecutionSource

#Synced-CL 48718022

diff --git a/Engine/Plugins/PCG/Source/PCG/Private/ChangeTracking/PCGChangeTrackingRegistry.cpp b/Engine/Plugins/PCG/Source/PCG/Private/ChangeTracking/PCGChangeTrackingRegistry.cpp
--- a/Engine/Plugins/PCG/Source/PCG/Private/ChangeTracking/PCGChangeTrackingRegistry.cpp
+++ b/Engine/Plugins/PCG/Source/PCG/Private/ChangeTracking/PCGChangeTrackingRegistry.cpp
@@ -50,4 +50,38 @@ TArray<TUniquePtr<IPCGChangeHandler>> FPCGChangeTrackingRegistry::CreateChangeHa
 	return ChangeHandlers;
 }
 
-#endif // WITH_EDITOR
\ No newline at end of file
+#endif // WITH_EDITOR
+
+TArray<IPCGGraphExecutionSource*> FPCGChangeTrackingRegistry::GetExecutionSourcesFromSelectionKey(const FPCGGetExecutionSourcesFromSelectionKeyParams& InParams) const
+{
+	const UObject* SelectedObject = InParams.SelectionKey.GetObjectFromPath();
+	if (SelectedObject)
+	{
+		UClass* ObjectClass = SelectedObject->GetClass();
+		while (ObjectClass)
+		{
+			if (const FPCGGetExecutionSourcesFromSelectionKey* Found = GetExecutionSourcesFromSelectionKeyEntries.Find(ObjectClass))
+			{
+				if (ensure(Found->IsBound()))
+				{
+					return Found->Execute(InParams);
+				}
+			}
+
+			ObjectClass = ObjectClass->GetSuperClass();
+		}
+	}
+
+	return {};
+}
+
+void FPCGChangeTrackingRegistry::RegisterGetExecutionSourcesFromSelectionKey(UClass* InSelectedObjectClass, FPCGGetExecutionSourcesFromSelectionKey InGetExecutionSourcesFromKey)
+{
+	check(!GetExecutionSourcesFromSelectionKeyEntries.Contains(InSelectedObjectClass));
+	GetExecutionSourcesFromSelectionKeyEntries.Add(InSelectedObjectClass, MoveTemp(InGetExecutionSourcesFromKey));
+}
+
+void FPCGChangeTrackingRegistry::UnregisterGetExecutionSourcesFromSelectionKey(UClass* InSelectedObjectClass)
+{
+	GetExecutionSourcesFromSelectionKeyEntries.Remove(InSelectedObjectClass);
+}
diff --git a/Engine/Plugins/PCG/Source/PCG/Private/Elements/PCGActorSelector.cpp b/Engine/Plugins/PCG/Source/PCG/Private/Elements/PCGActorSelector.cpp
--- a/Engine/Plugins/PCG/Source/PCG/Private/Elements/PCGActorSelector.cpp
+++ b/Engine/Plugins/PCG/Source/PCG/Private/Elements/PCGActorSelector.cpp
@@ -117,12 +117,13 @@ namespace PCGActorSelector
 		return {};
 	}
 
-	TArray<AActor*> FindActors(const FPCGActorSelectorSettings& Settings, const UPCGComponent* InComponent, const TFunction<bool(const AActor*)>& BoundsCheck, const TFunction<bool(const AActor*)>& SelfIgnoreCheck, TArrayView<AActor*> InputActors)
+	TArray<AActor*> FindActors(const FPCGActorSelectorSettings& Settings, const IPCGGraphExecutionSource* InExecutionSource, const TFunction<bool(const AActor*)>& BoundsCheck, const TFunction<bool(const AActor*)>& SelfIgnoreCheck, TArrayView<AActor*> InputActors)
 	{
 		TRACE_CPUPROFILER_EVENT_SCOPE(PCGActorSelector::FindActor);
 
-		UWorld* World = InComponent ? InComponent->GetWorld() : nullptr;
-		AActor* Self = InComponent ? InComponent->GetOwner() : nullptr;
+		UWorld* World = InExecutionSource ? InExecutionSource->GetExecutionState().GetWorld() : nullptr;
+		const UPCGComponent* PCGComponent = Cast<UPCGComponent>(InExecutionSource);
+		AActor* Self = PCGComponent ? PCGComponent->GetOwner() : nullptr;
 
 		TArray<AActor*> FoundActors;
 
@@ -207,7 +208,7 @@ namespace PCGActorSelector
 		case EPCGActorFilter::Original:
 		{
 			APCGPartitionActor* PartitionActor = Cast<APCGPartitionActor>(Self);
-			UPCGComponent* OriginalComponent = (PartitionActor && InComponent) ? PartitionActor->GetOriginalComponent(InComponent) : nullptr;
+			UPCGComponent* OriginalComponent = (PartitionActor && PCGComponent) ? PartitionActor->GetOriginalComponent(PCGComponent) : nullptr;
 			AActor* OriginalActor = OriginalComponent ? OriginalComponent->GetOwner() : nullptr;
 			if (OriginalActor)
 			{
@@ -217,6 +218,8 @@ namespace PCGActorSelector
 			{
 				ActorsToCheck.Add(Self);
 			}
+
+			break;
 		}
 
 		case EPCGActorFilter::FromInput:
@@ -303,12 +306,12 @@ namespace PCGActorSelector
 		return FilteredActors;
 	}
 
-	TArray<AActor*> FindActors(const FPCGActorSelectorSettings* ActorSettings, const FPCGComponentSelectorSettings* ComponentSettings, const UPCGComponent* InComponent, const TFunction<bool(const AActor*)>& BoundsCheck, const TFunction<bool(const AActor*)>& SelfIgnoreCheck, TArrayView<AActor*> InputActors)
+	TArray<AActor*> FindActors(const FPCGActorSelectorSettings* ActorSettings, const FPCGComponentSelectorSettings* ComponentSettings, const IPCGGraphExecutionSource* InExecutionSource, const TFunction<bool(const AActor*)>& BoundsCheck, const TFunction<bool(const AActor*)>& SelfIgnoreCheck, TArrayView<AActor*> InputActors)
 	{
 		TArray<AActor*> FoundActors;
 		if (ActorSettings)
 		{
-			FoundActors = FindActors(*ActorSettings, InComponent, BoundsCheck, SelfIgnoreCheck, InputActors);
+			FoundActors = FindActors(*ActorSettings, InExecutionSource, BoundsCheck, SelfIgnoreCheck, InputActors);
 		}
 		else
 		{
@@ -323,13 +326,13 @@ namespace PCGActorSelector
 		return FoundActors;
 	}
 
-	AActor* FindActor(const FPCGActorSelectorSettings& InSettings, UPCGComponent* InComponent, const TFunction<bool(const AActor*)>& BoundsCheck, const TFunction<bool(const AActor*)>& SelfIgnoreCheck, TArrayView<AActor*> InputActors)
+	AActor* FindActor(const FPCGActorSelectorSettings& InSettings, const IPCGGraphExecutionSource* InExecutionSource, const TFunction<bool(const AActor*)>& BoundsCheck, const TFunction<bool(const AActor*)>& SelfIgnoreCheck, TArrayView<AActor*> InputActors)
 	{
 		// In order to make sure we don't try to select multiple, we'll do a copy of the settings here.
 		FPCGActorSelectorSettings Settings = InSettings;
 		Settings.bSelectMultiple = false;
 
-		TArray<AActor*> Actors = FindActors(Settings, InComponent, BoundsCheck, SelfIgnoreCheck, InputActors);
+		TArray<AActor*> Actors = FindActors(Settings, InExecutionSource, BoundsCheck, SelfIgnoreCheck, InputActors);
 		return Actors.IsEmpty() ? nullptr : Actors[0];
 	}
 }
@@ -532,50 +535,23 @@ bool FPCGSelectionKey::IsMatching(const FPCGSelectionKey& InSelectionKey, const
 	// Since we can have a lot of components in InComponents, we go the other way around (Actor to component)
 	if (ActorFilter != EPCGActorFilter::AllWorldActors)
 	{
+		FPCGGetExecutionSourcesFromSelectionKeyParams Params(InSelectionKey, ActorFilter);
+		TArray<IPCGGraphExecutionSource*> FoundExecutionSources = FPCGModule::GetPCGModuleChecked().GetConstChangeTrackingRegistry().GetExecutionSourcesFromSelectionKey(Params);
+		
 		bool bFoundMatch = false;
 
-		const AActor* InActor = Cast<const AActor>(SelectionObject);
-
-		if (!InActor)
-		{
-			return false;
-		}
-
-		TArray<UActorComponent*, TInlineAllocator<64>> ActorComponents;
-
-		if (ActorFilter == EPCGActorFilter::Self || ActorFilter == EPCGActorFilter::Original)
+		for (IPCGGraphExecutionSource* FoundExecutionSource : FoundExecutionSources)
 		{
-			InActor->GetComponents(UPCGComponent::StaticClass(), ActorComponents);
-		}
-		else if (ActorFilter == EPCGActorFilter::Parent || (ActorFilter == EPCGActorFilter::Root && !InActor->GetParentActor()))
-		{
-			TArray<AActor*> ActorsToCheck;
-			InActor->GetAllChildActors(ActorsToCheck, /*bIncludeDescendants=*/ActorFilter == EPCGActorFilter::Root);
-			ActorsToCheck.Add(const_cast<AActor*>(InActor));
-			TArray<UActorComponent*, TInlineAllocator<64>> TempActorComponents;
-			for (AActor* Current : ActorsToCheck)
+			if (InExecutionSources.Contains(FoundExecutionSource))
 			{
-				// TempActorComponents is reset in GetComponents
-				Current->GetComponents(UPCGComponent::StaticClass(), TempActorComponents);
-				ActorComponents.Append(TempActorComponents);
-			}
-		}
-
-		for (UActorComponent* Component : ActorComponents)
-		{
-			if (IPCGGraphExecutionSource* ExecutionSource = Cast<IPCGGraphExecutionSource>(Component))
-			{
-				if (InExecutionSources.Contains(ExecutionSource))
+				bFoundMatch = true;
+				if (OptionalMatchedExecutionSources)
 				{
-					bFoundMatch = true;
-					if (OptionalMatchedExecutionSources)
-					{
-						OptionalMatchedExecutionSources->Add(ExecutionSource);
-					}
-					else
-					{
-						break;
-					}
+					OptionalMatchedExecutionSources->Add(FoundExecutionSource);
+				}
+				else
+				{
+					break;
 				}
 			}
 		}
diff --git a/Engine/Plugins/PCG/Source/PCG/Private/Elements/PCGDataFromActor.cpp b/Engine/Plugins/PCG/Source/PCG/Private/Elements/PCGDataFromActor.cpp
--- a/Engine/Plugins/PCG/Source/PCG/Private/Elements/PCGDataFromActor.cpp
+++ b/Engine/Plugins/PCG/Source/PCG/Private/Elements/PCGDataFromActor.cpp
@@ -360,27 +360,33 @@ bool FPCGDataFromActorElement::ExecuteInternal(FPCGContext* InContext) const
 	const UPCGDataFromActorSettings* Settings = Context->GetInputSettings<UPCGDataFromActorSettings>();
 	check(Settings);
 
-	UPCGComponent* PCGComponent = Cast<UPCGComponent>(Context->ExecutionSource.Get());
+	IPCGGraphExecutionSource* ExecutionSource = Context->ExecutionSource.Get();
+	if (!ExecutionSource)
+	{
+		return true;
+	}
+
+	UPCGComponent* PCGComponent = Cast<UPCGComponent>(ExecutionSource);
+	AActor* Self = PCGComponent ? PCGComponent->GetOwner() : nullptr;
 
 	if (!Context->bPerformedQuery)
 	{
 		Context->ComponentSelector = Settings->ComponentSelector;
 
 		TFunction<bool(const AActor*)> BoundsCheck = [](const AActor*) -> bool { return true; };
-		const AActor* Self = PCGComponent ? PCGComponent->GetOwner() : nullptr;
-		if (Self && Settings->ActorSelector.bMustOverlapSelf)
+		if (Settings->ActorSelector.bMustOverlapSelf)
 		{
-			// Capture ActorBounds by value because it goes out of scope
-			FBox ActorBounds = PCGHelpers::GetGridBounds(Self, PCGComponent);
-			BoundsCheck = [Settings, ActorBounds, PCGComponent](const AActor* OtherActor) -> bool
+			// Capture Bounds by value because it goes out of scope
+			FBox Bounds = ExecutionSource->GetExecutionState().GetBounds();
+			BoundsCheck = [Settings, Bounds, ExecutionSource](const AActor* OtherActor) -> bool
 			{
-				const FBox OtherActorBounds = OtherActor ? PCGHelpers::GetGridBounds(OtherActor, PCGComponent) : FBox(EForceInit::ForceInit);
-				return ActorBounds.Intersect(OtherActorBounds);
+				const FBox OtherActorBounds = OtherActor ? PCGHelpers::GetGridBounds(OtherActor, Cast<UPCGComponent>(ExecutionSource)) : FBox(EForceInit::ForceInit);
+				return Bounds.Intersect(OtherActorBounds);
 			};
 		}
 
 		TFunction<bool(const AActor*)> SelfIgnoreCheck = [](const AActor*) -> bool { return true; };
-		if (Self && Settings->ActorSelector.bIgnoreSelfAndChildren)
+		if (Settings->ActorSelector.bIgnoreSelfAndChildren)
 		{
 			SelfIgnoreCheck = [Self](const AActor* OtherActor) -> bool
 			{
@@ -416,8 +422,7 @@ bool FPCGDataFromActorElement::ExecuteInternal(FPCGContext* InContext) const
 		// or just gather all registered components (which is way faster than going through all actors in the world).
 		if (Settings->Mode == EPCGGetDataFromActorMode::GetDataFromPCGComponent && Settings->ActorSelector.ActorFilter == EPCGActorFilter::AllWorldActors)
 		{
-			UPCGSubsystem* Subsystem = PCGComponent ? PCGComponent->GetSubsystem() : nullptr;
-			if (Subsystem)
+			if (UPCGSubsystem* Subsystem = Cast<UPCGSubsystem>(ExecutionSource->GetExecutionState().GetSubsystem()))
 			{
 				TRACE_CPUPROFILER_EVENT_SCOPE(FPCGDataFromActorElement::Execute::FindPCGComponents);
 
@@ -426,10 +431,10 @@ bool FPCGDataFromActorElement::ExecuteInternal(FPCGContext* InContext) const
 				// TODO: Perhaps move the logic into the selector.
 				if (Settings->ActorSelector.bMustOverlapSelf)
 				{
-					FBox ActorBounds = PCGHelpers::GetGridBounds(Self, PCGComponent);
-					for (IPCGGraphExecutionSource* ExecutionSource : Subsystem->GetAllIntersectingExecutionSources(ActorBounds))
+					const FBox Bounds = ExecutionSource->GetExecutionState().GetBounds();
+					for (IPCGGraphExecutionSource* CurrentExecutionSource : Subsystem->GetAllIntersectingExecutionSources(Bounds))
 					{
-						if (UPCGComponent* Component = Cast<UPCGComponent>(ExecutionSource))
+						if (UPCGComponent* Component = Cast<UPCGComponent>(CurrentExecutionSource))
 						{
 							if (AActor* Actor = Component->GetOwner())
 							{
@@ -443,9 +448,9 @@ bool FPCGDataFromActorElement::ExecuteInternal(FPCGContext* InContext) const
 				}
 				else
 				{
-					for (IPCGGraphExecutionSource* ExecutionSource : Subsystem->GetAllRegisteredExecutionSources())
+					for (IPCGGraphExecutionSource* CurrentExecutionSource : Subsystem->GetAllRegisteredExecutionSources())
 					{
-						if (UPCGComponent* Component = Cast<UPCGComponent>(ExecutionSource))
+						if (UPCGComponent* Component = Cast<UPCGComponent>(CurrentExecutionSource))
 						{
 							if (AActor* Actor = Component->GetOwner())
 							{
@@ -487,7 +492,7 @@ bool FPCGDataFromActorElement::ExecuteInternal(FPCGContext* InContext) const
 
 			Context->ComponentSelector.ComponentList = std::move(ComponentsFromInput);
 
-			Context->FoundActors = PCGActorSelector::FindActors(&Settings->ActorSelector, &Context->ComponentSelector, PCGComponent, BoundsCheck, SelfIgnoreCheck, ActorsFromInput);
+			Context->FoundActors = PCGActorSelector::FindActors(&Settings->ActorSelector, &Context->ComponentSelector, ExecutionSource, BoundsCheck, SelfIgnoreCheck, ActorsFromInput);
 			Context->bPerformedQuery = true;
 
 #if WITH_EDITOR
@@ -593,7 +598,7 @@ bool FPCGDataFromActorElement::ExecuteInternal(FPCGContext* InContext) const
 #if WITH_EDITOR
 		// Remove ignored change origins now that we've completed the wait tasks.
 		UPCGComponent* OriginalComponent = PCGComponent ? PCGComponent->GetOriginalComponent() : nullptr;
-		if (ensure(OriginalComponent))
+		if (OriginalComponent)
 		{
 			for (TObjectKey<UObject>& IgnoredChangeOriginKey : Context->IgnoredChangeOrigins)
 			{
@@ -905,7 +910,7 @@ void FPCGDataFromActorElement::ProcessActor(FPCGContext* InContext, const UPCGDa
 	FPCGDataFromActorContext* Context = static_cast<FPCGDataFromActorContext*>(InContext);
 	check(Settings);
 
-	UPCGComponent* SourceComponent = CastChecked<UPCGComponent>(Context->ExecutionSource.Get());
+	UPCGComponent* SourceComponent = Cast<UPCGComponent>(Context->ExecutionSource.Get());
 	const UPCGComponent* SourceOriginalComponent = SourceComponent ? SourceComponent->GetOriginalComponent() : nullptr;
 
 	if (!FoundActor || !IsValid(FoundActor) || !SourceOriginalComponent)
diff --git a/Engine/Plugins/PCG/Source/PCG/Private/Helpers/PCGActorHelpers.cpp b/Engine/Plugins/PCG/Source/PCG/Private/Helpers/PCGActorHelpers.cpp
--- a/Engine/Plugins/PCG/Source/PCG/Private/Helpers/PCGActorHelpers.cpp
+++ b/Engine/Plugins/PCG/Source/PCG/Private/Helpers/PCGActorHelpers.cpp
@@ -21,6 +21,8 @@
 #include UE_INLINE_GENERATED_CPP_BY_NAME(PCGActorHelpers)
 
 #if WITH_EDITOR
+#include "ChangeTracking/PCGChangeTrackingRegistry.h"
+
 #include "Editor.h"
 #include "Selection.h"
 #include "UnrealEdGlobals.h"
@@ -944,3 +946,45 @@ AActor* UPCGActorHelpers::IsolateFromActor(AActor* Actor, const TArray<FName>& R
 
 	return IsolatedActor;
 }
+
+TArray<IPCGGraphExecutionSource*> UPCGActorHelpers::GetExecutionSourcesFromSelectionKey(const FPCGGetExecutionSourcesFromSelectionKeyParams& InParams)
+{
+	const AActor* InActor = Cast<const AActor>(InParams.SelectionKey.GetObjectFromPath());
+
+	if (!InActor)
+	{
+		return TArray<IPCGGraphExecutionSource*>();
+	}
+
+	TArray<UActorComponent*, TInlineAllocator<64>> ActorComponents;
+
+	if (InParams.ActorFilter == EPCGActorFilter::Self || InParams.ActorFilter == EPCGActorFilter::Original)
+	{
+		InActor->GetComponents(UPCGComponent::StaticClass(), ActorComponents);
+	}
+	else if (InParams.ActorFilter == EPCGActorFilter::Parent || (InParams.ActorFilter == EPCGActorFilter::Root && !InActor->GetParentActor()))
+	{
+		TArray<AActor*> ActorsToCheck;
+		InActor->GetAllChildActors(ActorsToCheck, /*bIncludeDescendants=*/InParams.ActorFilter == EPCGActorFilter::Root);
+		ActorsToCheck.Add(const_cast<AActor*>(InActor));
+		TArray<UActorComponent*, TInlineAllocator<64>> TempActorComponents;
+		for (AActor* Current : ActorsToCheck)
+		{
+			// TempActorComponents is reset in GetComponents
+			Current->GetComponents(UPCGComponent::StaticClass(), TempActorComponents);
+			ActorComponents.Append(TempActorComponents);
+		}
+	}
+
+	TArray<IPCGGraphExecutionSource*> FoundExecutionSources;
+	FoundExecutionSources.Reserve(ActorComponents.Num());
+	for (UActorComponent* ActorComponent : ActorComponents)
+	{
+		if (IPCGGraphExecutionSource* ExecutionSource = Cast<IPCGGraphExecutionSource>(ActorComponent))
+		{
+			FoundExecutionSources.Add(ExecutionSource);
+		}
+	}
+
+	return FoundExecutionSources;
+}
diff --git a/Engine/Plugins/PCG/Source/PCG/Private/PCGModule.cpp b/Engine/Plugins/PCG/Source/PCG/Private/PCGModule.cpp
--- a/Engine/Plugins/PCG/Source/PCG/Private/PCGModule.cpp
+++ b/Engine/Plugins/PCG/Source/PCG/Private/PCGModule.cpp
@@ -17,6 +17,7 @@
 #include "Data/DataView/PCGDataViewInterface.h"
 #include "Data/DataView/PCGDataViewNativeJsonConverters.h"
 #include "Data/DataView/PCGDataViewNativePropertySelectors.h"
+#include "Helpers/PCGActorHelpers.h"
 
 #include "Engine/Engine.h"
 #include "Interfaces/IPluginManager.h"
@@ -189,6 +190,8 @@ void FPCGModule::StartupModule()
 	DataViewRegistry.RegisterConverter(UPCGParamData::StaticClass(), UPCGDataViewJsonConverter::StaticClass());
 	DataViewRegistry.RegisterConverter(UPCGSpatialData::StaticClass(), UPCGDataViewJsonConverter::StaticClass());
 
+	ChangeTrackingRegistry.RegisterGetExecutionSourcesFromSelectionKey(AActor::StaticClass(), FPCGGetExecutionSourcesFromSelectionKey::CreateStatic(&UPCGActorHelpers::GetExecutionSourcesFromSelectionKey));
+	
 #if WITH_EDITOR
 	ObjectHashFactory.RegisterObjectHashContextFactory(UPCGGraphInterface::StaticClass(), FPCGOnCreateObjectHashContext::CreateStatic(&FPCGGraphHashContext::MakeInstance));
 	ObjectHashFactory.RegisterObjectHashContextFactory(UPCGSettingsInterface::StaticClass(), FPCGOnCreateObjectHashContext::CreateStatic(&FPCGSettingsHashContext::MakeInstance));
@@ -232,6 +235,8 @@ void FPCGModule::PreExit()
 
 	GraphExecutionRegistry.Shutdown();
 
+	ChangeTrackingRegistry.UnregisterGetExecutionSourcesFromSelectionKey(AActor::StaticClass());
+
 #if WITH_EDITOR
 	DeregisterNativeElementDeterminismTests();
 
diff --git a/Engine/Plugins/PCG/Source/PCG/Public/ChangeTracking/PCGChangeTrackingRegistry.h b/Engine/Plugins/PCG/Source/PCG/Public/ChangeTracking/PCGChangeTrackingRegistry.h
--- a/Engine/Plugins/PCG/Source/PCG/Public/ChangeTracking/PCGChangeTrackingRegistry.h
+++ b/Engine/Plugins/PCG/Source/PCG/Public/ChangeTracking/PCGChangeTrackingRegistry.h
@@ -2,14 +2,14 @@
 
 #pragma once
 
-#if WITH_EDITOR
-
 #include "PCGCommon.h"
 #include "Delegates/Delegate.h"
 #include "Elements/PCGActorSelector.h"
 
 class FPCGTrackingManager;
 
+#if WITH_EDITOR
+
 // Unique change types
 using FPCGTrackingChangeID = FGuid;
 
@@ -78,9 +78,30 @@ protected:
 DECLARE_DELEGATE_RetVal_OneParam(TUniquePtr<IPCGChangeTracker>, FPCGCreateTrackerInstance, FPCGTrackingManager*);
 DECLARE_DELEGATE_RetVal_OneParam(TUniquePtr<IPCGChangeHandler>, FPCGCreateHandlerInstance, FPCGTrackingManager*);
 
+#endif // WITH_EDITOR
+
+// FPCGGetExecutionSourcesFromSelectionKey params
+struct FPCGGetExecutionSourcesFromSelectionKeyParams
+{
+	FPCGGetExecutionSourcesFromSelectionKeyParams() = delete;
+	FPCGGetExecutionSourcesFromSelectionKeyParams(const FPCGSelectionKey& InSelectionKey, EPCGActorFilter InActorFilter)
+		: SelectionKey(InSelectionKey)
+		, ActorFilter(InActorFilter)
+	{
+
+	}
+
+	const FPCGSelectionKey& SelectionKey;
+	EPCGActorFilter ActorFilter;
+};
+
+DECLARE_DELEGATE_RetVal_OneParam(TArray<IPCGGraphExecutionSource*>, FPCGGetExecutionSourcesFromSelectionKey, const FPCGGetExecutionSourcesFromSelectionKeyParams&);
+
 class FPCGChangeTrackingRegistry
 {
 public:
+
+#if WITH_EDITOR
 	// Creates the tracker instances for the FPCGTrackingManager
 	TArray<TUniquePtr<IPCGChangeTracker>> CreateChangeTrackers(FPCGTrackingManager* InManager) const;
 	TArray<TUniquePtr<IPCGChangeHandler>> CreateChangeHandlers(FPCGTrackingManager* InManager) const;
@@ -92,10 +113,20 @@ public:
 	// Register/unregister custom handler factories (allows other modules to register new handlers)
 	PCG_API void RegisterHandlerFactory(FName InName, FPCGCreateHandlerInstance InHandlerFactory);
 	PCG_API void UnregisterHandlerFactory(FName InName);
+#endif
+
+	TArray<IPCGGraphExecutionSource*> GetExecutionSourcesFromSelectionKey(const FPCGGetExecutionSourcesFromSelectionKeyParams& InParams) const;
+
+	// Register/unregister execution source resolving
+	PCG_API void RegisterGetExecutionSourcesFromSelectionKey(UClass* InSelectedObjectClass, FPCGGetExecutionSourcesFromSelectionKey InGetExecutionSourcesFromKey);
+	PCG_API void UnregisterGetExecutionSourcesFromSelectionKey(UClass* InSelectedObjectClass);
 
 private:
+
+#if WITH_EDITOR
 	TMap<FName, FPCGCreateTrackerInstance> TrackerFactories;
 	TMap<FName, FPCGCreateHandlerInstance> HandlerFactories;
-};
+#endif
 
-#endif // WITH_EDITOR
\ No newline at end of file
+	TMap<UClass*, FPCGGetExecutionSourcesFromSelectionKey> GetExecutionSourcesFromSelectionKeyEntries;
+};
diff --git a/Engine/Plugins/PCG/Source/PCG/Public/Elements/PCGActorSelector.h b/Engine/Plugins/PCG/Source/PCG/Public/Elements/PCGActorSelector.h
--- a/Engine/Plugins/PCG/Source/PCG/Public/Elements/PCGActorSelector.h
+++ b/Engine/Plugins/PCG/Source/PCG/Public/Elements/PCGActorSelector.h
@@ -334,12 +334,11 @@ protected:
 namespace PCGActorSelector
 {
 	// Simple actor filtering
-	PCG_API TArray<AActor*> FindActors(const FPCGActorSelectorSettings& Settings, const UPCGComponent* InComponent, const TFunction<bool(const AActor*)>& BoundsCheck, const TFunction<bool(const AActor*)>& SelfIgnoreCheck, TArrayView<AActor*> InputActors = TArrayView<AActor*>());
-	PCG_API AActor* FindActor(const FPCGActorSelectorSettings& InSettings, UPCGComponent* InComponent, const TFunction<bool(const AActor*)>& BoundsCheck, const TFunction<bool(const AActor*)>& SelfIgnoreCheck, TArrayView<AActor*> InputActors = TArrayView<AActor*>());
+	PCG_API TArray<AActor*> FindActors(const FPCGActorSelectorSettings& Settings, const IPCGGraphExecutionSource* InExecutionSource, const TFunction<bool(const AActor*)>& BoundsCheck, const TFunction<bool(const AActor*)>& SelfIgnoreCheck, TArrayView<AActor*> InputActors = TArrayView<AActor*>());
+	PCG_API AActor* FindActor(const FPCGActorSelectorSettings& InSettings, const IPCGGraphExecutionSource* InExecutionSource, const TFunction<bool(const AActor*)>& BoundsCheck, const TFunction<bool(const AActor*)>& SelfIgnoreCheck, TArrayView<AActor*> InputActors = TArrayView<AActor*>());
 
 	// Actor + Component filtering
-	PCG_API TArray<AActor*> FindActors(const FPCGActorSelectorSettings* ActorSettings, const FPCGComponentSelectorSettings* ComponentSettings, const UPCGComponent* InComponent, const TFunction<bool(const AActor*)>& BoundsCheck, const TFunction<bool(const AActor*)>& SelfIgnoreCheck, TArrayView<AActor*> InputActors = TArrayView<AActor*>());
-	PCG_API AActor* FindActor(const FPCGActorSelectorSettings* ActorSettings, const FPCGComponentSelectorSettings* ComponentSettings, const UPCGComponent* InComponent, const TFunction<bool(const AActor*)>& BoundsCheck, const TFunction<bool(const AActor*)>& SelfIgnoreCheck, TArrayView<AActor*> InputActors = TArrayView<AActor*>());
+	PCG_API TArray<AActor*> FindActors(const FPCGActorSelectorSettings* ActorSettings, const FPCGComponentSelectorSettings* ComponentSettings, const IPCGGraphExecutionSource* InExecutionSource, const TFunction<bool(const AActor*)>& BoundsCheck, const TFunction<bool(const AActor*)>& SelfIgnoreCheck, TArrayView<AActor*> InputActors = TArrayView<AActor*>());
 	
 	// Additional filter methods
 	PCG_API TArray<AActor*> FilterActors(const FPCGComponentSelectorSettings& ComponentSettings, TArrayView<AActor*> InActors);
diff --git a/Engine/Plugins/PCG/Source/PCG/Public/Helpers/PCGActorHelpers.h b/Engine/Plugins/PCG/Source/PCG/Public/Helpers/PCGActorHelpers.h
--- a/Engine/Plugins/PCG/Source/PCG/Public/Helpers/PCGActorHelpers.h
+++ b/Engine/Plugins/PCG/Source/PCG/Public/Helpers/PCGActorHelpers.h
@@ -17,11 +17,11 @@
 
 #include "PCGActorHelpers.generated.h"
 
-#define UE_API PCG_API
-
 class AActor;
 class APCGIsolatedActor;
+class IPCGGraphExecutionSource;
 struct FPCGContext;
+struct FPCGGetExecutionSourcesFromSelectionKeyParams;
 class UActorComponent;
 class UDataLayerInstance;
 class UHLODLayer;
@@ -127,27 +127,27 @@ class UPCGActorHelpers : public UBlueprintFunctionLibrary
 	GENERATED_BODY()
 
 public:
-UE_API PRAGMA_DISABLE_DEPRECATION_WARNINGS
-	static UInstancedStaticMeshComponent* GetOrCreateISMC(AActor* InTargetActor, UPCGComponent* SourceComponent, uint64 SettingsUID, const FPCGISMCBuilderParameters& Params);
-	static UE_API UPCGManagedISMComponent* GetOrCreateManagedISMC(AActor* InTargetActor, UPCGComponent* SourceComponent, uint64 SettingsUID, const FPCGISMCBuilderParameters& Params);
+PRAGMA_DISABLE_DEPRECATION_WARNINGS
+	static PCG_API UInstancedStaticMeshComponent* GetOrCreateISMC(AActor* InTargetActor, UPCGComponent* SourceComponent, uint64 SettingsUID, const FPCGISMCBuilderParameters& Params);
+	static PCG_API UPCGManagedISMComponent* GetOrCreateManagedISMC(AActor* InTargetActor, UPCGComponent* SourceComponent, uint64 SettingsUID, const FPCGISMCBuilderParameters& Params);
 PRAGMA_ENABLE_DEPRECATION_WARNINGS
 
 	UE_DEPRECATED(5.6, "Use signature with no SettingsUID. Make sure to fill Params.SettingsCRC to enable component reuse.")
-	static UE_API UInstancedStaticMeshComponent* GetOrCreateISMC(AActor* InTargetActor, UPCGComponent* SourceComponent, uint64 SettingsUID, const FPCGISMComponentBuilderParams& Params);
+	static PCG_API UInstancedStaticMeshComponent* GetOrCreateISMC(AActor* InTargetActor, UPCGComponent* SourceComponent, uint64 SettingsUID, const FPCGISMComponentBuilderParams& Params);
 	UE_DEPRECATED(5.6, "Use signature with no SettingsUID. Make sure to fill Params.SettingsCRC to enable component reuse.")	
-	static UE_API UPCGManagedISMComponent* GetOrCreateManagedISMC(AActor* InTargetActor, UPCGComponent* SourceComponent, uint64 SettingsUID, const FPCGISMComponentBuilderParams& Params);
+	static PCG_API UPCGManagedISMComponent* GetOrCreateManagedISMC(AActor* InTargetActor, UPCGComponent* SourceComponent, uint64 SettingsUID, const FPCGISMComponentBuilderParams& Params);
 	UE_DEPRECATED(5.6, "Use signature with no SettingsUID. Make sure to fill Params.SettingsCRC to enable component reuse.")
-	static UE_API USplineMeshComponent* GetOrCreateSplineMeshComponent(AActor* InTargetActor, UPCGComponent* SourceComponent, uint64 SettingsUID, const FPCGSplineMeshComponentBuilderParameters& Params);
+	static PCG_API USplineMeshComponent* GetOrCreateSplineMeshComponent(AActor* InTargetActor, UPCGComponent* SourceComponent, uint64 SettingsUID, const FPCGSplineMeshComponentBuilderParameters& Params);
 	UE_DEPRECATED(5.6, "Use signature with no SettingsUID. Make sure to fill Params.SettingsCRC to enable component reuse.")
-	static UE_API UPCGManagedSplineMeshComponent* GetOrCreateManagedSplineMeshComponent(AActor* InTargetActor, UPCGComponent* SourceComponent, uint64 SettingsUID, const FPCGSplineMeshComponentBuilderParameters& Params);
+	static PCG_API UPCGManagedSplineMeshComponent* GetOrCreateManagedSplineMeshComponent(AActor* InTargetActor, UPCGComponent* SourceComponent, uint64 SettingsUID, const FPCGSplineMeshComponentBuilderParameters& Params);
 	
-	static UE_API UInstancedStaticMeshComponent* GetOrCreateISMC(AActor* InTargetActor, UPCGComponent* SourceComponent, const FPCGISMComponentBuilderParams& Params, FPCGContext* OptionalContext = nullptr);
-	static UE_API UPCGManagedISMComponent* GetOrCreateManagedISMC(AActor* InTargetActor, UPCGComponent* SourceComponent, const FPCGISMComponentBuilderParams& Params, FPCGContext* OptionalContext = nullptr);
-	static UE_API UPCGManagedISKMComponent* GetOrCreateManagedABMC(AActor* InTargetActor, UPCGComponent* SourceComponent, const FPCGSkinnedMeshComponentBuilderParams& Params, FPCGContext* OptionalContext = nullptr);
-	static UE_API USplineMeshComponent* GetOrCreateSplineMeshComponent(AActor* InTargetActor, UPCGComponent* SourceComponent, const FPCGSplineMeshComponentBuilderParameters& Params, FPCGContext* OptionalContext = nullptr);
-	static UE_API UPCGManagedSplineMeshComponent* GetOrCreateManagedSplineMeshComponent(AActor* InTargetActor, UPCGComponent* SourceComponent, const FPCGSplineMeshComponentBuilderParameters& Params, FPCGContext* OptionalContext = nullptr);
+	static PCG_API UInstancedStaticMeshComponent* GetOrCreateISMC(AActor* InTargetActor, UPCGComponent* SourceComponent, const FPCGISMComponentBuilderParams& Params, FPCGContext* OptionalContext = nullptr);
+	static PCG_API UPCGManagedISMComponent* GetOrCreateManagedISMC(AActor* InTargetActor, UPCGComponent* SourceComponent, const FPCGISMComponentBuilderParams& Params, FPCGContext* OptionalContext = nullptr);
+	static PCG_API UPCGManagedISKMComponent* GetOrCreateManagedABMC(AActor* InTargetActor, UPCGComponent* SourceComponent, const FPCGSkinnedMeshComponentBuilderParams& Params, FPCGContext* OptionalContext = nullptr);
+	static PCG_API USplineMeshComponent* GetOrCreateSplineMeshComponent(AActor* InTargetActor, UPCGComponent* SourceComponent, const FPCGSplineMeshComponentBuilderParameters& Params, FPCGContext* OptionalContext = nullptr);
+	static PCG_API UPCGManagedSplineMeshComponent* GetOrCreateManagedSplineMeshComponent(AActor* InTargetActor, UPCGComponent* SourceComponent, const FPCGSplineMeshComponentBuilderParameters& Params, FPCGContext* OptionalContext = nullptr);
 
-	static UE_API bool DeleteActors(UWorld* World, const TArray<TSoftObjectPtr<AActor>>& ActorsToDelete);
+	static PCG_API bool DeleteActors(UWorld* World, const TArray<TSoftObjectPtr<AActor>>& ActorsToDelete);
 
 	template <typename T, typename = typename std::enable_if_t<std::is_base_of_v<AActor, T>>>
 	inline static void ForEachActorInLevel(ULevel* Level, TFunctionRef<bool(AActor*)> Callback)
@@ -161,7 +161,7 @@ PRAGMA_ENABLE_DEPRECATION_WARNINGS
 	* @param ActorClass class of AActor to pass to the callback
 	* @param Callback Function to call with the found actor. Needs to return a bool, to indicate if it needs to continue (true = yes)
 	*/
-	static UE_API void ForEachActorInLevel(ULevel* Level, TSubclassOf<AActor> ActorClass, TFunctionRef<bool(AActor*)> Callback);
+	static PCG_API void ForEachActorInLevel(ULevel* Level, TSubclassOf<AActor> ActorClass, TFunctionRef<bool(AActor*)> Callback);
 
 	template <typename T, typename = typename std::enable_if_t<std::is_base_of_v<AActor, T>>>
 	inline static void ForEachActorInWorld(UWorld* World, TFunctionRef<bool(AActor*)> Callback)
@@ -175,7 +175,7 @@ PRAGMA_ENABLE_DEPRECATION_WARNINGS
 	* @param ActorClass class of AActor to pass to the callback
 	* @param Callback Function to call with the found actor. Needs to return a bool, to indicate if it needs to continue (true = yes)
 	*/
-	static UE_API void ForEachActorInWorld(UWorld* World, TSubclassOf<AActor> ActorClass, TFunctionRef<bool(AActor*)> Callback);
+	static PCG_API void ForEachActorInWorld(UWorld* World, TSubclassOf<AActor> ActorClass, TFunctionRef<bool(AActor*)> Callback);
 
 	/**
 	* Spawn a new actor of type T and attach it to the parent (if not null)
@@ -200,7 +200,7 @@ PRAGMA_ENABLE_DEPRECATION_WARNINGS
 	* @param Transform The transform for the new actor
 	* @param Parent Optional parent to attach to.
 	*/
-	static UE_API AActor* SpawnDefaultActor(UWorld* World, ULevel* Level, TSubclassOf<AActor> ActorClass, FName BaseName, const FTransform& Transform, AActor* Parent = nullptr);
+	static PCG_API AActor* SpawnDefaultActor(UWorld* World, ULevel* Level, TSubclassOf<AActor> ActorClass, FName BaseName, const FTransform& Transform, AActor* Parent = nullptr);
 
 	/**
 	* Spawn a new actor and attach it to the parent (if not null)
@@ -211,7 +211,7 @@ PRAGMA_ENABLE_DEPRECATION_WARNINGS
 	* @param SpawnParams The spawn parameters
 	* @param Parent Optional parent to attach to.
 	*/
-	static UE_API AActor* SpawnDefaultActor(UWorld* World, ULevel* Level, TSubclassOf<AActor> ActorClass, const FTransform& Transform, const FActorSpawnParameters& SpawnParams, AActor* Parent = nullptr);
+	static PCG_API AActor* SpawnDefaultActor(UWorld* World, ULevel* Level, TSubclassOf<AActor> ActorClass, const FTransform& Transform, const FActorSpawnParameters& SpawnParams, AActor* Parent = nullptr);
 
 	/**
 	* Struct containing all parameters needed to spawn the actor
@@ -240,20 +240,20 @@ PRAGMA_ENABLE_DEPRECATION_WARNINGS
 	* Spawn a new actor
 	* @param Params struct containing all the parameters needed to spawn the actor
 	*/
-	static UE_API AActor* SpawnDefaultActor(const FSpawnDefaultActorParams& Params);
+	static PCG_API AActor* SpawnDefaultActor(const FSpawnDefaultActorParams& Params);
 	
 	/**
 	 * Return the grid cell coordinates on the PCG partition grid given a position and the grid size.
 	 */
-	static UE_API FIntVector GetCellCoord(FVector InPosition, uint32 InGridSize, bool bUse2DGrid);
+	static PCG_API FIntVector GetCellCoord(FVector InPosition, uint32 InGridSize, bool bUse2DGrid);
 
 	/**
 	 * Return the center of the PCG partition grid cell given a position and the grid size.
 	 */
-	static UE_API FVector GetCellCenter(const FVector& InPosition, uint32 InGridSize, bool bUse2DGrid);
+	static PCG_API FVector GetCellCenter(const FVector& InPosition, uint32 InGridSize, bool bUse2DGrid);
 
 	/** Extract the tags and the actor reference of the given actor and hash it. Useful for CRC dependencies that depends on the tags or the instance of the actor.*/
-	static UE_API int ComputeHashFromActorTagsAndReference(const AActor* InActor, const bool bIncludeTags, const bool bIncludeActorReference);
+	static PCG_API int ComputeHashFromActorTagsAndReference(const AActor* InActor, const bool bIncludeTags, const bool bIncludeActorReference);
 
 #if WITH_EDITOR
 	static UPackage* CreatePreviewPackage(ULevel* InLevel, const FString& InActorName);
@@ -263,7 +263,8 @@ PRAGMA_ENABLE_DEPRECATION_WARNINGS
 	static TArray<APCGIsolatedActor*> GetIsolatedActorsOriginatingFrom(const AActor* InActor);
 	
 	/** Extracts (like Clear PCG Link) content from the target actor, assuming it has a PCG component, into a PCG Isolated Actor. */
-	static UE_API AActor* IsolateFromActor(AActor* Actor, const TArray<FName>& RequiredTags);
-};
+	static PCG_API AActor* IsolateFromActor(AActor* Actor, const TArray<FName>& RequiredTags);
 
-#undef UE_API
+	/** Finds all execution sources using the passed in param as context for search. */
+	static PCG_API TArray<IPCGGraphExecutionSource*> GetExecutionSourcesFromSelectionKey(const FPCGGetExecutionSourcesFromSelectionKeyParams& InParams);
+};
