[PCG] Add GPU data support to Remove Empty Data node.

This is the last node needing changing from review of existing nodes.

#rnx
#Synced-CL 48720584

diff --git a/Engine/Plugins/PCG/Source/PCG/Private/Compute/Data/PCGProxyForGPUData.cpp b/Engine/Plugins/PCG/Source/PCG/Private/Compute/Data/PCGProxyForGPUData.cpp
--- a/Engine/Plugins/PCG/Source/PCG/Private/Compute/Data/PCGProxyForGPUData.cpp
+++ b/Engine/Plugins/PCG/Source/PCG/Private/Compute/Data/PCGProxyForGPUData.cpp
@@ -63,6 +63,22 @@ FPCGDataTypeBaseId UPCGProxyForGPUData::GetUnderlyingDataTypeId() const
 	return FPCGDataTypeBaseId{};
 }
 
+FIntVector4 UPCGProxyForGPUData::GetNumElements() const
+{
+	if (TSharedPtr<const FPCGProxyForGPUDataCollection> DataCollectionGPU = GetGPUInfo())
+	{
+		if (TSharedPtr<const FPCGDataCollectionDesc> Desc = DataCollectionGPU->GetDescription())
+		{
+			if (Desc->GetDataDescriptions().IsValidIndex(DataIndexInCollection))
+			{
+				return Desc->GetDataDescriptions()[DataIndexInCollection].GetElementCount();
+			}
+		}
+	}
+
+	return FIntVector4::ZeroValue;
+}
+
 TSharedPtr<const FPCGProxyForGPUDataCollection> UPCGProxyForGPUData::GetInputDataCollectionInfo() const
 {
 	return DataCollectionOnGPU;
diff --git a/Engine/Plugins/PCG/Source/PCG/Private/Elements/PCGRemoveEmptyData.cpp b/Engine/Plugins/PCG/Source/PCG/Private/Elements/PCGRemoveEmptyData.cpp
--- a/Engine/Plugins/PCG/Source/PCG/Private/Elements/PCGRemoveEmptyData.cpp
+++ b/Engine/Plugins/PCG/Source/PCG/Private/Elements/PCGRemoveEmptyData.cpp
@@ -3,6 +3,7 @@
 #include "Elements/PCGRemoveEmptyData.h"
 
 #include "PCGContext.h"
+#include "Compute/Data/PCGProxyForGPUData.h"
 #include "Metadata/Accessors/PCGAttributeAccessorHelpers.h"
 #include "Metadata/Accessors/PCGAttributeAccessorKeys.h"
 
@@ -62,11 +63,21 @@ bool FPCGRemoveEmptyDataElement::ExecuteInternal(FPCGContext* InContext) const
 			continue;
 		}
 
-		// Get the keys on the index, if it doesn't exist or is empty, discard the data.
-		const TUniquePtr<const IPCGAttributeAccessorKeys> Keys = PCGAttributeAccessorHelpers::CreateConstKeys(Input.Data, FPCGAttributePropertySelector::CreateExtraPropertySelector(EPCGExtraProperties::Index));
-		if (!Keys || Keys->GetNum() == 0)
+		if (const UPCGProxyForGPUData* Proxy = Cast<UPCGProxyForGPUData>(Input.Data))
 		{
-			continue;
+			if (Proxy->GetNumElements() == FIntVector4::ZeroValue)
+			{
+				continue;
+			}
+		}
+		else
+		{
+			// Get the keys on the index, if it doesn't exist or is empty, discard the data.
+			const TUniquePtr<const IPCGAttributeAccessorKeys> Keys = PCGAttributeAccessorHelpers::CreateConstKeys(Input.Data, FPCGAttributePropertySelector::CreateExtraPropertySelector(EPCGExtraProperties::Index));
+			if (!Keys || Keys->GetNum() == 0)
+			{
+				continue;
+			}
 		}
 
 		InContext->OutputData.TaggedData.Add(Input);
diff --git a/Engine/Plugins/PCG/Source/PCG/Public/Compute/Data/PCGProxyForGPUData.h b/Engine/Plugins/PCG/Source/PCG/Public/Compute/Data/PCGProxyForGPUData.h
--- a/Engine/Plugins/PCG/Source/PCG/Public/Compute/Data/PCGProxyForGPUData.h
+++ b/Engine/Plugins/PCG/Source/PCG/Public/Compute/Data/PCGProxyForGPUData.h
@@ -49,6 +49,8 @@ public:
 	PCG_API virtual FPCGDataTypeBaseId GetUnderlyingDataTypeId() const override;
 	//~End UPCGData interface
 
+	PCG_API virtual FIntVector4 GetNumElements() const;
+
 	/** Returns the GPU info. Returns null if the buffer has been discarded. */
 	TSharedPtr<const FPCGProxyForGPUDataCollection> GetInputDataCollectionInfo() const;
 
diff --git a/Engine/Plugins/PCG/Source/PCG/Public/Elements/PCGRemoveEmptyData.h b/Engine/Plugins/PCG/Source/PCG/Public/Elements/PCGRemoveEmptyData.h
--- a/Engine/Plugins/PCG/Source/PCG/Public/Elements/PCGRemoveEmptyData.h
+++ b/Engine/Plugins/PCG/Source/PCG/Public/Elements/PCGRemoveEmptyData.h
@@ -34,5 +34,6 @@ class FPCGRemoveEmptyDataElement : public IPCGElement
 protected:
 	virtual bool ExecuteInternal(FPCGContext* InContext) const override;
 	virtual bool SupportsBasePointDataInputs(FPCGContext* InContext) const override { return true; }
+	virtual bool SupportsGPUResidentData(FPCGContext* InContext) const override { return true; }
 };
 
