[PCG] Added option to treat string data attribute as csv for tags.

#rnx
#Synced-CL 48730134

diff --git a/Engine/Plugins/PCG/Source/PCG/Private/Elements/Metadata/PCGDataAttributesAndTags.cpp b/Engine/Plugins/PCG/Source/PCG/Private/Elements/Metadata/PCGDataAttributesAndTags.cpp
--- a/Engine/Plugins/PCG/Source/PCG/Private/Elements/Metadata/PCGDataAttributesAndTags.cpp
+++ b/Engine/Plugins/PCG/Source/PCG/Private/Elements/Metadata/PCGDataAttributesAndTags.cpp
@@ -3,6 +3,7 @@
 #include "Elements/Metadata/PCGDataAttributesAndTags.h"
 
 #include "PCGContext.h"
+#include "Helpers/PCGHelpers.h"
 #include "Helpers/PCGTagHelpers.h"
 
 #include UE_INLINE_GENERATED_CPP_BY_NAME(PCGDataAttributesAndTags)
@@ -175,6 +176,24 @@ bool FPCGDataAttributesToTagsElement::ExecuteInternal(FPCGContext* InContext) co
 					else
 					{
 						const T Value = static_cast<const FPCGMetadataAttribute<T>*>(Attribute)->GetValueFromItemKey(PCGFirstEntryKey);
+
+						if constexpr (PCG::Private::IsOfTypes<T, FString, FName>())
+						{
+							if (Settings->bTreatStringAttributesAsTags)
+							{
+								FString StringValue = PCG::Private::MetadataTraits<T>::ToString(Value);
+								const TArray<FString> TagAndValues = PCGHelpers::GetStringArrayFromCommaSeparatedList(StringValue);
+
+								for (const FString& TagAndValue : TagAndValues)
+								{
+									OutputTaggedData.Tags.Add(TagAndValue);
+								}
+
+								return true;
+							}
+						}
+
+						// Otherwise/default case
 						OutputTaggedData.Tags.Add(FString::Format(TEXT("{0}:{1}"), {std::move(TagName), PCG::Private::MetadataTraits<T>::ToString(Value)}));
 					}
 					
diff --git a/Engine/Plugins/PCG/Source/PCG/Public/Elements/Metadata/PCGDataAttributesAndTags.h b/Engine/Plugins/PCG/Source/PCG/Public/Elements/Metadata/PCGDataAttributesAndTags.h
--- a/Engine/Plugins/PCG/Source/PCG/Public/Elements/Metadata/PCGDataAttributesAndTags.h
+++ b/Engine/Plugins/PCG/Source/PCG/Public/Elements/Metadata/PCGDataAttributesAndTags.h
@@ -56,6 +56,10 @@ public:
 	// Do not output the value of the attribute in the tag (e.g. MyAttr:1), but only the attribute name (e.g. MyAttr)
 	UPROPERTY(EditAnywhere, BlueprintReadWrite, AdvancedDisplay, Category = "Settings", meta = (PCG_Overridable))
 	bool bDiscardAttributeValue = false;
+
+	// Controls whether string attribute will be broken down as their contents only, and supports CSV too.
+	UPROPERTY(EditAnywhere, BlueprintReadWrite, AdvancedDisplay, Category = "Settings", meta = (PCG_Overridable))
+	bool bTreatStringAttributesAsTags = false;
 };
 
 UCLASS(MinimalAPI, BlueprintType, ClassGroup = (Procedural))
