[PCG] Remove temporary fix for Polygon 2D write accessor

* Add API on Polygon and PolygonData to update a vertex in place.

#rnx
#Synced-CL 48730092

diff --git a/Engine/Plugins/PCG/Source/PCG/Private/Data/PCGPolygon2DData.cpp b/Engine/Plugins/PCG/Source/PCG/Private/Data/PCGPolygon2DData.cpp
--- a/Engine/Plugins/PCG/Source/PCG/Private/Data/PCGPolygon2DData.cpp
+++ b/Engine/Plugins/PCG/Source/PCG/Private/Data/PCGPolygon2DData.cpp
@@ -17,8 +17,6 @@
 #include "Serialization/ArchiveCrc32.h"
 #include "UObject/Package.h"
 
-#include "Data/PCGPolygon2DData.inl"
-
 #include UE_INLINE_GENERATED_CPP_BY_NAME(PCGPolygon2DData)
 
 #define LOCTEXT_NAMESPACE "PCGPolygon2DData"
@@ -166,9 +164,9 @@ namespace PCGPolygon2DData
 			switch (Property)
 			{
 			case EPCGPolygon2DProperties::Position:
-				return MakeUnique<PCGTemporary::FPCGPolygon2DVerticesAccessor<FVector, EPCGPolygon2DProperties::Position>>(bConst);
+				return MakeUnique<FPCGPolygon2DVerticesAccessor<FVector, EPCGPolygon2DProperties::Position>>(bConst);
 			case EPCGPolygon2DProperties::LocalPosition:
-				return MakeUnique<PCGTemporary::FPCGPolygon2DVerticesAccessor<FVector2d, EPCGPolygon2DProperties::LocalPosition>>(bConst);
+				return MakeUnique<FPCGPolygon2DVerticesAccessor<FVector2d, EPCGPolygon2DProperties::LocalPosition>>(bConst);
 			case EPCGPolygon2DProperties::Rotation: //fall-through
 			case EPCGPolygon2DProperties::SegmentIndex: //fall-through
 			case EPCGPolygon2DProperties::HoleIndex: //fall-through
diff --git a/Engine/Plugins/PCG/Source/PCG/Private/Data/PCGPolygon2DData.inl b/Engine/Plugins/PCG/Source/PCG/Private/Data/PCGPolygon2DData.inl
deleted file mode 100644
diff --git a/Engine/Plugins/PCG/Source/PCG/Public/Data/PCGPolygon2DData.h b/Engine/Plugins/PCG/Source/PCG/Public/Data/PCGPolygon2DData.h
--- a/Engine/Plugins/PCG/Source/PCG/Public/Data/PCGPolygon2DData.h
+++ b/Engine/Plugins/PCG/Source/PCG/Public/Data/PCGPolygon2DData.h
@@ -133,6 +133,17 @@ public:
 	// Get the functions to the accessor factory
 	static PCG_API FPCGAttributeAccessorMethods GetPolygon2DAccessorMethods();
 
+	/**
+	 * Utility function to be able to modify a vertex in place. Will assert if out of bounds.
+	 * @param Vertex       Vertex value to update
+	 * @param HoleIndex    Hole Index, -1 if targeting Outer
+	 * @param VertexIndex Index of the vertex to update in the Hole/Outer polygon.
+	 */
+	void UpdateVertexInPlace(const FVector2d& Vertex, int32 HoleIndex, int32 VertexIndex)
+	{
+		Polygon.UpdateVertexInPlace(Vertex, HoleIndex, VertexIndex);
+	}
+
 public:
 	const UE::Geometry::FGeneralPolygon2d& GetPolygon() const { return Polygon; }
 	const TMap<int, TPair<int, int>>& GetSegmentIndexToSegmentAndHoleIndices() const { return SegmentIndexToSegmentAndHoleIndices; }
diff --git a/Engine/Plugins/PCG/Source/PCG/Public/Metadata/Accessors/PCGPolygon2DAccessor.h b/Engine/Plugins/PCG/Source/PCG/Public/Metadata/Accessors/PCGPolygon2DAccessor.h
--- a/Engine/Plugins/PCG/Source/PCG/Public/Metadata/Accessors/PCGPolygon2DAccessor.h
+++ b/Engine/Plugins/PCG/Source/PCG/Public/Metadata/Accessors/PCGPolygon2DAccessor.h
@@ -7,8 +7,8 @@
 
 /**
 * Templated accessor for polygon vertices accessor.
-* NOTE FOR 5.7.1: This accessor is NOT thread-safe on Set, and is replaced by a thread-safe copy in the PCGPolygon2DData.inl
-* since we can't modify this file in a minor version.
+* There is an assumption that no vertex will be set by 2 different threads at the same time and that the polygon itself should not be modified (by adding/removing vertices)
+* while the accessor is getting/setting values.
 * You must not use this accessor directly.
 */
 template<typename T, EPCGPolygon2DProperties Target>
@@ -124,7 +124,6 @@ public:
 
 		UPCGPolygon2DData* Data = static_cast<UPCGPolygon2DData*>(ContainerKeys);
 		const FTransform& Transform = Data->GetTransform();
-		const UE::Geometry::FGeneralPolygon2d& Polygon = Data->GetPolygon();
 		const TMap<int, TPair<int, int>>& SegmentIndexToSegmentAndHoleIndices = Data->GetSegmentIndexToSegmentAndHoleIndices();
 
 		const int32 NumPoints = SegmentIndexToSegmentAndHoleIndices.Num();
@@ -146,16 +145,6 @@ public:
 				InverseTransform = Transform.Inverse();
 			}
 
-			// Need to recreate a new polygon from scratch here.
-			UE::Geometry::FGeneralPolygon2d NewPolygon;
-			TArray<TArray<FVector2d>> Vertices;
-			Vertices.Reserve(1 + Polygon.GetHoles().Num());
-			Vertices.Add(Polygon.GetOuter().GetVertices());
-			for (const UE::Geometry::FPolygon2d& Hole : Polygon.GetHoles())
-			{
-				Vertices.Add(Hole.GetVertices());
-			}
-
 			for (int i = 0; i < InValues.Num(); ++i)
 			{
 				const int32 CurrIndex = (Index + i) % NumPoints;
@@ -165,12 +154,12 @@ public:
 				{
 					static_assert(std::is_same_v<Type, FVector3d>);
 					FVector LocalPosition = InverseTransform.TransformPosition(InValues[i]);
-					Vertices[HoleIndex + 1][SegmentIndex] = FVector2d(LocalPosition.X, LocalPosition.Y);
+					Data->UpdateVertexInPlace(FVector2d(LocalPosition.X, LocalPosition.Y), HoleIndex, SegmentIndex);
 				}
 				else if constexpr (Target == EPCGPolygon2DProperties::LocalPosition)
 				{
 					static_assert(std::is_same_v<Type, FVector2d>);
-					Vertices[HoleIndex + 1][SegmentIndex] = InValues[i];
+					Data->UpdateVertexInPlace(InValues[i], HoleIndex, SegmentIndex);
 				}
 				else
 				{
@@ -179,14 +168,6 @@ public:
 				}
 			}
 
-			NewPolygon.SetOuter(UE::Geometry::FPolygon2d(MoveTemp(Vertices[0])));
-			for (int i = 1; i < Vertices.Num(); ++i)
-			{
-				NewPolygon.AddHole(UE::Geometry::FPolygon2d(MoveTemp(Vertices[i])));
-			}
-
-			Data->SetPolygon(NewPolygon);
-
 			return true;
 		}
 	}
diff --git a/Engine/Source/Runtime/GeometryCore/Public/Curve/GeneralPolygon2.h b/Engine/Source/Runtime/GeometryCore/Public/Curve/GeneralPolygon2.h
--- a/Engine/Source/Runtime/GeometryCore/Public/Curve/GeneralPolygon2.h
+++ b/Engine/Source/Runtime/GeometryCore/Public/Curve/GeneralPolygon2.h
@@ -73,6 +73,19 @@ public:
 		return Holes;
 	}
 
+	/**
+	 * Utility function to be able to modify a vertex in place. Will assert if out of bounds.
+	 * @param Vertex       Vertex value to update
+	 * @param HoleIndex    Hole Index, -1 if targeting Outer
+	 * @param VertexIndex Index of the vertex to update in the Hole/Outer polygon.
+	 */
+	void UpdateVertexInPlace(const TVector2<T>& Vertex, int32 HoleIndex, int32 VertexIndex)
+	{
+		check(HoleIndex == -1 || Holes.IsValidIndex(HoleIndex));
+		TPolygon2<T>& Polygon = HoleIndex == -1 ? Outer : Holes[HoleIndex];
+		check(Polygon.GetVertices().IsValidIndex(VertexIndex));
+		Polygon[VertexIndex] = Vertex;
+	}
 
 	bool AddHole(TPolygon2<T> Hole, bool bCheckContainment = true, bool bCheckOrientation = true)
 	{
