[PCG] Add a "always accessible" context

* Store the current context when an element execute in a thread local variable

#rnx
#Synced-CL 48730092

diff --git a/Engine/Plugins/PCG/Source/PCG/Private/Helpers/PCGAsync.cpp b/Engine/Plugins/PCG/Source/PCG/Private/Helpers/PCGAsync.cpp
--- a/Engine/Plugins/PCG/Source/PCG/Private/Helpers/PCGAsync.cpp
+++ b/Engine/Plugins/PCG/Source/PCG/Private/Helpers/PCGAsync.cpp
@@ -588,6 +588,8 @@ namespace FPCGAsync
 				SynchroStruct.ChunkProcessedIndexAndNumElementsWrittenQueue.Enqueue(CurrentChunkToProcess, NumElementsWritten);
 			};
 
+			FPCGContext* CurrentContext = FPCGContext::GetContextForThread();
+
 			// Starting all the tasks, only if the associated task has still some work to do.
 			for (int32 TaskIndex = 0; TaskIndex < NumFutures; ++TaskIndex)
 			{
@@ -599,10 +601,12 @@ namespace FPCGAsync
 					continue;
 				}
 
-				AsyncTasks.Emplace(UE::Tasks::Launch(UE_SOURCE_LOCATION, [&JobTask, &SynchroStruct, &StartEnd = AsyncState.TasksChunksStartEnd[TaskIndex]]() -> void
+				AsyncTasks.Emplace(UE::Tasks::Launch(UE_SOURCE_LOCATION, [&JobTask, &SynchroStruct, &StartEnd = AsyncState.TasksChunksStartEnd[TaskIndex], CurrentContext]() -> void
 				{
 					LLM_SCOPE_BYTAG(PCG);
 
+					FPCGContextScope ContextScope(CurrentContext);
+
 					// Do at least one run, as we could fall into an infinite loop if we always have to stop before doing anything.
 					// Quit if we have no more work to do or we were told to stop.
 					do
diff --git a/Engine/Plugins/PCG/Source/PCG/Private/Metadata/Accessors/PCGPropertyAccessor.cpp b/Engine/Plugins/PCG/Source/PCG/Private/Metadata/Accessors/PCGPropertyAccessor.cpp
--- a/Engine/Plugins/PCG/Source/PCG/Private/Metadata/Accessors/PCGPropertyAccessor.cpp
+++ b/Engine/Plugins/PCG/Source/PCG/Private/Metadata/Accessors/PCGPropertyAccessor.cpp
@@ -4,6 +4,43 @@
 
 #include "UObject/EnumProperty.h"
 
+#define LOCTEXT_NAMESPACE "PCGPropertyAccessor"
+
+namespace PCGPropertyAccessor
+{
+	void LogErrorForInvalidObjectWrite(const UClass* PropertyClass, const UClass* ObjectClass, const FSoftObjectPath& Value)
+	{
+#if !NO_LOGGING
+		PCGLog::LogErrorOnGraph(FText::Format(
+			LOCTEXT("ObjectSetRangeError", "[FPCGPropertyObjectPtrAccessor::SetRange] Tried to write into an Object property of class \"{0}\" with an object of class \"{1}\" ({2}). This is invalid."),
+			FText::FromString(PropertyClass->GetName()),
+			FText::FromString(ObjectClass->GetName()),
+			FText::FromString(Value.ToString())));
+#endif // !NO_LOGGING
+	}
+
+	void LogErrorForInvalidClassWrite(const UClass* MetaClass, const UClass* ObjectMetaClass, const FSoftClassPath& Value)
+	{
+#if !NO_LOGGING
+		PCGLog::LogErrorOnGraph(FText::Format(
+			LOCTEXT("ClassSetRangeError", "[FPCGPropertyObjectPtrAccessor::SetRange] Tried to write into an Class property of metaclass \"{0}\" with a class of metaclass \"{1}\" ({2}). This is invalid."),
+			FText::FromString(MetaClass->GetName()),
+			FText::FromString(ObjectMetaClass->GetName()),
+			FText::FromString(Value.ToString())));
+#endif // !NO_LOGGING
+	}
+	
+	void LogErrorForInvalidEnumWrite(const UEnum* Enum, const int64 Value)
+	{
+#if !NO_LOGGING
+		PCGLog::LogErrorOnGraph(FText::Format(
+			LOCTEXT("EnumSetRangeError", "[FPCGEnumPropertyAccessor::SetRange] Tried to write into an Enum \"{0}\" a value \"{1}\" that is not a valid value for this enum."),
+			FText::FromString(Enum->GetName()),
+			Value));
+#endif // !NO_LOGGING
+	}
+}
+
 IPCGPropertyChain::IPCGPropertyChain(const FProperty* Property, TArray<const FProperty*>&& ExtraProperties)
 	: PropertyChain(std::move(ExtraProperties))
 {
@@ -56,7 +93,7 @@ bool FPCGEnumPropertyAccessor::SetRangeImpl(TArrayView<const int64> InValues, in
 	{
 		if (Enum && (!Enum->IsValidEnumValue(Value) || (bHasMaxValue && Value == EnumMaxValue)))
 		{
-			UE_LOG(LogPCG, Error, TEXT("[FPCGEnumPropertyAccessor::SetRange] Try to write into an Enum \"%s\" a value \"%lld\" that is not a valid value for this enum."), *Enum->GetName(), Value);
+			PCGPropertyAccessor::LogErrorForInvalidEnumWrite(Enum, Value);
 			bWriteOK = false;
 			return;
 		}
@@ -65,4 +102,6 @@ bool FPCGEnumPropertyAccessor::SetRangeImpl(TArrayView<const int64> InValues, in
 	});
 
 	return bIterateOK && bWriteOK;
-}
\ No newline at end of file
+}
+
+#undef LOCTEXT_NAMESPACE
\ No newline at end of file
diff --git a/Engine/Plugins/PCG/Source/PCG/Private/PCGContext.cpp b/Engine/Plugins/PCG/Source/PCG/Private/PCGContext.cpp
--- a/Engine/Plugins/PCG/Source/PCG/Private/PCGContext.cpp
+++ b/Engine/Plugins/PCG/Source/PCG/Private/PCGContext.cpp
@@ -744,4 +744,28 @@ bool FPCGContext::GetFromCache(const FPCGGetFromCacheParams& Params, FPCGDataCol
 	}
 }
 
-#undef LOCTEXT_NAMESPACE
\ No newline at end of file
+
+namespace PCGContext
+{
+	thread_local FPCGContext* ThisThreadContext = nullptr;
+}
+
+FPCGContext* FPCGContext::GetContextForThread()
+{
+	return PCGContext::ThisThreadContext;
+}
+
+FPCGContextScope::FPCGContextScope(FPCGContext* InContext)
+{
+	PreviousContext = PCGContext::ThisThreadContext;
+	ThisContext = InContext;
+	PCGContext::ThisThreadContext = InContext;
+}
+
+FPCGContextScope::~FPCGContextScope()
+{
+	checkf(PCGContext::ThisThreadContext == ThisContext, TEXT("Multiple scope are destroyed in the wrong order"));
+	PCGContext::ThisThreadContext = PreviousContext;
+}
+
+#undef LOCTEXT_NAMESPACE
diff --git a/Engine/Plugins/PCG/Source/PCG/Private/PCGElement.cpp b/Engine/Plugins/PCG/Source/PCG/Private/PCGElement.cpp
--- a/Engine/Plugins/PCG/Source/PCG/Private/PCGElement.cpp
+++ b/Engine/Plugins/PCG/Source/PCG/Private/PCGElement.cpp
@@ -208,6 +208,8 @@ bool IPCGElement::Execute(FPCGContext* Context) const
 	check(Context && Context->AsyncState.NumAvailableTasks != 0 && Context->CurrentPhase <= EPCGExecutionPhase::Done);
 	check(Context->AsyncState.bIsRunningOnMainThread || !CanExecuteOnlyOnMainThread(Context));
 
+	FPCGContextScope ContextScope(Context);
+
 	PCGUtils::FScopedCallOutputDevice OutputDevice;
 
 	while (Context->CurrentPhase != EPCGExecutionPhase::Done)
diff --git a/Engine/Plugins/PCG/Source/PCG/Private/PCGModule.cpp b/Engine/Plugins/PCG/Source/PCG/Private/PCGModule.cpp
--- a/Engine/Plugins/PCG/Source/PCG/Private/PCGModule.cpp
+++ b/Engine/Plugins/PCG/Source/PCG/Private/PCGModule.cpp
@@ -334,9 +334,9 @@ PCG_API DEFINE_LOG_CATEGORY(LogPCG);
 
 void PCGLog::LogErrorOnGraph(const FText& InMsg, const FPCGContext* InContext)
 {
-	if (InContext)
+	if (const FPCGContext* Context = InContext ? InContext : FPCGContext::GetContextForThread())
 	{
-		PCGE_LOG_C(Error, GraphAndLog, InContext, InMsg);
+		PCGE_LOG_C(Error, GraphAndLog, Context, InMsg);
 	}
 	else
 	{
@@ -346,9 +346,9 @@ void PCGLog::LogErrorOnGraph(const FText& InMsg, const FPCGContext* InContext)
 
 void PCGLog::LogWarningOnGraph(const FText& InMsg, const FPCGContext* InContext)
 {
-	if (InContext)
+	if (const FPCGContext* Context = InContext ? InContext : FPCGContext::GetContextForThread())
 	{
-		PCGE_LOG_C(Warning, GraphAndLog, InContext, InMsg);
+		PCGE_LOG_C(Warning, GraphAndLog, Context, InMsg);
 	}
 	else
 	{
diff --git a/Engine/Plugins/PCG/Source/PCG/Public/Metadata/Accessors/PCGPropertyAccessor.h b/Engine/Plugins/PCG/Source/PCG/Public/Metadata/Accessors/PCGPropertyAccessor.h
--- a/Engine/Plugins/PCG/Source/PCG/Public/Metadata/Accessors/PCGPropertyAccessor.h
+++ b/Engine/Plugins/PCG/Source/PCG/Public/Metadata/Accessors/PCGPropertyAccessor.h
@@ -105,6 +105,9 @@ namespace PCGPropertyAccessor
 
 		return true;
 	}
+
+	void LogErrorForInvalidObjectWrite(const UClass* PropertyClass, const UClass* ObjectClass, const FSoftObjectPath& Value);
+	void LogErrorForInvalidClassWrite(const UClass* MetaClass, const UClass* ObjectMetaClass, const FSoftClassPath& Value);
 }
 
 /**
@@ -452,7 +455,7 @@ public:
 			UObject* Object = Value.TryLoad();
 			if (Object && Property->PropertyClass && !Object->IsA(Property->PropertyClass))
 			{
-				UE_LOG(LogPCG, Error, TEXT("[FPCGPropertyObjectPtrAccessor::SetRange] Try to write into an Object property of class \"%s\" with an object of class \"%s\" (%s). This is invalid."), *Property->PropertyClass->GetName(), *Object->GetClass()->GetName(), *Value.ToString());
+				PCGPropertyAccessor::LogErrorForInvalidObjectWrite(Property->PropertyClass, Object->GetClass(), Value);
 				bWriteOK = false;
 				return;
 			}
@@ -462,7 +465,7 @@ public:
 				const UClass* Class = Cast<const UClass>(Object);
 				if (Class && Property->MetaClass && !Class->IsChildOf(Property->MetaClass))
 				{
-					UE_LOG(LogPCG, Error, TEXT("[FPCGPropertyObjectPtrAccessor::SetRange] Try to write into an Class property of metaclass \"%s\" with a class of metaclass \"%s\" (%s). This is invalid."), *Property->MetaClass->GetName(), *Class->GetName(), *Value.ToString());
+					PCGPropertyAccessor::LogErrorForInvalidObjectWrite(Property->MetaClass, Class, Value);
 					bWriteOK = false;
 					return;
 				}
diff --git a/Engine/Plugins/PCG/Source/PCG/Public/PCGContext.h b/Engine/Plugins/PCG/Source/PCG/Public/PCGContext.h
--- a/Engine/Plugins/PCG/Source/PCG/Public/PCGContext.h
+++ b/Engine/Plugins/PCG/Source/PCG/Public/PCGContext.h
@@ -295,6 +295,8 @@ PRAGMA_ENABLE_DEPRECATION_WARNINGS
 	UE_API bool IsCacheEnabled();
 	UE_API void StoreInCache(const FPCGStoreInCacheParams& Params, const FPCGDataCollection& InOutput);
 	UE_API bool GetFromCache(const FPCGGetFromCacheParams& Params, FPCGDataCollection& OutOutput) const;
+	
+	UE_API static FPCGContext* GetContextForThread();
 
 protected:
 	virtual UObject* GetExternalContainerForOverridableParam(const FPCGSettingsOverridableParam& InParam) { return nullptr; }
@@ -353,4 +355,25 @@ private:
 #endif
 };
 
+/**
+ * Allow to set the current context for the current thread. The context can then be queryable with "FPCGContext::GetContextForThread".
+ * Scopes can be nested, as we store the previous context to restore it on scope destruction.
+ * Be very careful in the scope nesting, scopes need to be destroyed in order (ScopeA ctor -> ScopeB ctor -> ScopeB dtor -> ScopeA dtor).
+ */
+struct FPCGContextScope
+{
+	PCG_API FPCGContextScope(FPCGContext* InContext);
+	PCG_API ~FPCGContextScope();
+
+	// Non-copyable, non-moveable
+	FPCGContextScope(const FPCGContextScope&) = delete;
+	FPCGContextScope(FPCGContextScope&&) = delete;
+	FPCGContextScope& operator=(const FPCGContextScope&) = delete;
+	FPCGContextScope& operator=(FPCGContextScope&&) = delete;
+	
+	// Keep this context to have some verification check that scopes are not destroyed in the wrong order.
+	FPCGContext* ThisContext = nullptr;
+	FPCGContext* PreviousContext = nullptr;
+};
+
 #undef UE_API
