[PCG] Partial backout of CL 46583984 causing ISM reuse issues (inputs with no crcs were colliding)

#rnx
#Synced-CL 48759297

diff --git a/Engine/Plugins/PCG/Source/PCG/Private/PCGData.cpp b/Engine/Plugins/PCG/Source/PCG/Private/PCGData.cpp
--- a/Engine/Plugins/PCG/Source/PCG/Private/PCGData.cpp
+++ b/Engine/Plugins/PCG/Source/PCG/Private/PCGData.cpp
@@ -590,16 +590,6 @@ void FPCGDataCollection::ComputeCrcs(bool bFullDataCrc)
 	}
 }
 
-void FPCGDataCollection::SetZeroCrcs()
-{
-	DataCrcs.Empty(TaggedData.Num());
-
-	for (int Index = 0; Index < TaggedData.Num(); ++Index)
-	{
-		DataCrcs.Emplace(0);
-	}
-}
-
 void FPCGDataCollection::AddData(const FPCGTaggedData& InData, const FPCGCrc& InDataCrc)
 {
 	TaggedData.Add(InData);
diff --git a/Engine/Plugins/PCG/Source/PCG/Private/PCGElement.cpp b/Engine/Plugins/PCG/Source/PCG/Private/PCGElement.cpp
--- a/Engine/Plugins/PCG/Source/PCG/Private/PCGElement.cpp
+++ b/Engine/Plugins/PCG/Source/PCG/Private/PCGElement.cpp
@@ -569,8 +569,6 @@ void IPCGElement::PostExecute(FPCGContext* Context) const
 		}
 	}
 
-	// Output data Crc
-	if (Context->IsCacheEnabled())
 	{
 		TRACE_CPUPROFILER_EVENT_SCOPE(IPCGElement::PostExecute::CRC);
 
@@ -582,12 +580,6 @@ void IPCGElement::PostExecute(FPCGContext* Context) const
 		// Compute Crc from output data which will include output pin labels.
 		Context->OutputData.ComputeCrcs(bShouldComputeFullOutputDataCrc);
 	}
-	else
-	{
-		// We check if Crcs are valid in some parts of the code which can be useful for catching bugs, so add valid 0 Crcs here.
-		// Assumption is that these Crcs will be unused if caching is disabled.
-		Context->OutputData.SetZeroCrcs();
-	}
 
 #if WITH_EDITOR
 	const bool bHasErrorsOrWarnings = Context->Node && Context->HasVisualLogs();
diff --git a/Engine/Plugins/PCG/Source/PCG/Public/PCGData.h b/Engine/Plugins/PCG/Source/PCG/Public/PCGData.h
--- a/Engine/Plugins/PCG/Source/PCG/Public/PCGData.h
+++ b/Engine/Plugins/PCG/Source/PCG/Public/PCGData.h
@@ -317,7 +317,6 @@ struct FPCGDataCollection
 
 	/** Computes CRCs for all data items. */
 	UE_API void ComputeCrcs(bool bFullDataCrc);
-	UE_API void SetZeroCrcs();
 
 	/** Add data and CRC to collection. */
 	UE_API void AddData(const FPCGTaggedData& InData, const FPCGCrc& InDataCrc);
