[PCG][Trivial] Missed renames

#rnx
#Synced-CL 48759297

diff --git a/Engine/Plugins/PCG/Source/PCG/Internal/PCGTrackingManager.h b/Engine/Plugins/PCG/Source/PCG/Internal/PCGTrackingManager.h
--- a/Engine/Plugins/PCG/Source/PCG/Internal/PCGTrackingManager.h
+++ b/Engine/Plugins/PCG/Source/PCG/Internal/PCGTrackingManager.h
@@ -80,7 +80,7 @@ public:
 	PCG_API void OnSelectionKeyChanged(const FPCGSelectionKey& InSelectionKey, const UObject* InOriginatingChangeObject, TArrayView<const FBox> InBounds);
 	PCG_API void OnSelectionKeyChanged(const FPCGSelectionKey& InSelectionKey, const UObject* InOriginatingChangeObject, const TSet<FName>& InRemovedTags, bool bInNoRefreshOwner, TArrayView<const FBox> InBounds);
 
-	PCG_API void NotifyObjectChanged(const UObject* InObject, const FPCGTrackingChangeID& InChangeType, const UObject* InOriginatingChangeObject = nullptr);
+	PCG_API void NotifyObjectChanged(const UObject* InObject, const FPCGTrackingChangeID& InChangeID, const UObject* InOriginatingChangeObject = nullptr);
 
 	static void BuildPartitionActorRecords(APCGWorldActor* PCGWorldActor, UWorldPartition* WorldPartition, TMap<FPCGGridCellDescriptor, FGuid>& OutPartitionActorRecords, TSet<FGuid>& OutInvalidPartitionActors);
 
@@ -198,7 +198,7 @@ private:
 	/** Unregister tracking when a execution source is removed or keys are unregistered */
 	void UnregisterTracking(IPCGGraphExecutionSource* InExecutionSource, const TSet<FPCGSelectionKey>* OptionalKeysToUntrack);
 
-	void OnObjectChanged(const UObject* InObject, const FPCGTrackingChangeID& InChangeType, const UObject* InOriginatingChangeObject = nullptr);
+	void OnObjectChanged(const UObject* InObject, const FPCGTrackingChangeID& InChangeID, const UObject* InOriginatingChangeObject = nullptr);
 
 	/** Build initial records for existing partition actors */
 	void BuildPartitionActorRecords();
diff --git a/Engine/Plugins/PCG/Source/PCG/Private/PCGTrackingManager.cpp b/Engine/Plugins/PCG/Source/PCG/Private/PCGTrackingManager.cpp
--- a/Engine/Plugins/PCG/Source/PCG/Private/PCGTrackingManager.cpp
+++ b/Engine/Plugins/PCG/Source/PCG/Private/PCGTrackingManager.cpp
@@ -1411,7 +1411,7 @@ void FPCGTrackingManager::OnObjectPropertyChanged(UObject* InObject, FPropertyCh
 	OnObjectChanged(InObject, PCGTrackingManager::ObjectChanged, /*InOriginatingChangeObject=*/ InObject);
 }
 
-void FPCGTrackingManager::OnObjectChanged(const UObject* InObject, const FPCGTrackingChangeID& InChangeType, const UObject* InOriginatingChangeObject)
+void FPCGTrackingManager::OnObjectChanged(const UObject* InObject, const FPCGTrackingChangeID& InChangeID, const UObject* InOriginatingChangeObject)
 {
 	TRACE_CPUPROFILER_EVENT_SCOPE(FPCGTrackingManager::OnObjectChanged);
 
@@ -1433,23 +1433,23 @@ void FPCGTrackingManager::OnObjectChanged(const UObject* InObject, const FPCGTra
 		return;
 	}
 
-	const bool bNoRefreshOwner = InChangeType == PCGTrackingManager::ExecutionSourceGenerated;
+	const bool bNoRefreshOwner = InChangeID == PCGTrackingManager::ExecutionSourceGenerated;
 
 	FPCGSelectionKey SelectionKey = FPCGSelectionKey::CreateFromObjectPtr(InObject);
 	OnSelectionKeyChanged(SelectionKey, InOriginatingChangeObject, {}, bNoRefreshOwner, {FBox(EForceInit::ForceInit)});
 }
 
-void FPCGTrackingManager::NotifyObjectChanged(const UObject* InObject, const FPCGTrackingChangeID& InChangeType, const UObject* InOriginatingChangeObject)
+void FPCGTrackingManager::NotifyObjectChanged(const UObject* InObject, const FPCGTrackingChangeID& InChangeID, const UObject* InOriginatingChangeObject)
 {
 	for (TUniquePtr<IPCGChangeTracker>& ChangeTracker : ChangeTrackers)
 	{
-		if (ChangeTracker->OnNotifyObjectChanged(InObject, InChangeType, InOriginatingChangeObject))
+		if (ChangeTracker->OnNotifyObjectChanged(InObject, InChangeID, InOriginatingChangeObject))
 		{
 			return;
 		}
 	}
 
-	OnObjectChanged(InObject, InChangeType, InOriginatingChangeObject);
+	OnObjectChanged(InObject, InChangeID, InOriginatingChangeObject);
 }
 
 void FPCGTrackingManager::OnSelectionKeyChanged(const FPCGSelectionKey& InSelectionKey)
@@ -1481,7 +1481,7 @@ void FPCGTrackingManager::OnSelectionKeyChanged(const FPCGSelectionKey& InSelect
 	TSharedRef<IPCGChangeHandler::FPCGChangeHandlerChange> Change = MakeShared<IPCGChangeHandler::FPCGChangeHandlerChange>();
 	Change->ChangedObject = InSelectionKey.GetObjectFromPath();
 	Change->OriginatingChangeObject = InOriginatingChangeObject;
-	Change->bInNoRefreshOwner = bInNoRefreshOwner;
+	Change->bNoRefreshOwner = bInNoRefreshOwner;
 	
 	auto Gather = [InSelectionKey, InOriginatingSelectionKey, &InRemovedTags, &Change](const TMap<FPCGSelectionKey, TSet<IPCGGraphExecutionSource*>>& InMap, TSet<IPCGGraphExecutionSource*>& OutSet)
 	{
diff --git a/Engine/Plugins/PCG/Source/PCG/Public/ChangeTracking/PCGChangeTrackingRegistry.h b/Engine/Plugins/PCG/Source/PCG/Public/ChangeTracking/PCGChangeTrackingRegistry.h
--- a/Engine/Plugins/PCG/Source/PCG/Public/ChangeTracking/PCGChangeTrackingRegistry.h
+++ b/Engine/Plugins/PCG/Source/PCG/Public/ChangeTracking/PCGChangeTrackingRegistry.h
@@ -32,7 +32,7 @@ public:
 	virtual bool OnObjectPropertyChanged(const UObject* InObject, FPropertyChangedEvent& InEvent) { return false; }
 
 	// Can be implemented to react to other tracker events (currently used to send out Landscape tracker event to Actor tracker).
-	virtual bool OnNotifyObjectChanged(const UObject* InObject, const FPCGTrackingChangeID& InChangeType, const UObject* InOriginatingChangeObject) { return false; }
+	virtual bool OnNotifyObjectChanged(const UObject* InObject, const FPCGTrackingChangeID& InChangeID, const UObject* InOriginatingChangeObject) { return false; }
 
 	// Called from FPCGTrackingManager::OnSelectionKeyChanged to allow skipping the execution source refresh
 	virtual bool ShouldSkipRefresh(const UObject* InChangedObject) { return false; }
@@ -61,7 +61,7 @@ public:
 		// Tracked keys that match the ChangedObject or OriginatingChangedObject
 		TArray<FPCGSelectionKey> MatchedKeys;
 		// Flag passed down to FPCGTrackingManager::OnSelectionKeyChanged
-		bool bInNoRefreshOwner = false;
+		bool bNoRefreshOwner = false;
 	};
 
 	virtual ~IPCGChangeHandler() = default;
diff --git a/Engine/Plugins/PCG/Source/PCGEditor/Internal/ChangeTracking/PCGActorTracker.h b/Engine/Plugins/PCG/Source/PCGEditor/Internal/ChangeTracking/PCGActorTracker.h
--- a/Engine/Plugins/PCG/Source/PCGEditor/Internal/ChangeTracking/PCGActorTracker.h
+++ b/Engine/Plugins/PCG/Source/PCGEditor/Internal/ChangeTracking/PCGActorTracker.h
@@ -33,7 +33,7 @@ public:
 	virtual void Tick() override;
 
 	virtual bool OnObjectPropertyChanged(const UObject* InObject, FPropertyChangedEvent& InEvent) override;
-	virtual bool OnNotifyObjectChanged(const UObject* InObject, const FPCGTrackingChangeID& InChangeType, const UObject* InOriginatingChangeObject) override;
+	virtual bool OnNotifyObjectChanged(const UObject* InObject, const FPCGTrackingChangeID& InChangeID, const UObject* InOriginatingChangeObject) override;
 
 private:
 	// Typedef to store the previous position and previous tags of changed actors.
@@ -49,9 +49,9 @@ private:
 	void OnActorLoaded(AActor& InActor);
 	void OnActorAdded(AActor* InActor);
 	void OnActorAdded_Internal(const AActor* InActor, bool bShouldDirty);
-	void OnActorChanged_Internal(const AActor* InActor, const UObject* InOriginatingChangeObject, const FPCGTrackingChangeID& InChangeType);
-	void OnActorChanged_Recursive(const AActor* InActor, const UObject* InOriginatingChangeObject, const FPCGTrackingChangeID& InChangeType);
-	void OnActorChanged_Impl(const AActor* InActor, const UObject* InOriginatingChangeObject, const FPCGTrackingChangeID& InChangeType, FActorPreviousData* InPreviousData);
+	void OnActorChanged_Internal(const AActor* InActor, const UObject* InOriginatingChangeObject, const FPCGTrackingChangeID& InChangeID);
+	void OnActorChanged_Recursive(const AActor* InActor, const UObject* InOriginatingChangeObject, const FPCGTrackingChangeID& InChangeID);
+	void OnActorChanged_Impl(const AActor* InActor, const UObject* InOriginatingChangeObject, const FPCGTrackingChangeID& InChangeID, FActorPreviousData* InPreviousData);
 	void OnActorUnloaded(AActor& InActor);
 	void OnActorDeleted(AActor* InActor);
 	void OnActorDeleted_Internal(const AActor* InActor, bool bShouldDirty);
diff --git a/Engine/Plugins/PCG/Source/PCGEditor/Private/ChangeTracking/PCGActorTracker.cpp b/Engine/Plugins/PCG/Source/PCGEditor/Private/ChangeTracking/PCGActorTracker.cpp
--- a/Engine/Plugins/PCG/Source/PCGEditor/Private/ChangeTracking/PCGActorTracker.cpp
+++ b/Engine/Plugins/PCG/Source/PCGEditor/Private/ChangeTracking/PCGActorTracker.cpp
@@ -289,7 +289,7 @@ void FPCGActorTracker::OnActorAdded_Internal(const AActor* InActor, bool bShould
 	DelayedChangedActors.Emplace(InActor, { nullptr, PCGActorTracker::ActorAdded });
 }
 
-void FPCGActorTracker::OnActorChanged_Internal(const AActor* InActor, const UObject* InOriginatingChangeObject, const FPCGTrackingChangeID& InChangeType)
+void FPCGActorTracker::OnActorChanged_Internal(const AActor* InActor, const UObject* InOriginatingChangeObject, const FPCGTrackingChangeID& InChangeID)
 {
 	// Currently delayed, return
 	if (DelayedChangedActors.Contains(InActor))
@@ -300,26 +300,26 @@ void FPCGActorTracker::OnActorChanged_Internal(const AActor* InActor, const UObj
 	// Should still be delayed, add to delayed actors, return
 	if (ShouldDelayActor(InActor))
 	{
-		DelayedChangedActors.Emplace(InActor, { InOriginatingChangeObject, InChangeType });
+		DelayedChangedActors.Emplace(InActor, { InOriginatingChangeObject, InChangeID });
 		return;
 	}
 
 	// Call Recursive method
-	OnActorChanged_Recursive(InActor, InOriginatingChangeObject, InChangeType);
+	OnActorChanged_Recursive(InActor, InOriginatingChangeObject, InChangeID);
 }
 
-void FPCGActorTracker::OnActorChanged_Recursive(const AActor* InActor, const UObject* InOriginatingChangeObject, const FPCGTrackingChangeID& InChangeType)
+void FPCGActorTracker::OnActorChanged_Recursive(const AActor* InActor, const UObject* InOriginatingChangeObject, const FPCGTrackingChangeID& InChangeID)
 {
 	check(!DelayedChangedActors.Contains(InActor));
 
 	// Process actors recursively
-	auto OnActorChanged = [this, InOriginatingChangeObject, InChangeType](const AActor* InActor, auto&& OnActorChangedRecursive) -> void
+	auto OnActorChanged = [this, InOriginatingChangeObject, InChangeID](const AActor* InActor, auto&& OnActorChangedRecursive) -> void
 	{
 		if (InActor)
 		{
 			FActorPreviousData* PreviousData = ActorToPreviousDataMap.Find(InActor);
 			
-			OnActorChanged_Impl(InActor, InOriginatingChangeObject, InChangeType, PreviousData);
+			OnActorChanged_Impl(InActor, InOriginatingChangeObject, InChangeID, PreviousData);
 
 			if (PreviousData)
 			{
@@ -337,7 +337,7 @@ void FPCGActorTracker::OnActorChanged_Recursive(const AActor* InActor, const UOb
 	OnActorChanged(InActor, OnActorChanged);
 }
 
-void FPCGActorTracker::OnActorChanged_Impl(const AActor* InActor, const UObject* InOriginatingChangeObject, const FPCGTrackingChangeID& InChangeType, FActorPreviousData* InPreviousData)
+void FPCGActorTracker::OnActorChanged_Impl(const AActor* InActor, const UObject* InOriginatingChangeObject, const FPCGTrackingChangeID& InChangeID, FActorPreviousData* InPreviousData)
 {
 	check(InActor);
 	
@@ -347,7 +347,7 @@ void FPCGActorTracker::OnActorChanged_Impl(const AActor* InActor, const UObject*
 	}
 
 	// Discard everything related to self for actor added/deleted
-	const bool bNoRefreshOwner = InChangeType == PCGActorTracker::ActorAdded || InChangeType == PCGActorTracker::ActorDeleted;
+	const bool bNoRefreshOwner = InChangeID == PCGActorTracker::ActorAdded || InChangeID == PCGActorTracker::ActorDeleted;
 
 	ensure(InActor->GetWorld() == Owner->GetWorld());
 
@@ -470,9 +470,9 @@ void FPCGActorTracker::OnObjectModified(UObject* InObject)
 	StorePreviousData(Actor, StorePreviousData);
 }
 
-bool FPCGActorTracker::OnNotifyObjectChanged(const UObject* InObject, const FPCGTrackingChangeID& InChangeType, const UObject* InOriginatingChangeObject)
+bool FPCGActorTracker::OnNotifyObjectChanged(const UObject* InObject, const FPCGTrackingChangeID& InChangeID, const UObject* InOriginatingChangeObject)
 {
-	if (InChangeType != PCGActorTracker::ActorChanged)
+	if (InChangeID != PCGActorTracker::ActorChanged)
 	{
 		return false;
 	}
@@ -483,7 +483,7 @@ bool FPCGActorTracker::OnNotifyObjectChanged(const UObject* InObject, const FPCG
 		return true;
 	}
 
-	OnActorChanged_Internal(Actor, InOriginatingChangeObject, InChangeType);
+	OnActorChanged_Internal(Actor, InOriginatingChangeObject, InChangeID);
 	return true;
 }
 
diff --git a/Engine/Plugins/PCG/Source/PCGEditor/Private/ChangeTracking/PCGComponentChangeHandler.cpp b/Engine/Plugins/PCG/Source/PCGEditor/Private/ChangeTracking/PCGComponentChangeHandler.cpp
--- a/Engine/Plugins/PCG/Source/PCGEditor/Private/ChangeTracking/PCGComponentChangeHandler.cpp
+++ b/Engine/Plugins/PCG/Source/PCGEditor/Private/ChangeTracking/PCGComponentChangeHandler.cpp
@@ -214,7 +214,7 @@ bool FPCGComponentChangeHandler::ShouldDiscardComponent(UPCGComponent* InCompone
 
 	if (!InComponent
 		|| CurrentChange->ChangedObjects.Contains(InComponent)
-		|| (Change->bInNoRefreshOwner && CurrentChange->ChangedObjects.Contains(InComponent->GetOwner()))
+		|| (Change->bNoRefreshOwner && CurrentChange->ChangedObjects.Contains(InComponent->GetOwner()))
 		|| InComponent->IsCleaningUp()
 		|| InComponent->IsAnyObjectManagedByResource(CurrentChange->ChangedObjects)
 		|| (CurrentChange->ChangedObjectDynamicTrackingPriority.IsSet() && InComponent->GetExecutionState().GetDynamicTrackingPriority() < CurrentChange->ChangedObjectDynamicTrackingPriority.GetValue()))
@@ -276,7 +276,7 @@ void FPCGComponentChangeHandler::EndChangeHandling(bool bSkipRefresh)
 
 		const bool bOwnerHasChanged = Component->GetOwner() == Change->ChangedObject;
 
-		if ((!Change->bInNoRefreshOwner || !bOwnerHasChanged) && (!Component->bOnlyTrackItself || bOwnerHasChanged))
+		if ((!Change->bNoRefreshOwner || !bOwnerHasChanged) && (!Component->bOnlyTrackItself || bOwnerHasChanged))
 		{
 			const AActor* Actor = Cast<AActor>(Change->ChangedObject);
 
