[PCG] Fix issue in the Static Mesh Data Provider when setuping primitives

* When setuping the primitive, it is possible the factory is not ready, and in that case it will come back to check on next tick
* But there is an early out if the descriptors are empty, and they are empty after the initialization because they are moved
* So move the check for Static mesh to be empty and compiled into the factory intialization block.

#rnx
#Synced-CL 0

diff --git a/Engine/Plugins/PCG/Source/PCG/Private/Compute/DataInterfaces/Elements/PCGStaticMeshSpawnerDataInterface.cpp b/Engine/Plugins/PCG/Source/PCG/Private/Compute/DataInterfaces/Elements/PCGStaticMeshSpawnerDataInterface.cpp
--- a/Engine/Plugins/PCG/Source/PCG/Private/Compute/DataInterfaces/Elements/PCGStaticMeshSpawnerDataInterface.cpp
+++ b/Engine/Plugins/PCG/Source/PCG/Private/Compute/DataInterfaces/Elements/PCGStaticMeshSpawnerDataInterface.cpp
@@ -927,32 +927,34 @@ void UPCGStaticMeshSpawnerDataProvider::CreatePrimitiveDescriptors(FPCGContext*
 
 bool UPCGStaticMeshSpawnerDataProvider::SetupPrimitives(FPCGContext* InContext, UPCGDataBinding* InBinding)
 {
-	if (PrimitiveDescriptors.IsEmpty())
+	FPCGSpawnerPrimitives* Primitives = InBinding->FindMeshSpawnerPrimitives(GetProducerKernel());
+	if (!Primitives)
 	{
 		return true;
 	}
 
 	const UPCGStaticMeshSpawnerSettings* Settings = CastChecked<UPCGStaticMeshSpawnerSettings>(GetProducerKernel()->GetSettings());
 
-	FPCGSpawnerPrimitives* Primitives = InBinding->FindMeshSpawnerPrimitives(GetProducerKernel());
-	if (!Primitives)
-	{
-		return true;
-	}
+	TSharedPtr<IPCGPrimitiveFactoryISMBase> Factory;
 
-	// Ensure SMs are compiled, required to prevent ensure when partially built SM is accessed later in FastGeo code.
-	for (const FPCGProceduralISMComponentDescriptor& Desc : PrimitiveDescriptors)
+	if (!Primitives->PrimitiveFactory)
 	{
-		if (ensure(Desc.StaticMesh.IsValid()) && Desc.StaticMesh->IsCompiling())
+		// If we have no descriptors, we can early out. Only checked when the primitive factory is not initialized,
+		// as we're moving the descritors into the factory (which will emtpy PrimitiveDescriptors)
+		if (PrimitiveDescriptors.IsEmpty())
 		{
-			return false;
+			return true;
 		}
-	}
 
-	TSharedPtr<IPCGPrimitiveFactoryISMBase> Factory;
+		// Ensure SMs are compiled, required to prevent ensure when partially built SM is accessed later in FastGeo code.
+		for (const FPCGProceduralISMComponentDescriptor& Desc : PrimitiveDescriptors)
+		{
+			if (ensure(Desc.StaticMesh.IsValid()) && Desc.StaticMesh->IsCompiling())
+			{
+				return false;
+			}
+		}
 
-	if (!Primitives->PrimitiveFactory)
-	{
 		bool bCreateComponentless = false;
 
 		if (PCGStaticMeshSpawnerDataInterface::CVarCreatePrimitivesComponentless.GetValueOnGameThread() && PCGPrimitiveFactoryHelpers::GetFastGeoPrimitiveFactory())
