[PCG] Added options to control tag behavior in merge nodes.

#Synced-CL 48730134

diff --git a/Engine/Plugins/PCG/Source/PCG/Private/Elements/PCGMergeAttributes.cpp b/Engine/Plugins/PCG/Source/PCG/Private/Elements/PCGMergeAttributes.cpp
--- a/Engine/Plugins/PCG/Source/PCG/Private/Elements/PCGMergeAttributes.cpp
+++ b/Engine/Plugins/PCG/Source/PCG/Private/Elements/PCGMergeAttributes.cpp
@@ -74,13 +74,20 @@ bool FPCGMergeAttributesElement::ExecuteInternal(FPCGContext* Context) const
 	check(Settings);
 
 	TArray<FPCGTaggedData> Sources;
+	int NumSourcesOnFirstPin = -1;
 	for (FName PinLabel : Settings->GetNodeDefinedPinLabels())
 	{
 		Sources.Append(Context->InputData.GetInputsByPin(PinLabel));
+
+		if (NumSourcesOnFirstPin == -1)
+		{
+			NumSourcesOnFirstPin = Sources.Num();
+		}
 	}
 
-	for (const FPCGTaggedData& Source : Sources)
+	for (int SourceIndex = 0; SourceIndex < Sources.Num(); ++SourceIndex)
 	{
+		const FPCGTaggedData& Source = Sources[SourceIndex];
 		const UPCGParamData* SourceData = Cast<const UPCGParamData>(Source.Data);
 
 		if (!SourceData)
@@ -100,6 +107,14 @@ bool FPCGMergeAttributesElement::ExecuteInternal(FPCGContext* Context) const
 		if (!MergedOutput)
 		{
 			MergedOutput = &Outputs.Add_GetRef(Source);
+
+			if (Settings->MergeTagsMode == EPCGMergeTagsMode::None ||
+				(Settings->MergeTagsMode == EPCGMergeTagsMode::SelectTagsFromFirstData && SourceIndex != 0) ||
+				(Settings->MergeTagsMode == EPCGMergeTagsMode::MergeTagsFromFirstPin && SourceIndex >= NumSourcesOnFirstPin))
+			{
+				MergedOutput->Tags.Reset();
+			}
+
 			continue;
 		}
 		
@@ -130,7 +145,11 @@ bool FPCGMergeAttributesElement::ExecuteInternal(FPCGContext* Context) const
 		MergedAttributeSet->Metadata->SetAttributes(SourceEntryKeys, SourceMetadata, EntryKeys, Context);
 
 		// - Merge tags too (to be in line with the Merge points node)
-		MergedOutput->Tags.Append(Source.Tags);
+		if (Settings->MergeTagsMode == EPCGMergeTagsMode::MergeAllTags ||
+			(Settings->MergeTagsMode == EPCGMergeTagsMode::MergeTagsFromFirstPin && SourceIndex < NumSourcesOnFirstPin))
+		{
+			MergedOutput->Tags.Append(Source.Tags);
+		}
 	}
 
 	return true;
diff --git a/Engine/Plugins/PCG/Source/PCG/Private/Elements/PCGMergeElement.cpp b/Engine/Plugins/PCG/Source/PCG/Private/Elements/PCGMergeElement.cpp
--- a/Engine/Plugins/PCG/Source/PCG/Private/Elements/PCGMergeElement.cpp
+++ b/Engine/Plugins/PCG/Source/PCG/Private/Elements/PCGMergeElement.cpp
@@ -60,12 +60,19 @@ bool FPCGMergeElement::ExecuteInternal(FPCGContext* Context) const
 	check(Context);
 
 	const UPCGMergeSettings* Settings = Context->GetInputSettings<UPCGMergeSettings>();
-	const bool bMergeMetadata = !Settings || Settings->bMergeMetadata;
+	check(Settings);
+	const bool bMergeMetadata = Settings->bMergeMetadata;
 
 	TArray<FPCGTaggedData> Sources;
+	int NumSourcesOnFirstPin = -1;
 	for (FName PinLabel : Settings->GetNodeDefinedPinLabels())
 	{
 		Sources.Append(Context->InputData.GetInputsByPin(PinLabel));
+
+		if (NumSourcesOnFirstPin == -1)
+		{
+			NumSourcesOnFirstPin = Sources.Num();
+		}
 	}
 
 	TArray<FPCGTaggedData>& Outputs = Context->OutputData.TaggedData;
@@ -78,8 +85,9 @@ bool FPCGMergeElement::ExecuteInternal(FPCGContext* Context) const
 
 	// Prepare data & metadata
 	// Done in two passes for futureproofing - expecting changes in the metadata attribute creation vs. usage in points
-	for (const FPCGTaggedData& Source : Sources)
+	for(int SourceIndex = 0; SourceIndex < Sources.Num(); ++SourceIndex)
 	{
+		const FPCGTaggedData& Source = Sources[SourceIndex];
 		const UPCGBasePointData* SourcePointData = Cast<const UPCGBasePointData>(Source.Data);
 
 		if (!SourcePointData)
@@ -100,6 +108,13 @@ bool FPCGMergeElement::ExecuteInternal(FPCGContext* Context) const
 		if (!TargetTaggedData)
 		{
 			TargetTaggedData = &(Outputs.Add_GetRef(Source));
+
+			if (Settings->MergeTagsMode == EPCGMergeTagsMode::None ||
+				(Settings->MergeTagsMode == EPCGMergeTagsMode::SelectTagsFromFirstData && SourceIndex != 0) ||
+				(Settings->MergeTagsMode == EPCGMergeTagsMode::MergeTagsFromFirstPin && SourceIndex >= NumSourcesOnFirstPin))
+			{
+				TargetTaggedData->Tags.Reset();
+			}
 		}
 		else if (!TargetPointData)
 		{
@@ -123,7 +138,12 @@ bool FPCGMergeElement::ExecuteInternal(FPCGContext* Context) const
 			}
 			
 			check(TargetTaggedData);
-			TargetTaggedData->Tags.Append(Source.Tags);
+			// By definition once we're here, it's the second data.
+			if (Settings->MergeTagsMode == EPCGMergeTagsMode::MergeAllTags ||
+				(Settings->MergeTagsMode == EPCGMergeTagsMode::MergeTagsFromFirstPin && SourceIndex < NumSourcesOnFirstPin))
+			{
+				TargetTaggedData->Tags.Append(Source.Tags);
+			}
 		}
 	}
 
diff --git a/Engine/Plugins/PCG/Source/PCG/Public/Elements/PCGMergeAttributes.h b/Engine/Plugins/PCG/Source/PCG/Public/Elements/PCGMergeAttributes.h
--- a/Engine/Plugins/PCG/Source/PCG/Public/Elements/PCGMergeAttributes.h
+++ b/Engine/Plugins/PCG/Source/PCG/Public/Elements/PCGMergeAttributes.h
@@ -3,6 +3,7 @@
 #pragma once
 
 #include "PCGSettingsWithDynamicInputs.h"
+#include "PCGMergeElement.h"
 
 #include "PCGMergeAttributes.generated.h"
 
@@ -37,6 +38,9 @@ public:
 	virtual void AddDefaultDynamicInputPin() override;
 #endif // WITH_EDITOR
 	//~End UPCGDynamicSettings interface
+
+	UPROPERTY(BlueprintReadWrite, EditAnywhere, Category = Settings, meta = (PCG_Overridable))
+	EPCGMergeTagsMode MergeTagsMode = EPCGMergeTagsMode::MergeAllTags;
 };
 
 class FPCGMergeAttributesElement : public IPCGElement
diff --git a/Engine/Plugins/PCG/Source/PCG/Public/Elements/PCGMergeElement.h b/Engine/Plugins/PCG/Source/PCG/Public/Elements/PCGMergeElement.h
--- a/Engine/Plugins/PCG/Source/PCG/Public/Elements/PCGMergeElement.h
+++ b/Engine/Plugins/PCG/Source/PCG/Public/Elements/PCGMergeElement.h
@@ -6,6 +6,15 @@
 
 #include "PCGMergeElement.generated.h"
 
+UENUM(BlueprintType)
+enum class EPCGMergeTagsMode : uint8
+{
+	MergeAllTags,
+	MergeTagsFromFirstPin,
+	SelectTagsFromFirstData,
+	None
+};
+
 /** Merges multiple data sources (currently only points supported) into a single output. */
 UCLASS(MinimalAPI, BlueprintType, Classgroup = (Procedural))
 class UPCGMergeSettings : public UPCGSettingsWithDynamicInputs
@@ -41,6 +50,9 @@ public:
 	/** Controls whether the resulting merge data will have any metadata */
 	UPROPERTY(BlueprintReadWrite, EditAnywhere, Category = Settings, meta = (PCG_Overridable))
 	bool bMergeMetadata = true;
+
+	UPROPERTY(BlueprintReadWrite, EditAnywhere, Category = Settings, meta = (PCG_Overridable))
+	EPCGMergeTagsMode MergeTagsMode = EPCGMergeTagsMode::MergeAllTags;
 };
 
 class FPCGMergeElement : public IPCGElement
