[PCG] Moved tag assignment when creating components so overriding component tags isn't destructive.

#rnx
#Synced-CL 48730134

diff --git a/Engine/Plugins/PCG/Source/PCG/Private/Elements/PCGCreateTargetActor.cpp b/Engine/Plugins/PCG/Source/PCG/Private/Elements/PCGCreateTargetActor.cpp
--- a/Engine/Plugins/PCG/Source/PCG/Private/Elements/PCGCreateTargetActor.cpp
+++ b/Engine/Plugins/PCG/Source/PCG/Private/Elements/PCGCreateTargetActor.cpp
@@ -487,6 +487,20 @@ bool FPCGCreateTargetActorElement::ExecuteInternal(FPCGContext* Context) const
 			return true;
 		}
 
+#if WITH_EDITOR
+		if (!Settings->ActorLabel.IsEmpty())
+		{
+			GeneratedActor->SetActorLabel(Settings->ActorLabel);
+		}
+#endif
+
+		// Always attach if root actor is provided
+		PCGHelpers::AttachToParent(GeneratedActor, TargetActor, Settings->RootActor.Get() ? EPCGAttachOptions::Attached : Settings->AttachOptions, Context);
+
+		// Apply property overrides to the GeneratedActor
+		PCGObjectPropertyOverrideHelpers::ApplyOverridesFromParams(Settings->PropertyOverrideDescriptions, GeneratedActor, PCGCreateTargetActorConstants::ActorPropertyOverridesLabel, Context);
+
+		// Finally add tags
 		GeneratedActor->Tags.Add(PCGHelpers::DefaultPCGActorTag);
 
 		if (!Settings->CommaSeparatedActorTags.IsEmpty())
@@ -506,19 +520,6 @@ bool FPCGCreateTargetActorElement::ExecuteInternal(FPCGContext* Context) const
 		{
 			GeneratedActor->Tags.AddUnique(Tag);
 		}
-
-#if WITH_EDITOR
-		if (!Settings->ActorLabel.IsEmpty())
-		{
-			GeneratedActor->SetActorLabel(Settings->ActorLabel);
-		}
-#endif
-
-		// Always attach if root actor is provided
-		PCGHelpers::AttachToParent(GeneratedActor, TargetActor, Settings->RootActor.Get() ? EPCGAttachOptions::Attached : Settings->AttachOptions, Context);
-
-		// Apply property overrides to the GeneratedActor
-		PCGObjectPropertyOverrideHelpers::ApplyOverridesFromParams(Settings->PropertyOverrideDescriptions, GeneratedActor, PCGCreateTargetActorConstants::ActorPropertyOverridesLabel, Context);
 	}
 	
 	for (UFunction* Function : PCGHelpers::FindUserFunctions(GeneratedActor->GetClass(), Settings->PostProcessFunctionNames, { UPCGFunctionPrototypes::GetPrototypeWithNoParams() }, Context))
diff --git a/Engine/Plugins/PCG/Source/PCG/Private/Elements/PCGSpawnSpline.cpp b/Engine/Plugins/PCG/Source/PCG/Private/Elements/PCGSpawnSpline.cpp
--- a/Engine/Plugins/PCG/Source/PCG/Private/Elements/PCGSpawnSpline.cpp
+++ b/Engine/Plugins/PCG/Source/PCG/Private/Elements/PCGSpawnSpline.cpp
@@ -106,15 +106,6 @@ bool FPCGSpawnSplineElement::ExecuteInternal(FPCGContext* InContext) const
 
 		check(IsInGameThread());
 		USplineComponent* SplineComponent = NewObject<USplineComponent>(TargetActor, *SplineComponentToSpawn);
-		SplineComponent->ComponentTags.Add(SourceComponent->GetFName());
-		SplineComponent->ComponentTags.Add(PCGHelpers::DefaultPCGTag);
-		for (const FName& Tag : Settings->TagsToAddOnComponents)
-		{
-			SplineComponent->ComponentTags.AddUnique(Tag);
-		}
-
-		// Also propagate the tags on the tagged data to the component.
-		Algo::Transform(InputData.Tags, SplineComponent->ComponentTags, [](const FString& Tag) { return FName(Tag); });
 
 		SplineComponent->RegisterComponent();
 		TargetActor->AddInstanceComponent(SplineComponent);
@@ -144,6 +135,18 @@ bool FPCGSpawnSplineElement::ExecuteInternal(FPCGContext* InContext) const
 			DescriptorOverride.Apply(0);
 		}
 
+		// Add tags, both the default ones (source name + PCG Generated) and the ones from the override pin.
+		// We do this after the override in the case that the override was writing the tags.
+		SplineComponent->ComponentTags.Add(SourceComponent->GetFName());
+		SplineComponent->ComponentTags.Add(PCGHelpers::DefaultPCGTag);
+		for (const FName& Tag : Settings->TagsToAddOnComponents)
+		{
+			SplineComponent->ComponentTags.AddUnique(Tag);
+		}
+
+		// Also propagate the tags on the tagged data to the component.
+		Algo::Transform(InputData.Tags, SplineComponent->ComponentTags, [](const FString& Tag) { return FName(Tag); });
+
 		// Execute PostProcess Functions
 		if (TargetActor)
 		{
diff --git a/Engine/Plugins/PCG/Source/PCG/Private/Elements/PCGSpawnSplineMesh.cpp b/Engine/Plugins/PCG/Source/PCG/Private/Elements/PCGSpawnSplineMesh.cpp
--- a/Engine/Plugins/PCG/Source/PCG/Private/Elements/PCGSpawnSplineMesh.cpp
+++ b/Engine/Plugins/PCG/Source/PCG/Private/Elements/PCGSpawnSplineMesh.cpp
@@ -313,6 +313,10 @@ bool FPCGSpawnSplineMeshElement::ExecuteInternal(FPCGContext* InContext) const
 				{
 					IterState.ComponentOverrides.UpdateTemplateObject(SplineMeshComponent);
 					IterState.ComponentOverrides.Apply(IterState.ElementIndex);
+
+					// Since the component overrides can override the component tags, we need to ensure that the default tags are still present
+					SplineMeshComponent->ComponentTags.AddUnique(SourceComponent->GetFName());
+					SplineMeshComponent->ComponentTags.AddUnique(PCGHelpers::DefaultPCGTag);
 				}
 
 				// Apply tags on the Spline Mesh Component
