[PCG] Added more options to move resources, namely extracting to specified actor, label configuration and tags to exclude specific resources.

#rnx
#Synced-CL 48730134

diff --git a/Engine/Plugins/PCG/Source/PCG/Private/PCGComponent.cpp b/Engine/Plugins/PCG/Source/PCG/Private/PCGComponent.cpp
--- a/Engine/Plugins/PCG/Source/PCG/Private/PCGComponent.cpp
+++ b/Engine/Plugins/PCG/Source/PCG/Private/PCGComponent.cpp
@@ -1129,31 +1129,46 @@ AActor* UPCGComponent::ClearPCGLink(const FPCGMoveResourceParams& InParams)
 		return nullptr;
 	}
 
+	// Create target actor if needed.
 	UWorld* World = GetWorld();
-	
-	FActorSpawnParameters ActorSpawnParams;
-	UClass* SpawnClass = Params.GetTemplateClass<AActor>() ? Params.GetTemplateClass<AActor>().Get() : AActor::StaticClass();
-	const FString DefaultName(TEXT("PCGStamp"));
-	ActorSpawnParams.Name = *DefaultName;
-	ActorSpawnParams.NameMode = FActorSpawnParameters::ESpawnActorNameMode::Requested;
-	ActorSpawnParams.OverrideLevel = GetOwner()->GetLevel();
+	AActor* NewActor = nullptr;
 
-	UPCGActorHelpers::FSpawnDefaultActorParams SpawnDefaultActorParams(World, SpawnClass, GetOwner()->GetTransform(), ActorSpawnParams);
+	if (!Params.Target)
+	{
+		FActorSpawnParameters ActorSpawnParams;
+		UClass* SpawnClass = Params.GetTemplateClass<AActor>() ? Params.GetTemplateClass<AActor>().Get() : AActor::StaticClass();
 
-#if WITH_EDITOR
-	SpawnDefaultActorParams.DataLayerInstances = GetOwner()->GetDataLayerInstances();
-	SpawnDefaultActorParams.HLODLayer = GetOwner()->GetHLODLayer();
-#endif
+		const FName DefaultName = Params.DefaultName == NAME_None ? TEXT("PCGStamp") : Params.DefaultName;
 
-	// First create a new actor that will be the new owner of all the resources
-	AActor* NewActor = UPCGActorHelpers::SpawnDefaultActor(SpawnDefaultActorParams);
-#if WITH_EDITOR
-	FActorLabelUtilities::SetActorLabelUnique(NewActor, GetOwner()->GetActorLabel() + TEXT("_") + DefaultName);
-	NewActor->SetFolderPath(GetOwner()->GetFolderPath());
-#endif
+		ActorSpawnParams.Name = DefaultName;
+		ActorSpawnParams.NameMode = FActorSpawnParameters::ESpawnActorNameMode::Requested;
+		ActorSpawnParams.OverrideLevel = GetOwner()->GetLevel();
 
-	// Affect newly created actor as the target actor for the subsequent resource move.
-	Params.Target = NewActor;
+		UPCGActorHelpers::FSpawnDefaultActorParams SpawnDefaultActorParams(World, SpawnClass, GetOwner()->GetTransform(), ActorSpawnParams);
+	
+		if (Params.bAttachToParent)
+		{
+			SpawnDefaultActorParams.Parent = GetOwner();
+		}
+
+	#if WITH_EDITOR
+		SpawnDefaultActorParams.DataLayerInstances = GetOwner()->GetDataLayerInstances();
+		SpawnDefaultActorParams.HLODLayer = GetOwner()->GetHLODLayer();
+	#endif
+
+		// First create a new actor that will be the new owner of all the resources
+		NewActor = UPCGActorHelpers::SpawnDefaultActor(SpawnDefaultActorParams);
+	#if WITH_EDITOR
+		if (NewActor)
+		{
+			FActorLabelUtilities::SetActorLabelUnique(NewActor, GetOwner()->GetActorLabel() + TEXT("_") + DefaultName.ToString());
+			NewActor->SetFolderPath(GetOwner()->GetFolderPath());
+		}
+	#endif
+
+		// Affect newly created actor as the target actor for the subsequent resource move.
+		Params.Target = NewActor;
+	}
 
 	// Then move all resources linked to this component to this actor
 	bool bHasMovedResources = MoveResourcesToNewActor(Params, /*bCreateChild=*/false);
@@ -1176,7 +1191,7 @@ AActor* UPCGComponent::ClearPCGLink(const FPCGMoveResourceParams& InParams)
 			DirtyGenerated(EPCGComponentDirtyFlag::Input, /*bDispatchToLocalComponents=*/true);
 #endif
 		}
-		else
+		else if(NewActor)
 		{
 			NewActor->Destroy();
 			NewActor = nullptr;
@@ -1367,7 +1382,11 @@ bool UPCGComponent::MoveResourcesToNewActor(const FPCGMoveResourceParams& Params
 
 		CreatedChildActor = UPCGActorHelpers::SpawnDefaultActor(SpawnDefaultActorParams);
 #if WITH_EDITOR
-		FActorLabelUtilities::SetActorLabelUnique(CreatedChildActor, DefaultName);
+
+		if (CreatedChildActor)
+		{
+			FActorLabelUtilities::SetActorLabelUnique(CreatedChildActor, DefaultName);
+		}
 
 		if (TargetActor->GetFolderPath() != NAME_None)
 		{
diff --git a/Engine/Plugins/PCG/Source/PCG/Private/PCGManagedResource.cpp b/Engine/Plugins/PCG/Source/PCG/Private/PCGManagedResource.cpp
--- a/Engine/Plugins/PCG/Source/PCG/Private/PCGManagedResource.cpp
+++ b/Engine/Plugins/PCG/Source/PCG/Private/PCGManagedResource.cpp
@@ -7,7 +7,9 @@
 #include "Helpers/PCGActorHelpers.h"
 #include "Helpers/PCGHelpers.h"
 
+#include "Algo/AllOf.h"
 #include "Algo/AnyOf.h"
+#include "Algo/NoneOf.h"
 #include "Components/InstancedStaticMeshComponent.h"
 #include "Components/InstancedSkinnedMeshComponent.h"
 #include "Components/SplineMeshComponent.h"
@@ -38,17 +40,23 @@ bool FPCGMoveResourceParams::Accepts(const UObject* InObject) const
 	{
 		return false;
 	}
-	else if (RequiredTags.IsEmpty()) // Always accept if there are no requirements
+	else if (RequiredTags.IsEmpty() && ExcludedTags.IsEmpty()) // Always accept if there are no requirements
 	{
 		return true;
 	}
 	else if (const AActor* Actor = Cast<AActor>(InObject))
 	{
-		return Algo::AllOf(RequiredTags, [Actor](FName Tag) { return Actor->Tags.Contains(Tag); });
+		auto TestActorTags = [Actor](FName Tag) { return Actor->Tags.Contains(Tag); };
+
+		return (RequiredTags.IsEmpty() || Algo::AllOf(RequiredTags, TestActorTags)) &&
+			(ExcludedTags.IsEmpty() || Algo::NoneOf(ExcludedTags, TestActorTags));
 	}
 	else if (const UActorComponent* Component = Cast<UActorComponent>(InObject))
 	{
-		return Algo::AllOf(RequiredTags, [Component](FName Tag) { return Component->ComponentHasTag(Tag); });
+		auto TestComponentTags = [Component](FName Tag) { return Component->ComponentHasTag(Tag); };
+
+		return (RequiredTags.IsEmpty() || Algo::AllOf(RequiredTags, TestComponentTags)) &&
+			(ExcludedTags.IsEmpty() || Algo::NoneOf(ExcludedTags, TestComponentTags));
 	}
 	else
 	{
@@ -63,13 +71,16 @@ bool FPCGMoveResourceParams::Accepts(const FPCGTaggedData& InData) const
 	{
 		return false;
 	}
-	else if (RequiredTags.IsEmpty())
+	else if (RequiredTags.IsEmpty() && ExcludedTags.IsEmpty())
 	{
 		return true;
 	}
 	else
 	{
-		return Algo::AllOf(RequiredTags, [&InData](FName Tag) { return InData.Tags.Contains(Tag.ToString()); });
+		auto TestDataTags = [&InData](FName Tag) { return InData.Tags.Contains(Tag.ToString()); };
+
+		return (RequiredTags.IsEmpty() || Algo::AllOf(RequiredTags, TestDataTags)) &&
+			(ExcludedTags.IsEmpty() || Algo::NoneOf(ExcludedTags, TestDataTags));
 	}
 }
 
diff --git a/Engine/Plugins/PCG/Source/PCG/Public/PCGManagedResource.h b/Engine/Plugins/PCG/Source/PCG/Public/PCGManagedResource.h
--- a/Engine/Plugins/PCG/Source/PCG/Public/PCGManagedResource.h
+++ b/Engine/Plugins/PCG/Source/PCG/Public/PCGManagedResource.h
@@ -35,8 +35,14 @@ struct FPCGMoveResourceParams
 	const UObject* ExpectedPreviousOwner = nullptr;
 	/** Conjunction of tags required to move a resource to the target. */
 	TArray<FName> RequiredTags;
+	/** Disjunction of tags that would disqualify moving a resource to the target. */
+	TArray<FName> ExcludedTags;
 	/** If there is generated data on the source objects, it will be moved using the same requirements, into the target collection. */
 	FPCGDataCollection* TargetDataCollection = nullptr;
+	/** Controls whether the created target should be 'attached' to the original object. */
+	bool bAttachToParent = false;
+	/** Name Override, can be left empty. */
+	FName DefaultName = NAME_None;
 
 	template<typename ObjectType>
 	TSubclassOf<ObjectType> GetTemplateClass() const
