[PCG] Make the Graph Parameters getter nodes able to run outside of the main thread

* Generalize the concept of User Parameter data even for non-subgraphs, done in the PreGraph element
* PreGraph element is now the default pre task, if none is specified on schedule.
* PreGraph element can have an override graph in the case the execution source graph is not the scheduled graph (NoMerging Spawn Actors for example)
* Freeze the user parameters in this data in order to be able to access it in MT

#Synced-CL 48760074

diff --git a/Engine/Plugins/PCG/Source/PCG/Private/Elements/PCGSpawnActor.cpp b/Engine/Plugins/PCG/Source/PCG/Private/Elements/PCGSpawnActor.cpp
--- a/Engine/Plugins/PCG/Source/PCG/Private/Elements/PCGSpawnActor.cpp
+++ b/Engine/Plugins/PCG/Source/PCG/Private/Elements/PCGSpawnActor.cpp
@@ -8,6 +8,7 @@
 #include "PCGSubsystem.h"
 #include "Data/PCGBasePointData.h"
 #include "Data/PCGSpatialData.h"
+#include "Graph/PCGGraphExecutor.h"
 #include "Graph/PCGStackContext.h"
 #include "Grid/PCGPartitionActor.h"
 #include "Helpers/PCGActorHelpers.h"
@@ -642,7 +643,7 @@ bool FPCGSpawnActorElement::SpawnAndPrepareSubgraphs(FPCGSubgraphContext* Contex
 							FPCGTaskId SubgraphTaskId = Context->ScheduleGraph(FPCGScheduleGraphParams(
 								Graph,
 								Context->ExecutionSource.Get(),
-								MakeShared<FPCGTrivialElement>(),// TODO: prepare user parameters like in subgraph/loop
+								MakeShared<FPCGPreGraphElement>(GraphInterface),
 								MakeShared<FPCGInputForwardingElement>(SubgraphInputData),
 								/*Dependencies=*/{},
 								&InvocationStack,
@@ -1083,8 +1084,7 @@ void FPCGSpawnActorElement::SpawnActors(FPCGSubgraphContext* Context, AActor* Ta
 						Subsystem->RegisterOrUpdateExecutionSource(PCGComponent);
 					}
 
-					// TODO: use ScheduleGraph if we want to pass user parameters
-					FPCGTaskId SubgraphTaskId = PCGComponent->GenerateLocalGetTaskId(/*bForce=*/true);
+					FPCGTaskId SubgraphTaskId = Subsystem ? Subsystem->ScheduleGraph(PCGComponent, {}) : PCGComponent->GenerateLocalGetTaskId(/*bForce=*/true);
 					if (SubgraphTaskId != InvalidPCGTaskId)
 					{
 						Context->SubgraphTaskIds.Add(SubgraphTaskId);
diff --git a/Engine/Plugins/PCG/Source/PCG/Private/Elements/PCGUserParameterGet.cpp b/Engine/Plugins/PCG/Source/PCG/Private/Elements/PCGUserParameterGet.cpp
--- a/Engine/Plugins/PCG/Source/PCG/Private/Elements/PCGUserParameterGet.cpp
+++ b/Engine/Plugins/PCG/Source/PCG/Private/Elements/PCGUserParameterGet.cpp
@@ -16,6 +16,7 @@
 #include "Helpers/PCGPropertyHelpers.h"
 #include "Metadata/PCGAttributePropertySelector.h"
 
+#include "Algo/NoneOf.h"
 #include "StructUtils/PropertyBag.h"
 #include "StructUtils/StructView.h"
 #include "UObject/FortniteReleaseBranchCustomObjectVersion.h"
@@ -41,7 +42,7 @@ namespace PCGUserParameterGet
 			const UPCGGraphInterface* GraphInterface = nullptr;
 
 			// First we will read from the input if we find an input for this graph, set by the subgraph element, then use it.
-			TArray<FPCGTaggedData> UserParameterData = InContext.InputData.GetTaggedTypedInputs<UPCGUserParametersData>(PCGBaseSubgraphConstants::UserParameterTagData);
+			TArray<FPCGTaggedData> UserParameterData = InContext.InputData.GetTaggedTypedInputs<UPCGUserParametersData>(PCG::Private::UserParameterTagData);
 			if (!UserParameterData.IsEmpty())
 			{
 	#if WITH_EDITOR
@@ -62,6 +63,8 @@ namespace PCGUserParameterGet
 
 						UpstreamData = UpstreamData->UpstreamData.Get();
 					}
+
+					return Layouts;
 				}
 				else
 				{
@@ -70,7 +73,13 @@ namespace PCGUserParameterGet
 				}
 			}
 
-			// We gather the outer graph of this node, to make sure it matches our interface.
+			// We gather the outer graph of this node, to make sure it matches our interface. It must happen on the game thread
+			if (!ensure(IsInGameThread()))
+			{
+				PCGE_LOG_C(Error, LogOnly, &InContext, LOCTEXT("NoUserParameterData", "Internal error, PCG User Parameters Data is absent and we are not executing on the game thread"));
+				return {};
+			}
+
 			const UPCGGraph* GraphFromNode = Cast<UPCGGraph>(InContext.Node->GetOuter());
 			check(GraphFromNode);
 			FConstStructView GraphFromNodeParameters = GraphFromNode->GetUserParametersStruct() ? GraphFromNode->GetUserParametersStruct()->GetValue() : FConstStructView{};
@@ -534,6 +543,14 @@ bool UPCGGenericUserParameterGetSettings::ConvertNode(const FPCGPreconfiguredInf
 	return false;
 }
 
+bool FPCGUserParameterGetElement::CanExecuteOnlyOnMainThread(FPCGContext* Context) const
+{
+	// If we have no user parameter data as input, we'll have to access the parameters on the component/graph directly, which is not thread safe
+	// If there is no context, we can't check the inputs, so do the assumption it is possible to run on another thread, and we'll ask to come back to
+	// the main thread on execute in the unlikely possibility that the parameters are not there. 
+	return Context ? Algo::NoneOf(Context->InputData.TaggedData, [](const FPCGTaggedData& InData) { return InData.Data.IsA<UPCGUserParametersData>(); }) : false;
+}
+
 //////////////////////////////////////////////
 
 bool FPCGUserParameterGetElement::ExecuteInternal(FPCGContext* Context) const
@@ -541,6 +558,15 @@ bool FPCGUserParameterGetElement::ExecuteInternal(FPCGContext* Context) const
 	TRACE_CPUPROFILER_EVENT_SCOPE(FPCGUserParameterGetElement::Execute);
 
 	check(Context);
+	// If we have no user parameter data as input, we'll have to access the parameters on the component/graph directly, which is not thread safe
+	if (CanExecuteOnlyOnMainThread(Context) && !IsInGameThread())
+	{
+		PCGLog::LogWarningOnGraph(LOCTEXT("ForcedBackOnMainThread", "Get Graph Parameter node is forced back on the main thread because it doesn't have UPCGUserParametersData as (implicit) input.\n"
+			"If the graph was scheduled manually through the Subsystem, make sure that you provide no PreTask element, or compute the UserParameterData in it."), Context);
+		Context->RequestToRunOnMainThread();
+		return false;
+	}
+	
 	const UPCGUserParameterGetSettings* Settings = Context->GetInputSettings<UPCGUserParameterGetSettings>();
 	const UPCGGenericUserParameterGetSettings* GenericSettings = Context->GetInputSettings<UPCGGenericUserParameterGetSettings>();
 	check(Settings || GenericSettings);
diff --git a/Engine/Plugins/PCG/Source/PCG/Private/Graph/PCGGraphExecutor.cpp b/Engine/Plugins/PCG/Source/PCG/Private/Graph/PCGGraphExecutor.cpp
--- a/Engine/Plugins/PCG/Source/PCG/Private/Graph/PCGGraphExecutor.cpp
+++ b/Engine/Plugins/PCG/Source/PCG/Private/Graph/PCGGraphExecutor.cpp
@@ -13,6 +13,7 @@
 #include "PCGSubsystem.h"
 #include "PCGWorldActor.h"
 #include "Compute/Elements/PCGComputeGraphElement.h"
+#include "Data/PCGUserParametersData.h"
 #include "Grid/PCGPartitionActor.h"
 #include "Helpers/PCGActorHelpers.h"
 #include "Helpers/PCGAsync.h"
@@ -305,10 +306,28 @@ FPCGTaskId FPCGGraphExecutor::ScheduleGraph(const FPCGScheduleGraphParams& InPar
 		check(ScheduledTask.Tasks.Num() >= 2 && ScheduledTask.Tasks[ScheduledTask.Tasks.Num() - 2].Node == nullptr);
 		FPCGGraphTask& PreGraphTask = ScheduledTask.Tasks[ScheduledTask.Tasks.Num() - 2];
 
-		if(InParams.PreGraphElement.IsValid())
+		if (InParams.PreGraphElement.IsValid())
 		{
 			PreGraphTask.Element = InParams.PreGraphElement;
 		}
+		else if (InParams.ExecutionSource->GetExecutionState().GetGraph() != InParams.Graph)
+		{
+			const UPCGGraphInterface* OverrideGraph = InParams.ExecutionSource->GetExecutionState().GetGraphInstance();
+			if (!OverrideGraph)
+			{
+				OverrideGraph = InParams.ExecutionSource->GetExecutionState().GetGraph();
+			}
+
+			if (OverrideGraph)
+			{
+				PreGraphTask.Element = MakeShared<FPCGPreGraphElement>(OverrideGraph);
+			}
+		}
+
+		if (!PreGraphTask.Element)
+		{
+			PreGraphTask.Element = GetPreGraphElement();
+		}
 
 		for (FPCGTaskId ExternalDependency : InParams.ExternalDependencies)
 		{
@@ -2708,6 +2727,8 @@ void FPCGGraphExecutor::AddReferencedObjects(FReferenceCollector& Collector)
 
 FPCGElementPtr FPCGGraphExecutor::GetFetchInputElement()
 {
+	PCG::TScopeLock Lock(DefaultElementsLock);
+	
 	if (!FetchInputElementPtr)
 	{
 		FetchInputElementPtr = MakeShared<FPCGFetchInputElement>();
@@ -2718,6 +2739,8 @@ FPCGElementPtr FPCGGraphExecutor::GetFetchInputElement()
 
 FPCGElementPtr FPCGGraphExecutor::GetPreGraphElement()
 {
+	PCG::TScopeLock Lock(DefaultElementsLock);
+	
 	if (!PreGraphElementPtr)
 	{
 		PreGraphElementPtr = MakeShared<FPCGPreGraphElement>();
@@ -2909,6 +2932,28 @@ bool FPCGPreGraphElement::ExecuteInternal(FPCGContext* Context) const
 
 	ExecutionSource->GetExecutionState().ExecutePreGraph(Context);
 
+	// Also emplace in the output a copy of the user parameters. Check in priority: Graph override, graph instance on the execution source and graph on the execution source.
+	const UPCGGraphInterface* Graph = GraphOverride.Get();
+
+	if (!Graph)
+	{
+		Graph = ExecutionSource->GetExecutionState().GetGraphInstance();
+	}
+	
+	if (!Graph)
+	{
+		Graph = ExecutionSource->GetExecutionState().GetGraph();
+	}
+	
+	if (const UPCGUserParametersData* UserParametersData = Graph ? Graph->ConstructUserParametersData(Context) : nullptr)
+	{
+		FPCGTaggedData& TaggedData = Context->OutputData.TaggedData.Emplace_GetRef();
+		TaggedData.Data = UserParametersData;
+		TaggedData.Tags.Add(PCG::Private::UserParameterTagData);
+		// Mark this data pinless, since it is internal data, not meant to be shown in the graph editor.
+		TaggedData.bPinlessData = true;
+	}
+
 	return true;
 }
 
@@ -3008,6 +3053,11 @@ bool FPCGFetchInputElement::ExecuteInternal(FPCGContext* Context) const
 	return true;
 }
 
+FPCGPreGraphElement::FPCGPreGraphElement(const UPCGGraphInterface* InGraphOverride)
+	: GraphOverride(InGraphOverride)
+{
+}
+
 FPCGGenericElement::FPCGGenericElement(TFunction<bool(FPCGContext*)> InOperation, const FContextAllocator& InContextAllocator)
 	: Operation(InOperation)
 	, ContextAllocator(InContextAllocator)
diff --git a/Engine/Plugins/PCG/Source/PCG/Private/Graph/PCGGraphExecutor.h b/Engine/Plugins/PCG/Source/PCG/Private/Graph/PCGGraphExecutor.h
--- a/Engine/Plugins/PCG/Source/PCG/Private/Graph/PCGGraphExecutor.h
+++ b/Engine/Plugins/PCG/Source/PCG/Private/Graph/PCGGraphExecutor.h
@@ -361,6 +361,9 @@ private:
 	/** Monotonically increasing id. Should be reset once all tasks are executed, should be protected by the ScheduleLock */
 	FPCGTaskId NextTaskId = 0;
 
+	/** To protect element creation. */
+	PCG::FLock DefaultElementsLock;
+
 	std::atomic<bool> bNeedToExecuteTasksEnded = false;
 	std::atomic<int32> NewReadyTaskCount = 0;
 
@@ -396,15 +399,27 @@ protected:
 	virtual bool SupportsBasePointDataInputs(FPCGContext* InContext) const override { return true; }
 };
 
+/**
+ * Potential first task executed in a graph that will do some verifications to make sure the graph can be executed.
+ * Will also execute some pre task logic provided by the ExecutionSource (ExecutePreGraph) and will freeze the user parameters
+ * of the graph for them to be accessed in multi-thread through the graph execution.
+ * Note that by default, the user parameters are pulled from the ExecutionSource graph, but in a case where we schedule a different graph
+ * than the one on the ExecutionSource, the PreGraph element can be setup with a graph override.
+ */
 class FPCGPreGraphElement : public IPCGElement
 {
 public:
+	FPCGPreGraphElement(const UPCGGraphInterface* InGraphOverride = nullptr);
+	
 	virtual bool IsCacheable(const UPCGSettings* InSettings) const override { return false; }
 	virtual bool CanExecuteOnlyOnMainThread(FPCGContext* Context) const override { return true; }
 protected:
 	virtual bool ExecuteInternal(FPCGContext* Context) const override;
 	virtual bool IsPassthrough(const UPCGSettings* InSettings) const override { return true; }
 	virtual bool SupportsBasePointDataInputs(FPCGContext* InContext) const override { return true; }
+
+private:
+	TWeakObjectPtr<const UPCGGraphInterface> GraphOverride;
 };
 
 class FPCGGenericElement : public IPCGElement
diff --git a/Engine/Plugins/PCG/Source/PCG/Private/PCGCommon.cpp b/Engine/Plugins/PCG/Source/PCG/Private/PCGCommon.cpp
--- a/Engine/Plugins/PCG/Source/PCG/Private/PCGCommon.cpp
+++ b/Engine/Plugins/PCG/Source/PCG/Private/PCGCommon.cpp
@@ -324,6 +324,17 @@ namespace PCGPointCustomPropertyNames
 	const FName ScaledLocalSizeName = TEXT("ScaledLocalSize");
 }
 
+namespace PCGAttributeNameConstants
+{
+	const FLazyName ActorReferenceAttribute = "ActorReference";
+	const FLazyName ComponentReferenceAttribute = "ComponentReference";
+}
+
+namespace PCG::Private
+{
+	const FString UserParameterTagData = TEXT("PCGUserParametersTagData");
+}
+
 namespace PCGPinConstants
 {
 	const FName DefaultInputLabel = TEXT("In");
diff --git a/Engine/Plugins/PCG/Source/PCG/Private/PCGContext.cpp b/Engine/Plugins/PCG/Source/PCG/Private/PCGContext.cpp
--- a/Engine/Plugins/PCG/Source/PCG/Private/PCGContext.cpp
+++ b/Engine/Plugins/PCG/Source/PCG/Private/PCGContext.cpp
@@ -278,13 +278,13 @@ void FPCGContext::InitializeSettings(bool bSkipPostLoad)
 								if (!ObjectPath.ResolveObject())
 								{
 									// We have an override value that is not loaded. We need to schedule the task on the main thread.
-									bOverrideSettingsOnMainThread = true;
+									RequestToRunOnMainThread();
 									break;
 								}
 							}
 						}
 
-						if (bOverrideSettingsOnMainThread)
+						if (CanExecuteOnlyOnMainThread())
 						{
 							break;
 						}
@@ -343,8 +343,8 @@ void FPCGContext::OverrideSettings()
 	TRACE_CPUPROFILER_EVENT_SCOPE(FPCGContext::OverrideSettings);
 
 	// If we have to be on the main thread check it there. It is necessary for overrides of hard references that are not yet loaded. We can reset the flag afterwards.
-	check(!bOverrideSettingsOnMainThread || IsInGameThread());
-	bOverrideSettingsOnMainThread = false;
+	check(!bMustRunOnMainThread || IsInGameThread());
+	bMustRunOnMainThread = false;
 
 	// Use original settings to avoid recomputing OverridableParams() everytime
 	const UPCGSettings* OriginalSettings = GetOriginalSettings<UPCGSettings>();
diff --git a/Engine/Plugins/PCG/Source/PCG/Private/PCGDefaultExecutionSource.cpp b/Engine/Plugins/PCG/Source/PCG/Private/PCGDefaultExecutionSource.cpp
--- a/Engine/Plugins/PCG/Source/PCG/Private/PCGDefaultExecutionSource.cpp
+++ b/Engine/Plugins/PCG/Source/PCG/Private/PCGDefaultExecutionSource.cpp
@@ -188,7 +188,7 @@ void UPCGDefaultExecutionSource::Generate()
 		FPCGScheduleGraphParams GraphParams{
 			/*InGraph=*/ GetGraph(),
 			/*InExecutionSource=*/ this,
-			/*InPreGraphElement=*/ MakeShared<FPCGInputForwardingElement>(FPCGDataCollection{}),
+			/*InPreGraphElement=*/ nullptr,
 			/*InInputElement=*/ MakeShared<FPCGInputForwardingElement>(FPCGDataCollection{}),
 			/*InExternalDependencies*/ {},
 			/*InFromStack*/ nullptr,
diff --git a/Engine/Plugins/PCG/Source/PCG/Private/PCGGraph.cpp b/Engine/Plugins/PCG/Source/PCG/Private/PCGGraph.cpp
--- a/Engine/Plugins/PCG/Source/PCG/Private/PCGGraph.cpp
+++ b/Engine/Plugins/PCG/Source/PCG/Private/PCGGraph.cpp
@@ -11,6 +11,7 @@
 #include "PCGPin.h"
 #include "PCGSubgraph.h"
 #include "PCGSubsystem.h"
+#include "Data/PCGUserParametersData.h"
 #include "Editor/IPCGEditorModule.h"
 #include "Elements/PCGHiGenGridSize.h"
 #include "Elements/PCGUserParameterGet.h"
@@ -323,6 +324,19 @@ bool UPCGGraphInterface::IsEquivalent(const UPCGGraphInterface* Other) const
 	return true;
 }
 
+UPCGUserParametersData* UPCGGraphInterface::ConstructUserParametersData(FPCGContext* InContext) const
+{
+	const FInstancedPropertyBag* UserParameters = GetUserParametersStruct();
+	if (!UserParameters)
+	{
+		return nullptr;
+	}
+
+	UPCGUserParametersData* UserParamData = FPCGContext::NewObject_AnyThread<UPCGUserParametersData>(InContext);
+	UserParamData->UserParameters = UserParameters->GetValue();
+	return UserParamData;
+}
+
 #if WITH_EDITOR
 bool UPCGGraphInterface::VerifyAndUpdateIfGraphParameterValueChanged(const FPropertyChangedChainEvent& PropertyChangedEvent)
 {
diff --git a/Engine/Plugins/PCG/Source/PCG/Private/PCGSubgraph.cpp b/Engine/Plugins/PCG/Source/PCG/Private/PCGSubgraph.cpp
--- a/Engine/Plugins/PCG/Source/PCG/Private/PCGSubgraph.cpp
+++ b/Engine/Plugins/PCG/Source/PCG/Private/PCGSubgraph.cpp
@@ -690,16 +690,11 @@ void FPCGSubgraphElement::PrepareSubgraphUserParameters(const UPCGSubgraphSettin
 		}
 		else if (const FInstancedPropertyBag* InstancedPropertyBag = SubgraphInterface->GetUserParametersStruct())
 		{
-			// FIXME: Copy is done there.
 			UserParamData->UserParameters = InstancedPropertyBag->GetValue();
 		}
-		else
-		{
-			// Do nothing, we still want to have a User Parameter Data to indicate we are in a subgraph context.
-		}
 
 		// Hook up user parameter data from upstream
-		TArray<FPCGTaggedData> UpstreamUserParameterData = Context->InputData.GetTaggedTypedInputs<UPCGUserParametersData>(PCGBaseSubgraphConstants::UserParameterTagData);
+		TArray<FPCGTaggedData> UpstreamUserParameterData = Context->InputData.GetTaggedTypedInputs<UPCGUserParametersData>(PCG::Private::UserParameterTagData);
 		if (!UpstreamUserParameterData.IsEmpty())
 		{
 #if WITH_EDITOR
@@ -715,7 +710,7 @@ void FPCGSubgraphElement::PrepareSubgraphUserParameters(const UPCGSubgraphSettin
 
 		FPCGTaggedData& TaggedData = OutputData.TaggedData.Emplace_GetRef();
 		TaggedData.Data = UserParamData;
-		TaggedData.Tags.Add(PCGBaseSubgraphConstants::UserParameterTagData);
+		TaggedData.Tags.Add(PCG::Private::UserParameterTagData);
 		// Mark this data pinless, since it is internal data, not meant to be shown in the graph editor.
 		TaggedData.bPinlessData = true;
 	}
diff --git a/Engine/Plugins/PCG/Source/PCG/Public/Elements/PCGUserParameterGet.h b/Engine/Plugins/PCG/Source/PCG/Public/Elements/PCGUserParameterGet.h
--- a/Engine/Plugins/PCG/Source/PCG/Public/Elements/PCGUserParameterGet.h
+++ b/Engine/Plugins/PCG/Source/PCG/Public/Elements/PCGUserParameterGet.h
@@ -135,8 +135,7 @@ protected:
 class FPCGUserParameterGetElement : public IPCGElement
 {
 public:
-	// Graph parameters access is not thread safe
-	virtual bool CanExecuteOnlyOnMainThread(FPCGContext* Context) const override { return true; }
+	virtual bool CanExecuteOnlyOnMainThread(FPCGContext* Context) const override;
 	virtual bool IsCacheable(const UPCGSettings* InSettings) const override { return false; }
 
 protected:
diff --git a/Engine/Plugins/PCG/Source/PCG/Public/PCGCommon.h b/Engine/Plugins/PCG/Source/PCG/Public/PCGCommon.h
--- a/Engine/Plugins/PCG/Source/PCG/Public/PCGCommon.h
+++ b/Engine/Plugins/PCG/Source/PCG/Public/PCGCommon.h
@@ -307,6 +307,17 @@ namespace PCGValueConstants
 	constexpr int32 DefaultSeed = 42;
 }
 
+namespace PCGAttributeNameConstants
+{
+	extern PCG_API const FLazyName ActorReferenceAttribute;
+	extern PCG_API const FLazyName ComponentReferenceAttribute;
+}
+
+namespace PCG::Private
+{
+	extern PCG_API const FString UserParameterTagData;
+}
+
 namespace PCGPinConstants
 {
 	extern PCG_API const FName DefaultInputLabel;
@@ -715,7 +726,7 @@ struct FPCGScheduleGraphParams
 	// PCG execution source associated with this task. Can be null.
 	IPCGGraphExecutionSource* ExecutionSource = nullptr;
 
-	// First task to run.
+	// First task to run. If null, will fallback on the default FPCGPreGraphElement.
 	FPCGElementPtr PreGraphElement;
 
 	// Task to run as an input to the provided graph.
diff --git a/Engine/Plugins/PCG/Source/PCG/Public/PCGContext.h b/Engine/Plugins/PCG/Source/PCG/Public/PCGContext.h
--- a/Engine/Plugins/PCG/Source/PCG/Public/PCGContext.h
+++ b/Engine/Plugins/PCG/Source/PCG/Public/PCGContext.h
@@ -142,7 +142,8 @@ PRAGMA_ENABLE_DEPRECATION_WARNINGS
 	UE_DEPRECATED(5.6, "Access the stack via GetStack()")
 	const FPCGStack* Stack = nullptr;
 
-	bool CanExecuteOnlyOnMainThread() const { return bOverrideSettingsOnMainThread; } 
+	bool CanExecuteOnlyOnMainThread() const { return bMustRunOnMainThread; }
+	void RequestToRunOnMainThread() { bMustRunOnMainThread = true; }
 
 	UE_API const UPCGSettingsInterface* GetInputSettingsInterface() const;
 	
@@ -295,7 +296,11 @@ PRAGMA_ENABLE_DEPRECATION_WARNINGS
 	UE_API bool IsCacheEnabled();
 	UE_API void StoreInCache(const FPCGStoreInCacheParams& Params, const FPCGDataCollection& InOutput);
 	UE_API bool GetFromCache(const FPCGGetFromCacheParams& Params, FPCGDataCollection& OutOutput) const;
-	
+
+	/**
+	 * Get the context setup for this thread, can be null. It is set within an element execution, but it is the responsibility of the
+	 * implementer to set this context (with FPCGContextScope) when starting new async tasks.
+	 */
 	UE_API static FPCGContext* GetContextForThread();
 
 protected:
@@ -333,8 +338,8 @@ private:
 	// List of params that were in effect overriden
 	TArray<const FPCGSettingsOverridableParam*> OverriddenParams;
 
-	// If the settings need to be overridden on the main thread (because we have to load objects)
-	bool bOverrideSettingsOnMainThread = false;
+	// If the element must run on the main thread (like if the settings need to be overridden on the main thread, because we have to load objects))
+	bool bMustRunOnMainThread = false;
 
 	// List of objects created by the PCG Elements, we need to track them so we can remove their Async flags when storing results on main thread
 	// so that they can be considered as existing on the main thread (and get properly GCed)
diff --git a/Engine/Plugins/PCG/Source/PCG/Public/PCGGraph.h b/Engine/Plugins/PCG/Source/PCG/Public/PCGGraph.h
--- a/Engine/Plugins/PCG/Source/PCG/Public/PCGGraph.h
+++ b/Engine/Plugins/PCG/Source/PCG/Public/PCGGraph.h
@@ -26,6 +26,7 @@ struct FAssetData;
 struct FPCGCompilerDiagnostics;
 class UPCGGraphInterface;
 class UPCGGraphCompilationData;
+class UPCGUserParametersData;
 #if WITH_EDITOR
 class UPCGEditorGraph;
 struct FEdGraphPinType;
@@ -131,6 +132,9 @@ public:
 	/** A graph interface is equivalent to another graph interface if they are the same (same ptr), or if they have the same graph. Will be overridden when graph instance supports overrides. */
 	UE_API virtual bool IsEquivalent(const UPCGGraphInterface* Other) const;
 
+	/** Construct a data that will contain a copy of the UserParameters, in order to freeze them while executing a graph. */
+	UE_API UPCGUserParametersData* ConstructUserParametersData(FPCGContext* InContext) const;
+
 #if WITH_EDITOR
 	FOnPCGGraphChanged OnGraphChangedDelegate;
 	FOnPCGGraphParametersChanged OnGraphParametersChangedDelegate;
diff --git a/Engine/Plugins/PCG/Source/PCG/Public/PCGSubgraph.h b/Engine/Plugins/PCG/Source/PCG/Public/PCGSubgraph.h
--- a/Engine/Plugins/PCG/Source/PCG/Public/PCGSubgraph.h
+++ b/Engine/Plugins/PCG/Source/PCG/Public/PCGSubgraph.h
@@ -2,6 +2,7 @@
 
 #pragma once
 
+#include "PCGCommon.h"
 #include "PCGContext.h"
 #include "PCGGraph.h"
 #include "PCGSettings.h"
@@ -14,7 +15,8 @@
 
 namespace PCGBaseSubgraphConstants
 {
-	static const FString UserParameterTagData = TEXT("PCGUserParametersTagData");
+	UE_DEPRECATED(5.8, "Use PCG::Private::UserParameterTagData")
+	inline const FString& UserParameterTagData = PCG::Private::UserParameterTagData;
 }
 
 UCLASS(MinimalAPI, Abstract)
