[PCG] Change equality comparison with bitwise flag comparison

#rnx
#Synced-CL 48766296

diff --git a/Engine/Plugins/PCG/Source/PCG/Private/PCGComponent.cpp b/Engine/Plugins/PCG/Source/PCG/Private/PCGComponent.cpp
--- a/Engine/Plugins/PCG/Source/PCG/Private/PCGComponent.cpp
+++ b/Engine/Plugins/PCG/Source/PCG/Private/PCGComponent.cpp
@@ -2106,9 +2106,8 @@ void UPCGComponent::RefreshAfterGraphChanged(UPCGGraphInterface* InGraph, EPCGCh
 		return;
 	}
 
-	if (ChangeType == EPCGChangeType::Cosmetic ||
-		ChangeType == EPCGChangeType::GraphCustomization ||
-		ChangeType == EPCGChangeType::None)
+	const bool bIsTrivial = (ChangeType & ~(EPCGChangeType::Cosmetic | EPCGChangeType::GraphCustomization)) == EPCGChangeType::None;
+	if (bIsTrivial)
 	{
 		// If it is a cosmetic change (or no change), nothing to do
 		return;
