[PCG] Fix projection node typing its output pin to the first input pin, giving an incorrect concrete type and emitting a warning because the pin is supposed to be Spatial.

[RN] [PCG] Fix warning about data types on projection node when projecting a spline onto a surface.
#Synced-CL 48757882

diff --git a/Engine/Plugins/PCG/Source/PCG/Private/Elements/PCGProjectionElement.cpp b/Engine/Plugins/PCG/Source/PCG/Private/Elements/PCGProjectionElement.cpp
--- a/Engine/Plugins/PCG/Source/PCG/Private/Elements/PCGProjectionElement.cpp
+++ b/Engine/Plugins/PCG/Source/PCG/Private/Elements/PCGProjectionElement.cpp
@@ -46,7 +46,12 @@ FPCGDataTypeIdentifier UPCGProjectionSettings::GetCurrentPinTypesID(const UPCGPi
 	{
 		if (bForceCollapseToPoint || GetTypeUnionIDOfIncidentEdges(PCGPinConstants::DefaultInputLabel) == EPCGDataType::Point)
 		{
-			return FPCGDataTypeIdentifier{EPCGDataType::Point};
+			return FPCGDataTypeIdentifier{ EPCGDataType::Point };
+		}
+		else
+		{
+			// Return type here, because falling back to the super call will return the incident of the types on the primary input.
+			return FPCGDataTypeIdentifier{ EPCGDataType::Spatial };
 		}
 	}
 
diff --git a/Engine/Plugins/PCG/Source/PCG/Public/PCGSettings.h b/Engine/Plugins/PCG/Source/PCG/Public/PCGSettings.h
--- a/Engine/Plugins/PCG/Source/PCG/Public/PCGSettings.h
+++ b/Engine/Plugins/PCG/Source/PCG/Public/PCGSettings.h
@@ -263,7 +263,8 @@ public:
 	/** For symmetry reason, do the same with output pins. For now forward just the call to OutputPinProperties */
 	UE_API TArray<FPCGPinProperties> AllOutputPinProperties() const;
 
-	/** If the node has any dynamic pins that can change based on input or settings */
+	/** If the node has any dynamic pins that can change type based on input types or settings. If true the default behavior 
+	* for the output pin types will be to take the union of types incident on primary input pin. */
 	virtual bool HasDynamicPins() const { return false; }
 
 	/** Whether this node can deactivate its output pins during execution, which can dynamically cull downstream nodes. */
