[PCG] Fixed some performance issues in the metadata ops when there is a large number of inputs.

#rnx
#Synced-CL 48730134

diff --git a/Engine/Plugins/PCG/Source/PCG/Private/Elements/Metadata/PCGMetadataOpElementBase.cpp b/Engine/Plugins/PCG/Source/PCG/Private/Elements/Metadata/PCGMetadataOpElementBase.cpp
--- a/Engine/Plugins/PCG/Source/PCG/Private/Elements/Metadata/PCGMetadataOpElementBase.cpp
+++ b/Engine/Plugins/PCG/Source/PCG/Private/Elements/Metadata/PCGMetadataOpElementBase.cpp
@@ -518,12 +518,26 @@
 void FPCGMetadataElementBase::PassthroughInput(FPCGContext* Context, TArray<FPCGTaggedData>& Outputs, const int32 Index) const
 {
 	check(Context);
+	FPCGMetadataElementBase::ContextType* TimeSlicedContext = static_cast<FPCGMetadataElementBase::ContextType*>(Context);
+
 	const UPCGMetadataSettingsBase* Settings = Context->GetInputSettings<UPCGMetadataSettingsBase>();
 	check(Settings);
 
 	const uint32 NumberOfOutputs = Settings->GetResultNum();
 	const uint32 PrimaryPinIndex = Settings->GetInputPinToForward();
-	TArray<FPCGTaggedData> InputsToForward = Context->InputData.GetInputsByPin(Settings->GetInputPinLabel(PrimaryPinIndex));
 
+	// Cache the primary pin inputs if we haven't already done it.
+	if (TimeSlicedContext->GetPerExecutionState().OperandPinData.IsEmpty())
+	{
+		TimeSlicedContext->GetPerExecutionState().OperandPinData.SetNum(Settings->GetOperandNum());
+	}
+
+	if (TimeSlicedContext->GetPerExecutionState().OperandPinData[PrimaryPinIndex].IsEmpty())
+	{
+		TimeSlicedContext->GetPerExecutionState().OperandPinData[PrimaryPinIndex] = Context->InputData.GetInputsByPin(Settings->GetInputPinLabel(PrimaryPinIndex));
+	}
+
+	const TArray<FPCGTaggedData>& InputsToForward = TimeSlicedContext->GetPerExecutionState().OperandPinData[PrimaryPinIndex];
+
 	if (InputsToForward.IsEmpty())
 	{
@@ -645,58 +659,33 @@
 
 	int32 OperandInputNumMax = 0;
 
-	// There's no execution state, so just flag that it is ready to continue
-	TimeSlicedContext->InitializePerExecutionState([this, Settings, &OperandInputNumMax, OperandNum](ContextType* Context, PCGTimeSlice::FEmptyStruct& OutState) -> EPCGTimeSliceInitResult
+	// Initiialize execution state, which will setup default data as need and perform early validation.
+	EPCGTimeSliceInitResult ExecStateInitResult = TimeSlicedContext->InitializePerExecutionState([this, Settings, &OperandInputNumMax, OperandNum](ContextType* Context, ExecStateType& OutState) -> EPCGTimeSliceInitResult
 	{
-		for (uint32 i = 0; i < OperandNum; ++i)
-		{
-			const int32 CurrentInputNum = Context->InputData.GetInputsByPin(Settings->GetInputPinLabel(i)).Num();
-
-			OperandInputNumMax = FMath::Max(OperandInputNumMax, FMath::Max(1, CurrentInputNum));
-
-			// For the current input, no input (0) could be default value and we support N:1 and 1:N
-			if (CurrentInputNum > 1 && CurrentInputNum != OperandInputNumMax)
-			{
-				PCGE_LOG(Error, GraphAndLog, LOCTEXT("MismatchedOperandDataCount", "Number of data elements provided on inputs must be 1:N, N:1, or N:N."));
-				return EPCGTimeSliceInitResult::AbortExecution;
-			}
-		}
-
-		return EPCGTimeSliceInitResult::Success;
-	});
-
-	// Set up the iterations on the multiple inputs of the primary pin
-	TimeSlicedContext->InitializePerIterationStates(OperandInputNumMax, [this, Context, OperandNum, Settings, OperandInputNumMax](IterStateType& OutState, const ExecStateType& ExecState, const uint32 IterationIndex)
-	{
+		TRACE_CPUPROFILER_EVENT_SCOPE(FPCGMetadataElementBase::PrepareDataInternal::InitializePerExecutionState);
 		FPCGMetadataElementBase::ContextType* TimeSlicedContext = static_cast<FPCGMetadataElementBase::ContextType*>(Context);
-		TArray<FPCGTaggedData>& Outputs = Context->OutputData.TaggedData;
-		const uint32 NumberOfResults = Settings->GetResultNum();
-
-		OutState.Context = Context;
-
-		// Gathering all the inputs metadata
-		TArray<const UPCGMetadata*> SourceMetadata;
-		TArray<const FPCGMetadataAttributeBase*> SourceAttribute;
-		TArray<FPCGTaggedData> InputTaggedData;
-		SourceMetadata.SetNum(OperandNum);
-		SourceAttribute.SetNum(OperandNum);
-		InputTaggedData.SetNum(OperandNum);
-
-		// Since we add the output data (in CreateAttribute below) if the operation is valid in the PrepareData, if we ever have a no-op, we have to passthrough the inputs now and not in the Execute. So that the order is respected
-		// in the end. (ie. { Input1(Valid), Input2(No-Op), Input3(Valid) } will have in output { Output1, Input2, Output3 } and not { Output1, Output3, Input2 } if we do the passthrough in the Execute.)
-		auto NoOperation = [this, Context, IterationIndex, &Outputs]() { PassthroughInput(Context, Outputs, IterationIndex); return EPCGTimeSliceInitResult::NoOperation; };
-
-		const uint32 PrimaryPinIndex = Settings->GetInputPinToForward();
 		OutState.DefaultValueOverriddenPins.AddZeroed(OperandNum);
+		OutState.OperandPinData.SetNum(OperandNum);
 
-		// Iterate over the inputs and validate
 		for (uint32 OperandPinIndex = 0; OperandPinIndex < OperandNum; ++OperandPinIndex)
 		{
 			const FName CurrentPinLabel = Settings->GetInputPinLabel(OperandPinIndex);
 			const bool bIsInputConnected = Context->Node ? Context->Node->IsInputPinConnected(CurrentPinLabel) : false;
-			TArray<FPCGTaggedData> CurrentPinInputData = Context->InputData.GetInputsByPin(CurrentPinLabel);
+			TArray<FPCGTaggedData>& CurrentPinInputData = OutState.OperandPinData[OperandPinIndex];
+			// Override data regardless.
+			CurrentPinInputData = Context->InputData.GetInputsByPin(CurrentPinLabel);
 
-			// This only needs to be checked once
+			const int32 CurrentInputNum = CurrentPinInputData.Num();
+
+			OperandInputNumMax = FMath::Max(OperandInputNumMax, FMath::Max(1, CurrentInputNum));
+
+			// For the current input, no input (0) could be default value and we support N:1 and 1:N
+			if (CurrentInputNum > 1 && CurrentInputNum != OperandInputNumMax) 
+			{
+				PCGE_LOG(Error, GraphAndLog, LOCTEXT("MismatchedOperandDataCount", "Number of data elements provided on inputs must be 1:N, N:1, or N:N."));
+				return EPCGTimeSliceInitResult::AbortExecution;
+			}
+
 			if (Settings->DefaultValuesAreEnabled() && !bIsInputConnected && CurrentPinInputData.IsEmpty() && Settings->IsPinDefaultValueActivated(CurrentPinLabel))
 			{
 				FPCGTaggedData& DefaultData = CurrentPinInputData.Emplace_GetRef();
@@ -732,7 +721,7 @@
 			{
 				// If we have no data, there is no operation
 				PCGE_LOG(Verbose, LogOnly, FText::Format(LOCTEXT("MissingInputDataForPin", "No data provided on pin '{0}'."), FText::FromName(CurrentPinLabel)));
-				return NoOperation();
+				return EPCGTimeSliceInitResult::NoOperation;
 			}
 			else if (CurrentPinInputData.Num() != 1 && CurrentPinInputData.Num() != OperandInputNumMax)
 			{
@@ -743,15 +732,66 @@
 						OperandInputNumMax));
 				return EPCGTimeSliceInitResult::AbortExecution;
 			}
+		}
 
+		return EPCGTimeSliceInitResult::Success;
+	});
+
+	if (ExecStateInitResult == EPCGTimeSliceInitResult::AbortExecution)
+	{
+		return true;
+	}
+	else if (ExecStateInitResult == EPCGTimeSliceInitResult::NoOperation)
+	{
+		// Passthrough all inputs
+		// @todo_pcg - this could be optimized somewhat by doing a version that takes the OperandInputNumMax here.
+		TArray<FPCGTaggedData>& Outputs = Context->OutputData.TaggedData;
+		for (int IterationIndex = 0; IterationIndex < OperandInputNumMax; ++IterationIndex)
+		{
+			PassthroughInput(Context, Outputs, IterationIndex);
+		}
+
+		return true;
+	}
+
+	// Set up the iterations on the multiple inputs of the primary pin
+	TimeSlicedContext->InitializePerIterationStates(OperandInputNumMax, [this, Context, OperandNum, Settings, OperandInputNumMax](IterStateType& OutState, const ExecStateType& ExecState, const uint32 IterationIndex)
+	{
+		TRACE_CPUPROFILER_EVENT_SCOPE(FPCGMetadataElementBase::PrepareDataInternal::InitializePerIterationStates);
+		FPCGMetadataElementBase::ContextType* TimeSlicedContext = static_cast<FPCGMetadataElementBase::ContextType*>(Context);
+		TArray<FPCGTaggedData>& Outputs = Context->OutputData.TaggedData;
+		const uint32 NumberOfResults = Settings->GetResultNum();
+
+		OutState.Context = Context;
+
+		// Gathering all the inputs metadata
+		TArray<const UPCGMetadata*> SourceMetadata;
+		TArray<const FPCGMetadataAttributeBase*> SourceAttribute;
+		TArray<FPCGTaggedData> InputTaggedData;
+		SourceMetadata.SetNum(OperandNum);
+		SourceAttribute.SetNum(OperandNum);
+		InputTaggedData.SetNum(OperandNum);
+
+		// Since we add the output data (in CreateAttribute below) if the operation is valid in the PrepareData, if we ever have a no-op, we have to passthrough the inputs now and not in the Execute. So that the order is respected
+		// in the end. (ie. { Input1(Valid), Input2(No-Op), Input3(Valid) } will have in output { Output1, Input2, Output3 } and not { Output1, Output3, Input2 } if we do the passthrough in the Execute.)
+		auto NoOperation = [this, Context, IterationIndex, &Outputs]() { PassthroughInput(Context, Outputs, IterationIndex); return EPCGTimeSliceInitResult::NoOperation; };
+
+		const uint32 PrimaryPinIndex = Settings->GetInputPinToForward();
+
+		// Iterate over the inputs and validate
+		for (uint32 OperandPinIndex = 0; OperandPinIndex < OperandNum; ++OperandPinIndex)
+		{
+			const TArray<FPCGTaggedData>& CurrentPinInputData = ExecState.OperandPinData[OperandPinIndex];
+
 			// The operand inputs must either be N:1 or N:N or 1:N
-			InputTaggedData[OperandPinIndex] = CurrentPinInputData.Num() == 1 ? CurrentPinInputData[0] : MoveTemp(CurrentPinInputData[IterationIndex]);
+			InputTaggedData[OperandPinIndex] = CurrentPinInputData.Num() == 1 ? CurrentPinInputData[0] : CurrentPinInputData[IterationIndex];
 
 			// Check if we have any points
 			if (const UPCGBasePointData* PointInput = Cast<const UPCGBasePointData>(InputTaggedData[OperandPinIndex].Data))
 			{
 				if (PointInput->GetNumPoints() == 0)
 				{
+					const FName CurrentPinLabel = Settings->GetInputPinLabel(OperandPinIndex);
 					// If we have no points, there is no operation
 					PCGE_LOG(Verbose, LogOnly, FText::Format(LOCTEXT("NoPointsForPin", "No points in point data provided on pin {0}"), FText::FromName(CurrentPinLabel)));
 					return NoOperation();
@@ -761,6 +801,7 @@
 			SourceMetadata[OperandPinIndex] = InputTaggedData[OperandPinIndex].Data->ConstMetadata();
 			if (!SourceMetadata[OperandPinIndex])
 			{
+				const FName CurrentPinLabel = Settings->GetInputPinLabel(OperandPinIndex);
 				// Since this aborts execution, and the user can fix it, it should be a node error
 				PCGE_LOG(Error, GraphAndLog, FText::Format(LOCTEXT("InvalidInputDataTypeForPin", "Invalid data provided on pin '{0}', must be of type Spatial or Attribute Set."), FText::FromName(CurrentPinLabel)));
 				return EPCGTimeSliceInitResult::AbortExecution;
@@ -777,7 +818,7 @@
 		const FPCGTaggedData& PrimaryPinData = InputTaggedData[PrimaryPinIndex];
 
 		// First create an accessor for the input to forward (it's our control data)
-		const FPCGAttributePropertyInputSelector PrimarySelector = !OutState.DefaultValueOverriddenPins[PrimaryPinIndex] ? Settings->GetInputSource(PrimaryPinIndex) : FPCGAttributePropertyInputSelector{};
+		const FPCGAttributePropertyInputSelector PrimarySelector = !ExecState.DefaultValueOverriddenPins[PrimaryPinIndex] ? Settings->GetInputSource(PrimaryPinIndex) : FPCGAttributePropertyInputSelector{};
 
 		PCGMetadataOpPrivate::CreateAccessor(PrimarySelector, PrimaryPinData, OperationData, PrimaryPinIndex);
 		if (!PCGMetadataOpPrivate::ValidateAccessor(Context, Settings, PrimaryPinData, OperationData, PrimaryPinIndex))
@@ -805,7 +846,7 @@
 					return EPCGTimeSliceInitResult::AbortExecution;
 				}
 
-				const FPCGAttributePropertyInputSelector Selector = !OutState.DefaultValueOverriddenPins[Index] ? Settings->GetInputSource(Index) : FPCGAttributePropertyInputSelector{};
+				const FPCGAttributePropertyInputSelector Selector = !ExecState.DefaultValueOverriddenPins[Index] ? Settings->GetInputSource(Index) : FPCGAttributePropertyInputSelector{};
 
 				PCGMetadataOpPrivate::CreateAccessor(Selector, InputTaggedData[Index], OperationData, Index);
 				if (!PCGMetadataOpPrivate::ValidateAccessor(Context, Settings, InputTaggedData[Index], OperationData, Index))
diff --git a/Engine/Plugins/PCG/Source/PCG/Public/Elements/Metadata/PCGMetadataOpElementBase.h b/Engine/Plugins/PCG/Source/PCG/Public/Elements/Metadata/PCGMetadataOpElementBase.h
--- a/Engine/Plugins/PCG/Source/PCG/Public/Elements/Metadata/PCGMetadataOpElementBase.h
+++ b/Engine/Plugins/PCG/Source/PCG/Public/Elements/Metadata/PCGMetadataOpElementBase.h
@@ -228,6 +228,12 @@ private:
 
 namespace PCGMetadataOps
 {
+	struct FInputStateData
+	{
+		TArray<TArray<FPCGTaggedData>> OperandPinData;
+		TArray<bool, TFixedAllocator<UPCGMetadataSettingsBase::MaxNumberOfInputs>> DefaultValueOverriddenPins;
+	};
+
 	struct FOperationData
 	{
 		int32 NumberOfElementsToProcess = -1;
@@ -245,14 +251,12 @@ namespace PCGMetadataOps
 		TArray<TUniquePtr<const IPCGAttributeAccessor>> InputAccessors;
 		TArray<TUniquePtr<IPCGAttributeAccessor>> OutputAccessors;
 
-		TArray<bool, TFixedAllocator<UPCGMetadataSettingsBase::MaxNumberOfInputs>> DefaultValueOverriddenPins;
-
 		template <int32 NbInputs, int32 NbOutputs>
 		void Validate();
 	};
 }
 
-class FPCGMetadataElementBase : public TPCGTimeSlicedElementBase<PCGTimeSlice::FEmptyStruct, PCGMetadataOps::FOperationData>
+class FPCGMetadataElementBase : public TPCGTimeSlicedElementBase<PCGMetadataOps::FInputStateData, PCGMetadataOps::FOperationData>
 {
 protected:
 	virtual bool PrepareDataInternal(FPCGContext* Context) const override;
@@ -285,7 +289,7 @@ protected:
 	template <typename InType1, typename InType2, typename InType3, typename InType4, typename... Callbacks>
 	bool DoQuaternaryOp(PCGMetadataOps::FOperationData& InOperationData, Callbacks&& ...InCallbacks) const;
 
-	/** To be called if we have no data to perform any operation, it will passthrough the input. */
+	/** To be called if we have no data to perform any operation, it will passthrough the input. The per execution state should be initialized before calling this, or performance issues might ensue if there is a large amount of inputs. */
 	void PassthroughInput(FPCGContext* Context, TArray<FPCGTaggedData>& Outputs, const int32 Index) const;
 };
 
