[PCG] Added csv support in replace data by tags

#rnx
#Synced-CL 48730134

diff --git a/Engine/Plugins/PCG/Source/PCG/Private/Elements/Data/PCGReplaceDataByTags.cpp b/Engine/Plugins/PCG/Source/PCG/Private/Elements/Data/PCGReplaceDataByTags.cpp
--- a/Engine/Plugins/PCG/Source/PCG/Private/Elements/Data/PCGReplaceDataByTags.cpp
+++ b/Engine/Plugins/PCG/Source/PCG/Private/Elements/Data/PCGReplaceDataByTags.cpp
@@ -3,6 +3,7 @@
 #include "Elements/Data/PCGReplaceDataByTags.h"
 
 #include "PCGContext.h"
+#include "Helpers/PCGHelpers.h"
 
 #include UE_INLINE_GENERATED_CPP_BY_NAME(PCGReplaceDataByTags)
 
@@ -74,6 +75,12 @@ bool FPCGReplaceDataByTagElement::ExecuteInternal(FPCGContext* Context) const
 		return FString();
 	};
 
+	TArray<FString> TagsFromSettings;
+	for (const FString& Tag : Settings->Tags)
+	{
+		TagsFromSettings.Append(PCGHelpers::GetStringArrayFromCommaSeparatedList(Tag));
+	}
+
 	// For each input, if there is a data in the replacement set that has the same tags, then use this one instead.
 	for (const FPCGTaggedData& Input : Inputs)
 	{
@@ -85,7 +92,7 @@ bool FPCGReplaceDataByTagElement::ExecuteInternal(FPCGContext* Context) const
 			bool bReplaceByThisData = false;
 
 			// If the settings don't specify what tags to match against, then do a full inclusion test (e.g. all tags from input are in replacement data)
-			if (Settings->Tags.IsEmpty())
+			if (TagsFromSettings.IsEmpty())
 			{
 				bReplaceByThisData = !Input.Tags.IsEmpty() && Replacement.Tags.Includes(Input.Tags);
 			}
@@ -94,7 +101,7 @@ bool FPCGReplaceDataByTagElement::ExecuteInternal(FPCGContext* Context) const
 			{
 				bReplaceByThisData = true;
 
-				for (const FString& Tag : Settings->Tags)
+				for (const FString& Tag : TagsFromSettings)
 				{
 					FString TagFromInput = FindTag(Input.Tags, Tag);
 					FString TagFromReplacement = FindTag(Replacement.Tags, Tag);
