[PCG] Add profiler scope to selection update (was showing up as perf regression but now we know its editor only)

#rnx
#Synced-CL 48785580

diff --git a/Engine/Plugins/PCG/Source/PCGEditor/Private/Widgets/SPCGExecutionSourceActionWidget.cpp b/Engine/Plugins/PCG/Source/PCGEditor/Private/Widgets/SPCGExecutionSourceActionWidget.cpp
--- a/Engine/Plugins/PCG/Source/PCGEditor/Private/Widgets/SPCGExecutionSourceActionWidget.cpp
+++ b/Engine/Plugins/PCG/Source/PCGEditor/Private/Widgets/SPCGExecutionSourceActionWidget.cpp
@@ -32,6 +32,7 @@ namespace PCGExecutionSourceActionWidget
 			FPCGScheduleGenericParams Params(
 				/*InOperation=*/ [](FPCGContext*) -> bool
 				{
+					TRACE_CPUPROFILER_EVENT_SCOPE(PCGExecutionSourceActionWidget::TriggerSelectionUpdateTask);
 					if (GEditor && GEditor->GetEditorSubsystem<ULevelEditorSubsystem>())
 					{
 						// Since we might have been moving things around (component wise) we need to force a selection set change, 
