[PCG] Support isolated actor data in GetActorData node

#rnx

#Synced-CL 48762154

diff --git a/Engine/Plugins/PCG/Source/PCG/Private/Elements/PCGActorSelector.cpp b/Engine/Plugins/PCG/Source/PCG/Private/Elements/PCGActorSelector.cpp
--- a/Engine/Plugins/PCG/Source/PCG/Private/Elements/PCGActorSelector.cpp
+++ b/Engine/Plugins/PCG/Source/PCG/Private/Elements/PCGActorSelector.cpp
@@ -150,6 +150,48 @@ namespace PCGActorSelector
 			return FoundActors;
 		}
 
+		// And after filtering, check for adding isolated actors
+		auto UpdateFoundActorsWithIsolatedIfNeeded = [IncludeIsolatedActors = Settings.IncludeIsolatedActors](TArray<AActor*>& InOutFoundActors)
+		{
+			if (IncludeIsolatedActors != EPCGIsolatedActorIncludeFlag::Ignore)
+			{
+				int ActorIndex = 0;
+				while (ActorIndex < InOutFoundActors.Num())
+				{
+					TArray<APCGIsolatedActor*> IsolatedActors = UPCGActorHelpers::GetIsolatedActorsOriginatingFrom(InOutFoundActors[ActorIndex]);
+
+					if (IncludeIsolatedActors == EPCGIsolatedActorIncludeFlag::Include)
+					{
+						++ActorIndex;
+
+						if (IsolatedActors.Num() > 0)
+						{
+							AActor** IsolatedActorsPtr = reinterpret_cast<AActor**>(IsolatedActors.GetData());
+							InOutFoundActors.Insert(IsolatedActorsPtr, IsolatedActors.Num(), ActorIndex);
+							ActorIndex += IsolatedActors.Num();
+						}
+					}
+					else if (IncludeIsolatedActors == EPCGIsolatedActorIncludeFlag::Only)
+					{
+						if (IsolatedActors.IsEmpty())
+						{
+							InOutFoundActors.RemoveAt(ActorIndex);
+						}
+						else
+						{
+							InOutFoundActors[ActorIndex++] = IsolatedActors[0];
+							if (IsolatedActors.Num() > 1)
+							{
+								AActor** IsolatedActorsPtr = reinterpret_cast<AActor**>(IsolatedActors.GetData());
+								InOutFoundActors.Insert(IsolatedActorsPtr + 1, IsolatedActors.Num() - 1, ActorIndex);
+								ActorIndex += IsolatedActors.Num() - 1;
+							}
+						}
+					}
+				}
+			}				
+		};
+		
 		// In case of iterating over all actors in the world, call our filtering function and get out.
 		if (Settings.ActorFilter == EPCGActorFilter::AllWorldActors)
 		{
@@ -157,7 +199,8 @@ namespace PCGActorSelector
 			const TSubclassOf<AActor> ActorSelectionClass = Settings.ActorSelection == EPCGActorSelection::ByClass && Settings.ActorSelectionClass ? Settings.ActorSelectionClass.Get() : AActor::StaticClass();
 			// A potential optimization if we know the sought actors are collide-able could be to obtain overlaps via a collision query.
 			UPCGActorHelpers::ForEachActorInWorld(World, ActorSelectionClass, FilteringFunction);
-
+			UpdateFoundActorsWithIsolatedIfNeeded(FoundActors);
+			
 			// FoundActor is set by the FilteringFunction (captured)
 			return FoundActors;
 		}
@@ -240,44 +283,6 @@ namespace PCGActorSelector
 			}
 		}
 
-		if (Settings.IncludeIsolatedActors != EPCGIsolatedActorIncludeFlag::Ignore)
-		{
-			int ActorIndex = 0;
-			while (ActorIndex < ActorsToCheck.Num())
-			{
-				TArray<APCGIsolatedActor*> IsolatedActors = UPCGActorHelpers::GetIsolatedActorsOriginatingFrom(ActorsToCheck[ActorIndex]);
-
-				if (Settings.IncludeIsolatedActors == EPCGIsolatedActorIncludeFlag::Include)
-				{
-					++ActorIndex;
-
-					if (IsolatedActors.Num() > 0)
-					{
-						AActor** IsolatedActorsPtr = reinterpret_cast<AActor**>(IsolatedActors.GetData());
-						ActorsToCheck.Insert(IsolatedActorsPtr, IsolatedActors.Num(), ActorIndex);
-						ActorIndex += IsolatedActors.Num();
-					}
-				}
-				else if (Settings.IncludeIsolatedActors == EPCGIsolatedActorIncludeFlag::Only)
-				{
-					if (IsolatedActors.IsEmpty())
-					{
-						ActorsToCheck.RemoveAt(ActorIndex);
-					}
-					else
-					{
-						ActorsToCheck[ActorIndex++] = IsolatedActors[0];
-						if (IsolatedActors.Num() > 1)
-						{
-							AActor** IsolatedActorsPtr = reinterpret_cast<AActor**>(IsolatedActors.GetData());
-							ActorsToCheck.Insert(IsolatedActorsPtr + 1, IsolatedActors.Num() - 1, ActorIndex);
-							ActorIndex += IsolatedActors.Num() - 1;
-						}
-					}
-				}
-			}
-		}
-
 		for (AActor* Actor : ActorsToCheck)
 		{
 			// FoundActor is set by the FilteringFunction (captured)
@@ -287,6 +292,8 @@ namespace PCGActorSelector
 			}
 		}
 
+		UpdateFoundActorsWithIsolatedIfNeeded(FoundActors);
+
 		return FoundActors;
 	}
 
diff --git a/Engine/Plugins/PCG/Source/PCG/Private/Elements/PCGDataFromActor.cpp b/Engine/Plugins/PCG/Source/PCG/Private/Elements/PCGDataFromActor.cpp
--- a/Engine/Plugins/PCG/Source/PCG/Private/Elements/PCGDataFromActor.cpp
+++ b/Engine/Plugins/PCG/Source/PCG/Private/Elements/PCGDataFromActor.cpp
@@ -945,6 +945,10 @@ void FPCGDataFromActorElement::ProcessActor(FPCGContext* InContext, const UPCGDa
 		{
 			bHasGeneratedPCGData |= !Component->GetGeneratedGraphOutput().TaggedData.IsEmpty();
 		}
+		if (APCGIsolatedActor* IsolatedActor = Cast<APCGIsolatedActor>(FoundActor))
+		{
+			bHasGeneratedPCGData |= !IsolatedActor->Data.GetAllInputs().IsEmpty();
+		}
 	}
 	else if (Settings->Mode == EPCGGetDataFromActorMode::GetDataFromProperty)
 	{
@@ -1023,6 +1027,39 @@ void FPCGDataFromActorElement::ProcessActor(FPCGContext* InContext, const UPCGDa
 				}
 			}
 		}
+		
+		// If the Actor is an Isolated actor, we also want its own data.
+		if (APCGIsolatedActor* IsolatedActor = Cast<APCGIsolatedActor>(FoundActor))
+		{
+			for (const FPCGTaggedData& TaggedData : IsolatedActor->Data.GetAllInputs())
+			{
+				if (ensure(TaggedData.Data))
+				{
+					FPCGTaggedData& DuplicatedTaggedData = Outputs.Add_GetRef(TaggedData);
+
+					// In cases where the owner of the data (component -> owner) isn't in the main world, we MUST duplicate the data, otherwise we'll cause reference leaks in the cache.
+#if WITH_EDITOR
+					if (!SourceOwner || (FoundActor && FoundActor->GetLevel() != SourceOwner->GetWorld()->PersistentLevel))
+					{
+						DuplicatedTaggedData.Data = Cast<UPCGData>(StaticDuplicateObject(TaggedData.Data, GetTransientPackage()));
+					}
+#endif
+					
+					//todo_pcg Should we add the grid size tags from the owner'component?
+					//  DuplicatedTaggedData.Tags.Add(PCGDataFromActorConstants::PCGComponentDataGridSizeTagPrefix + FString::FromInt(PCGHiGenGrid::GridToGridSize(Component->GetGenerationGrid())));
+
+					// Add some logging if we have un-expected output pin
+					if (DuplicatedTaggedData.Pin != PCGPinConstants::DefaultOutputLabel && !Settings->ExpectedPins.IsEmpty() && !Settings->ExpectedPins.Contains(TaggedData.Pin))
+					{
+						PCGLog::LogErrorOnGraph(FText::Format(LOCTEXT("InvalidOutputPinPCGComponentData", 
+							"Isolated actor '{0}' has data from an unexpected output pin: '{1}'.\nMaybe the data is outdated, try regenerate this component to refresh its output data"),
+							FText::FromName(FoundActor->GetFName()),
+							FText::FromName(TaggedData.Pin)), Context);
+					}
+				}
+			}
+			
+		}
 	}
 	else if (FoundProperty)
 	{
