[PCG] Prevent modification of GeneratedResources array outside of container so we can build acceleration tables from within (next step)

#rnx
#Synced-CL 48806509

diff --git a/Engine/Plugins/PCG/Source/PCG/Private/PCGComponent.cpp b/Engine/Plugins/PCG/Source/PCG/Private/PCGComponent.cpp
--- a/Engine/Plugins/PCG/Source/PCG/Private/PCGComponent.cpp
+++ b/Engine/Plugins/PCG/Source/PCG/Private/PCGComponent.cpp
@@ -287,9 +287,8 @@ void UPCGComponent::AddComponentsToManagedResources(const TArray<UActorComponent
 	if (!DefaultList)
 	{
 		DefaultList = NewObject<UPCGManagedComponentDefaultList>(this);
-
-		PCGGeneratedResourcesLogging::LogAddToManagedResources(this, DefaultList);
-		ManagedResourceContainer.GetManagedResourcesNoLock().Add(DefaultList);
+		FPCGManagerResourceContainerHelper Helper(ManagedResourceContainer);
+		Helper.AddManagedResourceNoLock(DefaultList);
 	}
 
 	check(DefaultList);
@@ -973,7 +972,7 @@ void UPCGComponent::PurgeUnlinkedResources(const AActor* InActor)
 		UE::TScopeLock ResourcesLock(Component->GeneratedResourcesLock);
 		if(ensure(Component->ManagedResourceContainer.AreResourcesAccessible()))
 		{
-			for (UPCGManagedResource* ManagedResource : Component->ManagedResourceContainer.GetManagedResourcesNoLock())
+			for (UPCGManagedResource* ManagedResource : Component->ManagedResourceContainer.GetConstManagedResourcesNoLock())
 			{
 				if (UPCGManagedComponent* ManagedComponent = Cast<UPCGManagedComponent>(ManagedResource))
 				{
@@ -1414,11 +1413,12 @@ bool UPCGComponent::MoveResourcesToNewActor(const FPCGMoveResourceParams& Params
 	{
 		UE::TScopeLock ResourcesLock(GeneratedResourcesLock);
 		ensure(ManagedResourceContainer.AreResourcesAccessible());
+		FPCGManagerResourceContainerHelper Helper(ManagedResourceContainer);
 #if WITH_EDITOR
 		Scope.AddResources(this, ManagedResourceContainer.GetConstManagedResourcesNoLock());
 #endif
 		TArray<TSoftObjectPtr<AActor>> TentativeMovedGeneratedActors;
-		TArray<TObjectPtr<UPCGManagedResource>>& ManagedResources = ManagedResourceContainer.GetManagedResourcesNoLock();
+		const TArray<TObjectPtr<UPCGManagedResource>>& ManagedResources = ManagedResourceContainer.GetConstManagedResourcesNoLock();
 
 		for(int ResourceIndex = ManagedResources.Num() - 1; ResourceIndex >=0; --ResourceIndex)
 		{
@@ -1446,7 +1446,7 @@ bool UPCGComponent::MoveResourcesToNewActor(const FPCGMoveResourceParams& Params
 					}
 #endif
 
-					ManagedResources.RemoveAtSwap(ResourceIndex);
+					Helper.RemoveManagedResourceAtNoLock(ResourceIndex);
 				}
 			}
 			else
@@ -1725,7 +1725,7 @@ void UPCGComponent::OnUnregister()
 			UE::TScopeLock ResourcesLock(GeneratedResourcesLock);
 			if (ensure(ManagedResourceContainer.AreResourcesAccessible()))
 			{
-				for (UPCGManagedResource* ManagedResource : ManagedResourceContainer.GetManagedResourcesNoLock())
+				for (UPCGManagedResource* ManagedResource : ManagedResourceContainer.GetConstManagedResourcesNoLock())
 				{
 					// Make sure resource is indeed a transient/preview resource
 					if (ManagedResource &&
@@ -1884,8 +1884,8 @@ PRAGMA_DISABLE_DEPRECATION_WARNINGS
 PRAGMA_ENABLE_DEPRECATION_WARNINGS
 
 	// Debug resources are transient and will generate null entries on save, clean them up here
-	TArray<TObjectPtr<UPCGManagedResource>>& ContainerGeneratedResources = ManagedResourceContainer.GetManagedResourcesNoLock();
-	ContainerGeneratedResources.Remove(nullptr);
+	FPCGManagerResourceContainerHelper Helper(ManagedResourceContainer);
+	Helper.RemoveManagedResourceNoLock(nullptr);
 
 	// Force dirty to be false on load. We should never refresh on load.
 	bDirtyGenerated = false;
@@ -1917,7 +1917,7 @@ PRAGMA_ENABLE_DEPRECATION_WARNINGS
 	}
 
 	/** Deprecation code, should be removed once generated data has been updated */
-	if (GetOwner() && bGenerated && ContainerGeneratedResources.Num() == 0)
+	if (GetOwner() && bGenerated && ManagedResourceContainer.IsEmpty())
 	{
 		TArray<UInstancedStaticMeshComponent*> ISMCs;
 		GetOwner()->GetComponents(ISMCs);
@@ -1928,7 +1928,7 @@ PRAGMA_ENABLE_DEPRECATION_WARNINGS
 			{
 				UPCGManagedISMComponent* ManagedComponent = NewObject<UPCGManagedISMComponent>(this);
 				ManagedComponent->GeneratedComponent = ISMC;
-				ContainerGeneratedResources.Add(ManagedComponent);
+				Helper.AddManagedResourceNoLock(ManagedComponent);
 			}
 		}
 
@@ -1936,7 +1936,7 @@ PRAGMA_ENABLE_DEPRECATION_WARNINGS
 		{
 			UPCGManagedActors* ManagedActors = NewObject<UPCGManagedActors>(this);
 			ManagedActors->GetMutableGeneratedActors() = GeneratedActors_DEPRECATED.Array();
-			ContainerGeneratedResources.Add(ManagedActors);
+			Helper.AddManagedResourceNoLock(ManagedActors);
 			GeneratedActors_DEPRECATED.Reset();
 		}
 	}
diff --git a/Engine/Plugins/PCG/Source/PCG/Private/PCGManagedResourceContainer.cpp b/Engine/Plugins/PCG/Source/PCG/Private/PCGManagedResourceContainer.cpp
--- a/Engine/Plugins/PCG/Source/PCG/Private/PCGManagedResourceContainer.cpp
+++ b/Engine/Plugins/PCG/Source/PCG/Private/PCGManagedResourceContainer.cpp
@@ -22,17 +22,35 @@ void FPCGManagedResourceContainer::SetOwner(IPCGGraphExecutionSource* InOwner)
 
 void FPCGManagedResourceContainer::AddManagedResource(UPCGManagedResource* InManagedResource)
 {
-	PCGGeneratedResourcesLogging::LogAddToManagedResources(Owner, InManagedResource);
+	AddManagedResourceInternal(InManagedResource, Owner->GetExecutionState().GetManagedResourceContainerLock());
+}
 
+void FPCGManagedResourceContainer::AddManagedResourceNoLock(UPCGManagedResource* InManagedResource)
+{
+	AddManagedResourceInternal(InManagedResource, nullptr);
+}
+
+void FPCGManagedResourceContainer::AddManagedResourceInternal(UPCGManagedResource* InManagedResource, FTransactionallySafeCriticalSection* InLock)
+{
+	PCGGeneratedResourcesLogging::LogAddToManagedResources(Owner, InManagedResource);
 	if (InManagedResource)
 	{
-		PCG::TScopeLock ResourcesLock(Owner->GetExecutionState().GetManagedResourceContainerLock());
+		PCG::TScopeLock ResourcesLock(InLock);
 		ensure(!bAreResourcesInaccessible);
 		GeneratedResources.Add(InManagedResource);
 		Owner->GetExecutionState().OnManagedResourceAdded(InManagedResource);
 	}
 }
 
+void FPCGManagedResourceContainer::RemoveManagedResourceAtNoLock(int32 ResourceIndex)
+{
+	GeneratedResources.RemoveAtSwap(ResourceIndex);
+}
+void FPCGManagedResourceContainer::RemoveManagedResourceNoLock(UPCGManagedResource* InManagedResource)
+{
+	GeneratedResources.Remove(InManagedResource);
+}
+
 void FPCGManagedResourceContainer::SetManagedResources(const TArray<TObjectPtr<UPCGManagedResource>>& InManagedResources)
 {
 	PCG::TScopeLock ResourcesLock(Owner->GetExecutionState().GetManagedResourceContainerLock());
diff --git a/Engine/Plugins/PCG/Source/PCG/Public/PCGManagedResourceContainer.h b/Engine/Plugins/PCG/Source/PCG/Public/PCGManagedResourceContainer.h
--- a/Engine/Plugins/PCG/Source/PCG/Public/PCGManagedResourceContainer.h
+++ b/Engine/Plugins/PCG/Source/PCG/Public/PCGManagedResourceContainer.h
@@ -29,7 +29,6 @@ public:
 	PCG_API void SetManagedResources(const TArray<TObjectPtr<UPCGManagedResource>>& InManagedResources);
 	PCG_API TArray<TObjectPtr<UPCGManagedResource>> GetManagedResourcesCopy() const;
 	
-	TArray<TObjectPtr<UPCGManagedResource>>& GetManagedResourcesNoLock() { return GeneratedResources; }
 	const TArray<TObjectPtr<UPCGManagedResource>>& GetConstManagedResourcesNoLock() const { return GeneratedResources; }
 	
 	PCG_API void ForEachManagedResource(TFunctionRef<void(UPCGManagedResource*)> InFunction);
@@ -65,6 +64,14 @@ public:
 #endif	
 
 private:
+	friend struct FPCGManagerResourceContainerHelper;
+
+	void AddManagedResourceNoLock(UPCGManagedResource* InManagedResource);
+	void RemoveManagedResourceNoLock(UPCGManagedResource* InManagedResource);
+	void RemoveManagedResourceAtNoLock(int32 ResourceIndex);
+
+	void AddManagedResourceInternal(UPCGManagedResource* InManagedResource, FTransactionallySafeCriticalSection* InLock);
+
 	bool bAreResourcesInaccessible = false;
 
 	UPROPERTY()
@@ -76,4 +83,34 @@ private:
 #endif
 
 	IPCGGraphExecutionSource* Owner = nullptr;
+};
+
+// Gives access to NoLock methods for friend classes only
+struct FPCGManagerResourceContainerHelper
+{
+	FPCGManagerResourceContainerHelper() = delete;
+	explicit FPCGManagerResourceContainerHelper(FPCGManagedResourceContainer& InContainer)
+		: Container(InContainer)
+	{
+	}
+
+private:
+	friend class UPCGComponent;
+
+	void AddManagedResourceNoLock(UPCGManagedResource* InManagedResource)
+	{
+		Container.AddManagedResourceNoLock(InManagedResource);
+	}
+
+	void RemoveManagedResourceNoLock(UPCGManagedResource* InManagedResource)
+	{
+		Container.RemoveManagedResourceNoLock(InManagedResource);
+	}
+
+	void RemoveManagedResourceAtNoLock(int32 ResourceIndex)
+	{
+		Container.RemoveManagedResourceAtNoLock(ResourceIndex);
+	}
+
+	FPCGManagedResourceContainer& Container;
 };
\ No newline at end of file
