[PCG] Allow grid linkage to request unbounded generation if necessary

#rnx
#Synced-CL 48806643

diff --git a/Engine/Plugins/PCG/Source/PCG/Private/Graph/PCGGraphExecutor.cpp b/Engine/Plugins/PCG/Source/PCG/Private/Graph/PCGGraphExecutor.cpp
--- a/Engine/Plugins/PCG/Source/PCG/Private/Graph/PCGGraphExecutor.cpp
+++ b/Engine/Plugins/PCG/Source/PCG/Private/Graph/PCGGraphExecutor.cpp
@@ -3207,12 +3207,10 @@ namespace PCGGraphExecutor
 			PCGGraphExecutionLogging::LogGridLinkageTaskExecuteRetrieve(InContext, InGenerationGrid, FromGridSize, ToGridSize, InResourceKey);
 
 			IPCGGraphExecutionSource* ExecSourceWithData = nullptr;
-			bool bExecSourceWithDataIsOriginal = false;
 
 			if (FromGridSize == PCGHiGenGrid::UnboundedGridSize())
 			{
 				ExecSourceWithData = ExecutionState.GetOriginalSource();
-				bExecSourceWithDataIsOriginal = true;
 			}
 			else
 			{
@@ -3257,11 +3255,8 @@ namespace PCGGraphExecutor
 			}
 
 			// At this point we could not get to the data, so we'll try executing the graph if we did not do that already.
-			// If we need data from a local execution source but the local execution source is still generating, then we'll wait for it. On the other hand
-			// if we need data from the original execution source we assume the generation has already happened because it is always scheduled before
-			// the local execution source.
-			const bool bWaitForGeneration = bExecSourceWithDataIsOriginal ? false : ExecSourceWithData->GetExecutionState().IsGenerating();
-			if (bWaitForGeneration)
+			// If we need data from a higher execution source but it's still generating, then we'll wait for it.
+			if (ExecSourceWithData->GetExecutionState().IsGenerating())
 			{
 				PCGGraphExecutionLogging::LogGridLinkageTaskExecuteRetrieveWaitOnScheduledGraph(InContext, ExecSourceWithData, InResourceKey);
 
@@ -3280,8 +3275,7 @@ namespace PCGGraphExecutor
 			}
 
 			// Graph is not currently generating. If we have not already tried generating, try it once now.
-			// But don't do this for the original execution source as this will always be scheduled before the local execution source.
-			if (!InContext->bScheduledGraph && !bExecSourceWithDataIsOriginal)
+			if (!InContext->bScheduledGraph)
 			{
 				PCGGraphExecutionLogging::LogGridLinkageTaskExecuteRetrieveScheduleGraph(InContext, ExecSourceWithData, InResourceKey);
 
diff --git a/Engine/Plugins/PCG/Source/PCG/Private/Utils/PCGGraphExecutionLogging.cpp b/Engine/Plugins/PCG/Source/PCG/Private/Utils/PCGGraphExecutionLogging.cpp
--- a/Engine/Plugins/PCG/Source/PCG/Private/Utils/PCGGraphExecutionLogging.cpp
+++ b/Engine/Plugins/PCG/Source/PCG/Private/Utils/PCGGraphExecutionLogging.cpp
@@ -496,7 +496,7 @@ namespace PCGGraphExecutionLogging
 		}
 		check(InContext);
 
-		UE_LOG(LogPCG, Warning, TEXT("[GRIDLINKING] [%s] RETRIEVE: FAILED: No data found on local source. Source=%s, Path=%s"),
+		UE_LOG(LogPCG, Warning, TEXT("[GRIDLINKING] [%s] RETRIEVE: FAILED: No data found on source. Source=%s, Path=%s"),
 			*PCGLog::GetExecutionSourceName(InContext->ExecutionSource.Get(), /*bUseLabel=*/true),
 			*PCGLog::GetExecutionSourceName(InExecutionSource, /*bUseLabel=*/true),
 			*InResourcePath);
