[PCG] Add support for GPU proxies to Filter Data By Tag

#rnx
#Synced-CL 48802662

diff --git a/Engine/Plugins/PCG/Source/PCG/Public/Elements/PCGFilterByTag.h b/Engine/Plugins/PCG/Source/PCG/Public/Elements/PCGFilterByTag.h
--- a/Engine/Plugins/PCG/Source/PCG/Public/Elements/PCGFilterByTag.h
+++ b/Engine/Plugins/PCG/Source/PCG/Public/Elements/PCGFilterByTag.h
@@ -61,4 +61,5 @@ protected:
 	virtual bool ExecuteInternal(FPCGContext* Context) const override;
 	virtual EPCGElementExecutionLoopMode ExecutionLoopMode(const UPCGSettings* Settings) const override { return EPCGElementExecutionLoopMode::SinglePrimaryPin; }
 	virtual bool SupportsBasePointDataInputs(FPCGContext* InContext) const override { return true; }
+	virtual bool SupportsGPUResidentData(FPCGContext* InContext) const override { return true; }
 };
