[PCG] Fix infinite recursion crash by deferring tick from ALV request refresh rather than pumping ticks directly.

#rnx
#Synced-CL 48811918

diff --git a/Engine/Plugins/PCG/Source/PCGEditor/Private/Widgets/SPCGEditorGraphAttributeListView.cpp b/Engine/Plugins/PCG/Source/PCGEditor/Private/Widgets/SPCGEditorGraphAttributeListView.cpp
--- a/Engine/Plugins/PCG/Source/PCGEditor/Private/Widgets/SPCGEditorGraphAttributeListView.cpp
+++ b/Engine/Plugins/PCG/Source/PCGEditor/Private/Widgets/SPCGEditorGraphAttributeListView.cpp
@@ -715,6 +715,8 @@ void SPCGEditorGraphAttributeListView::OnTick()
 
 		CurrentUpdateTask.Reset();
 	}
+
+	bTickForced = false;
 }
 
 bool SPCGEditorGraphAttributeListView::ReadbackGPUData() const
@@ -747,7 +749,24 @@ void SPCGEditorGraphAttributeListView::RequestRefresh()
 	PRAGMA_ENABLE_DEPRECATION_WARNINGS
 	{
 		RequestViewportRefresh();
-		OnTick();
+
+		if (!bTickForced)
+		{
+			// Force tick on next game thread tick to ensure refresh happens.
+			FTSTicker::GetCoreTicker().AddTicker(
+				FTickerDelegate::CreateLambda([WeakThis=AsWeak()](float DeltaTime)
+				{
+					if (SPCGEditorGraphAttributeListView* This = static_cast<SPCGEditorGraphAttributeListView*>(WeakThis.Pin().Get()))
+					{
+						This->OnTick();
+					}
+					return false; // Run once
+				}),
+				/*InDelay=*/0.0f
+			);
+
+			bTickForced = true;
+		}
 	}
 }
 
diff --git a/Engine/Plugins/PCG/Source/PCGEditor/Private/Widgets/SPCGEditorGraphAttributeListView.h b/Engine/Plugins/PCG/Source/PCGEditor/Private/Widgets/SPCGEditorGraphAttributeListView.h
--- a/Engine/Plugins/PCG/Source/PCGEditor/Private/Widgets/SPCGEditorGraphAttributeListView.h
+++ b/Engine/Plugins/PCG/Source/PCGEditor/Private/Widgets/SPCGEditorGraphAttributeListView.h
@@ -336,6 +336,7 @@ private:
 
 	bool bNeedsRefresh : 1 = false;
 	bool bViewportNeedsRefresh : 1 = false;
+	bool bTickForced : 1 = false;
 
 	bool bKeepDataComboBoxSelection = true;
 
