[PCG] Fix unit test

#rnx
#Synced-CL 48819314

diff --git a/Engine/Plugins/PCG/Source/PCG/Private/Tests/Elements/PCGCopyAttributeTest.cpp b/Engine/Plugins/PCG/Source/PCG/Private/Tests/Elements/PCGCopyAttributeTest.cpp
--- a/Engine/Plugins/PCG/Source/PCG/Private/Tests/Elements/PCGCopyAttributeTest.cpp
+++ b/Engine/Plugins/PCG/Source/PCG/Private/Tests/Elements/PCGCopyAttributeTest.cpp
@@ -1077,7 +1077,7 @@ bool FPCGCopyAttributeTests_MultiDomain_CopyAll_InvalidDomain::RunTest(const FSt
 	PCGCopyAttributeTests::ConnectToSource(Context.Get(), SourcePointData);
 	PCGCopyAttributeTests::ConnectToTarget(Context.Get(), TargetPointData);
 
-	AddExpectedError(TEXT("Metadata domain Bliblibli is invalid for this data."), EAutomationExpectedErrorFlags::MatchType::Exact, 1, /*bIsRegex=*/false);
+	AddExpectedError(TEXT("Metadata domain Bliblibli is invalid for this data."), EAutomationExpectedErrorFlags::MatchType::Contains, 1, /*bIsRegex=*/false);
 
 	FPCGElementPtr TestElement = TestData.Settings->GetElement();
 	while (!TestElement->Execute(Context.Get())) {}
