[PCG] Move logging function in the right cpp file

#rnx
#okforgihub public
#Synced-CL 48819314

diff --git a/Engine/Plugins/PCG/Source/PCG/Private/PCGModule.cpp b/Engine/Plugins/PCG/Source/PCG/Private/PCGModule.cpp
--- a/Engine/Plugins/PCG/Source/PCG/Private/PCGModule.cpp
+++ b/Engine/Plugins/PCG/Source/PCG/Private/PCGModule.cpp
@@ -332,28 +332,4 @@ IMPLEMENT_MODULE(FPCGModule, PCG);
 
 PCG_API DEFINE_LOG_CATEGORY(LogPCG);
 
-void PCGLog::LogErrorOnGraph(const FText& InMsg, const FPCGContext* InContext)
-{
-	if (const FPCGContext* Context = InContext ? InContext : FPCGContext::GetContextForThread())
-	{
-		PCGE_LOG_C(Error, GraphAndLog, Context, InMsg);
-	}
-	else
-	{
-		UE_LOG(LogPCG, Error, TEXT("%s"), *InMsg.ToString());
-	}
-}
-
-void PCGLog::LogWarningOnGraph(const FText& InMsg, const FPCGContext* InContext)
-{
-	if (const FPCGContext* Context = InContext ? InContext : FPCGContext::GetContextForThread())
-	{
-		PCGE_LOG_C(Warning, GraphAndLog, Context, InMsg);
-	}
-	else
-	{
-		UE_LOG(LogPCG, Warning, TEXT("%s"), *InMsg.ToString());
-	}
-}
-
 #undef LOCTEXT_NAMESPACE
diff --git a/Engine/Plugins/PCG/Source/PCG/Private/Utils/PCGLogErrors.cpp b/Engine/Plugins/PCG/Source/PCG/Private/Utils/PCGLogErrors.cpp
--- a/Engine/Plugins/PCG/Source/PCG/Private/Utils/PCGLogErrors.cpp
+++ b/Engine/Plugins/PCG/Source/PCG/Private/Utils/PCGLogErrors.cpp
@@ -2,8 +2,36 @@
 
 #include "Utils/PCGLogErrors.h"
 
+#include "PCGContext.h"
+#include "PCGElement.h"
+#include "PCGModule.h"
+
 #define LOCTEXT_NAMESPACE "PCGLog"
 
+void PCGLog::LogErrorOnGraph(const FText& InMsg, const FPCGContext* InContext)
+{
+	if (const FPCGContext* Context = InContext ? InContext : FPCGContext::GetContextForThread())
+	{
+		PCGE_LOG_C(Error, GraphAndLog, Context, InMsg);
+	}
+	else
+	{
+		UE_LOG(LogPCG, Error, TEXT("[NOCONTEXT] %s"), *InMsg.ToString());
+	}
+}
+
+void PCGLog::LogWarningOnGraph(const FText& InMsg, const FPCGContext* InContext)
+{
+	if (const FPCGContext* Context = InContext ? InContext : FPCGContext::GetContextForThread())
+	{
+		PCGE_LOG_C(Warning, GraphAndLog, Context, InMsg);
+	}
+	else
+	{
+		UE_LOG(LogPCG, Warning, TEXT("[NOCONTEXT] %s"), *InMsg.ToString());
+	}
+}
+
 namespace PCGLog::InputOutput
 {
 	namespace Format
diff --git a/Engine/Plugins/PCG/Source/PCG/Public/Utils/PCGLogErrors.h b/Engine/Plugins/PCG/Source/PCG/Public/Utils/PCGLogErrors.h
--- a/Engine/Plugins/PCG/Source/PCG/Public/Utils/PCGLogErrors.h
+++ b/Engine/Plugins/PCG/Source/PCG/Public/Utils/PCGLogErrors.h
@@ -10,9 +10,16 @@ struct FPCGContext;
 
 namespace PCGLog
 {
-	/** Convenience function that would either log error on the graph if there is a context, or in the console if not. */
+	/**
+	 * Convenience function that would either log error on the graph if there is a context, or in the console if not.
+	 * Note that if the context is not set, it will look for the current context for this thread.
+	 */
 	PCG_API void LogErrorOnGraph(const FText& InMsg, const FPCGContext* InContext = nullptr);
-	/** Convenience function that would either log warning on the graph if there is a context, or in the console if not. */
+	
+	/**
+	 * Convenience function that would either log warning on the graph if there is a context, or in the console if not.
+	 * Note that if the context is not set, it will look for the current context for this thread.
+	 */
 	PCG_API void LogWarningOnGraph(const FText& InMsg, const FPCGContext* InContext = nullptr);
 
 	namespace InputOutput
