[PCG] First pass multi-material support for Generate Landscape Textures / Generate Grass Maps.

* Builds list of all components from compatible LS proxies
* Aggregrates all grass types from all components into a global grass type list, and similar for height
* Give the map from global index -> export tile & channel index to unpacker

[RN] [PCG] Added support for varying landscape materials and grass types to grass map generation.
#Synced-CL 48811918

diff --git a/Engine/Plugins/PCG/Shaders/Private/PCGGrassMapUnpackerCS.usf b/Engine/Plugins/PCG/Shaders/Private/PCGGrassMapUnpackerCS.usf
--- a/Engine/Plugins/PCG/Shaders/Private/PCGGrassMapUnpackerCS.usf
+++ b/Engine/Plugins/PCG/Shaders/Private/PCGGrassMapUnpackerCS.usf
@@ -10,12 +10,14 @@
 // Generated grass maps all end up in one big row (see FLandscapeGrassWeightExporter::FLandscapeGrassWeightExporter()).
 //
 //  -----------------------------------------
-// | LC0  | LC0  | LC1  | LC1  | LC2  | LC2  |
-// | x=h0 | x=g2 | x=h0 | x=g2 | x=h0 | x=g2 |  NumComponents=3, NumGrassMaps=3, InNumGrassMapPasses=2
-// | y=h1 | y=   | y=h1 | y=   | y=h1 | y=   |
-// | z=g0 | z=   | z=g0 | z=   | z=g0 | z=   |
-// | w=g1 | w=   | w=g1 | w=   | w=g1 | w=   |
+// | LC0  | LC0  | LC1  | LC1  | LC1  | LC2  |
+// | x=h0 | x=g2 | x=h0 | x=g2 | x=g6 | x=h0 |  InLandscapeComponentResolution
+// | y=h1 | y=   | y=h1 | y=g3 | y=g7 | y=h1 |
+// | z=g0 | z=   | z=g0 | z=g4 | z=g8 | z=g0 |
+// | w=g1 | w=   | w=g1 | w=g5 | w=   | w=   |
 //  -----------------------------------------
+// InLandscapeComponentResolution*NumExportedTiles
+
 Texture2D<float4> InPackedGrassMaps;
 
 // Full output domain consists of multiple components, but we only run on a "window" that covers the generation cell bounds.
@@ -47,10 +49,13 @@ RWTexture2DArray<float> OutUnpackedGrassMaps;
 // Height output. Dimension is (NumComponentsX * InLandscapeComponentResolution, NumComponentsY * InLandscapeComponentResolution, NumGrassMaps).
 RWTexture2D<float> OutUnpackedHeight;
 
+StructuredBuffer<int> InHeightIndexMap;
+StructuredBuffer<int> InGrassIndexMap;
+int InNumGlobalGrassTypes;
+
 int4 InLinearTileIndexToComponentIndex[PCG_MAX_NUM_LANDSCAPE_COMPONENTS]; // Only x component used, 16 byte alignment on arrays
 uint2 InNumTiles;
 uint InLandscapeComponentResolution; // One texel per vert, so number of quads + 1
-uint InNumGrassMapPasses;
 float3 InLandscapeGridScale; // x == LS Actor DrawScale.X, y == LS Actor DrawScale.y, z == LS Actor DrawScale.z / 128.0f (ZSCALE)
 float InLandscapeLocationZ;
 uint InOutputHeight;
@@ -89,24 +94,44 @@ void PCGGrassMapUnpacker_CS(const uint3 WindowCoord_Slice : SV_DispatchThreadID)
 		return;
 	}
 
-	// Coordinates into input height & grass maps (see packing above).
-	const uint PassIndex = (Slice + 2) / 4;
-	const uint ComponentXOffset = (ComponentIndex * InNumGrassMapPasses + PassIndex) * InLandscapeComponentResolution;
-
-	// Coordinates into packed data which are laid out horizontally (see above).
 	const uint2 LocalTileCoords = uint2(FullCoord.x - TileCoords.x * (InLandscapeComponentResolution - 1), FullCoord.y - TileCoords.y * (InLandscapeComponentResolution - 1));
-	const uint2 InputCoords = uint2(LocalTileCoords.x + ComponentXOffset, LocalTileCoords.y);
 
 	// Heights. The first layer threads also deal with height.
 	if (InOutputHeight == 1 && Slice == 0)
 	{
-		const float RelativeHeight = (UnpackHeight(InPackedGrassMaps[InputCoords].xy) - LANDSCAPE_MID_VALUE) * InLandscapeGridScale.z * TERRAIN_ZSCALE;
-		OutUnpackedHeight[WindowCoord] = RelativeHeight + InLandscapeLocationZ;
+		const int HeightExportTileIndex = InHeightIndexMap[ComponentIndex];
+		if (HeightExportTileIndex != -1)
+		{
+			const uint ComponentXOffset = HeightExportTileIndex * InLandscapeComponentResolution;
+			const uint2 InputCoords = uint2(LocalTileCoords.x + ComponentXOffset, LocalTileCoords.y);
+
+			const float RelativeHeight = (UnpackHeight(InPackedGrassMaps[InputCoords].xy) - LANDSCAPE_MID_VALUE) * InLandscapeGridScale.z * TERRAIN_ZSCALE;
+			OutUnpackedHeight[WindowCoord] = RelativeHeight + InLandscapeLocationZ;
+		}
 	}
 
-	// Grass map source channel (see packing above).
-	const uint InputChannelIndex = (Slice + 2) % 4;
+	// Coordinates into input height & grass maps (see packing above).
+	const int GrassExportChannelIndex = InGrassIndexMap[ComponentIndex * InNumGlobalGrassTypes + Slice];
+
+	if (GrassExportChannelIndex != -1)
+	{
+		const int GrassExportTileIndex = GrassExportChannelIndex / 4;
+		const uint ComponentXOffset = GrassExportTileIndex * InLandscapeComponentResolution;
+		// Grass map source channel (see packing above).
+		const uint ExportChannelIndex = GrassExportChannelIndex % 4;
+
+		// Coordinates into packed data which are laid out horizontally (see above).
+		const uint2 InputCoords = uint2(LocalTileCoords.x + ComponentXOffset, LocalTileCoords.y);
+
+		// Grass maps
+		OutUnpackedGrassMaps[uint3(WindowCoord, Slice)] = InPackedGrassMaps[InputCoords][ExportChannelIndex];
 
-	// Grass maps
-	OutUnpackedGrassMaps[uint3(WindowCoord, Slice)] = InPackedGrassMaps[InputCoords][InputChannelIndex];
+		// TODO: position of height data can move around based on slices, so also needs to be in mapping somehow
+		// Heights. The first layer threads also deal with height.
+		if (InOutputHeight == 1 && Slice == 0)
+		{
+			const float RelativeHeight = (UnpackHeight(InPackedGrassMaps[InputCoords].xy) - LANDSCAPE_MID_VALUE) * InLandscapeGridScale.z * TERRAIN_ZSCALE;
+			OutUnpackedHeight[WindowCoord] = RelativeHeight + InLandscapeLocationZ;
+		}
+	}
 }
diff --git a/Engine/Plugins/PCG/Source/PCG/Private/Elements/PCGGenerateLandscapeTextures.cpp b/Engine/Plugins/PCG/Source/PCG/Private/Elements/PCGGenerateLandscapeTextures.cpp
--- a/Engine/Plugins/PCG/Source/PCG/Private/Elements/PCGGenerateLandscapeTextures.cpp
+++ b/Engine/Plugins/PCG/Source/PCG/Private/Elements/PCGGenerateLandscapeTextures.cpp
@@ -273,131 +273,154 @@ bool FPCGGenerateLandscapeTexturesElement::ExecuteInternal(FPCGContext* InContex
 		}
 
 		const FBox GenerationBounds = ExecutionSource->GetExecutionState().GetBounds();
-		TArray<TObjectPtr<ULandscapeGrassType>> GrassTypes;
+		TArray<ULandscapeComponent*> OverlappingComponents;
 
-		for (TSoftObjectPtr<ALandscapeProxy> LandscapeProxyPtr : LandscapeData->Landscapes)
+		for (const TSoftObjectPtr<ALandscapeProxy>& LandscapeProxyPtr : LandscapeData->Landscapes)
 		{
 			ALandscapeProxy* LandscapeProxy = LandscapeProxyPtr.Get();
-
 			if (!LandscapeProxy || !LandscapeProxy->GetLandscapeMaterial())
 			{
 				continue;
 			}
 
-			if (!Context->LandscapeProxy.Get())
+			// Grow storage to accommodate new components.
+			const int32 NumToReserve = FMath::Min(OverlappingComponents.Num() + LandscapeProxy->LandscapeComponents.Num(), (int32)FPCGGrassMapUnpackerCS::MaxNumLandscapeComponents);
+			OverlappingComponents.Reserve(NumToReserve);
+			Context->LandscapeComponents.Reserve(NumToReserve);
+
+			if (!Context->FirstLandscapeProxy.Get())
 			{
-				Context->LandscapeProxy = LandscapeProxy;
+				Context->FirstLandscapeProxy = LandscapeProxy;
 			}
-
-			for (ULandscapeComponent* LandscapeComponent : LandscapeProxy->LandscapeComponents)
+			else
 			{
-				if (!ensure(LandscapeComponent))
+				// Validate the proxies are compatible.
+				if (Context->FirstLandscapeProxy->ComponentSizeQuads != LandscapeProxy->ComponentSizeQuads
+					|| Context->FirstLandscapeProxy->SubsectionSizeQuads != LandscapeProxy->SubsectionSizeQuads
+					|| Context->FirstLandscapeProxy->NumSubsections != LandscapeProxy->NumSubsections
+					|| !Context->FirstLandscapeProxy->GetRootComponent()->GetRelativeScale3D().Equals(LandscapeProxy->GetRootComponent()->GetRelativeScale3D())
+					|| !FMath::IsNearlyEqual(Context->FirstLandscapeProxy->GetRootComponent()->GetRelativeLocation().Z, LandscapeProxy->GetRootComponent()->GetRelativeLocation().Z))
 				{
-					continue;
+					PCGLog::LogErrorOnGraph(LOCTEXT("LandscapeProxyMismatch", "Encountered multiple landscape proxies with different settings (component size quads or subsection sizes) or different transforms (scale XYZ or translation Z) which is not supported."), Context);
+					return true;
 				}
+			}
 
-				// Only generate grass map if there is meaningful overlap with our domain of interest.
-				const FBox LandscapeComponentBounds = LandscapeComponent->Bounds.GetBox();
-				if (LandscapeComponentBounds.Overlap(GenerationBounds).GetVolume() > UE_KINDA_SMALL_NUMBER)
+			for (int32 ComponentIndex = 0; ComponentIndex < LandscapeProxy->LandscapeComponents.Num(); ++ComponentIndex)
+			{
+				if (ULandscapeComponent* LandscapeComponent = LandscapeProxy->LandscapeComponents[ComponentIndex])
 				{
-					if (Context->LandscapeComponents.Num() >= FPCGGrassMapUnpackerCS::MaxNumLandscapeComponents)
+					// Only generate grass map if there is meaningful overlap with our domain of interest.
+					const FBox LandscapeComponentBounds = LandscapeComponent->Bounds.GetBox();
+					if (LandscapeComponentBounds.Overlap(GenerationBounds).GetVolume() > UE_KINDA_SMALL_NUMBER)
 					{
-						PCGLog::LogWarningOnGraph(LOCTEXT("MaxLandscapeComponentsExceeded", "Too many landscape components overlap the generation domain. Consider partitioning the component onto a smaller grid size."), Context);
-						break;
-					}
-
-					LandscapeComponent->UpdateGrassTypes();
+						if (OverlappingComponents.Num() >= FPCGGrassMapUnpackerCS::MaxNumLandscapeComponents)
+						{
+							PCGLog::LogWarningOnGraph(LOCTEXT("MaxLandscapeComponentsExceeded", "Too many landscape components overlap the generation domain. Consider partitioning the component onto a smaller grid size."), Context);
+							break;
+						}
 
-					Context->LandscapeComponents.Add(LandscapeComponent);
+						for (auto& GrassType : LandscapeComponent->GetGrassTypes())
+						{
+							Context->GlobalGrassTypes.AddUnique(GrassType);
+						}
 
-					if (Context->LandscapeComponents.Num() == 1)
-					{
-						GrassTypes = LandscapeComponent->GetGrassTypes();
-						Context->NumGrassTypes = GrassTypes.Num();
+						LandscapeComponent->UpdateGrassTypes();
 
-						FString GrassTypeName;
-						GrassTypeName.Reserve(256);
+						Context->LandscapeComponents.Add(LandscapeComponent);
+						OverlappingComponents.Add(LandscapeComponent);
 
-						for (int GrassTypeIndex = 0; GrassTypeIndex < GrassTypes.Num(); ++GrassTypeIndex)
+						if (Context->LandscapeComponents.Num() == 1)
 						{
-							if (TObjectPtr<ULandscapeGrassType> GrassType = GrassTypes[GrassTypeIndex])
-							{
-								GrassType->GetName(GrassTypeName);
+							Context->TotalLandscapeComponentBounds = LandscapeComponentBounds;
 
-								const bool bIsSelectedLayer = SelectedGrassTypes.Contains(GrassTypeName);
+							Context->LandscapeComponentExtent = LandscapeComponent->Bounds.BoxExtent.X * 2.0;
+							Context->TexelSizeWorld = Context->LandscapeComponentExtent / static_cast<double>(LandscapeProxy->ComponentSizeQuads);
 
-								if (bIsSelectedLayer != Settings->bExcludeSelectedGrassTypes)
-								{
-									Context->SelectedGrassTypes.Emplace(GrassType, GrassTypeIndex);
-								}
+							if (!ensure(FMath::IsNearlyEqual(LandscapeComponent->Bounds.BoxExtent.X, LandscapeComponent->Bounds.BoxExtent.Y)))
+							{
+								return true;
 							}
-						}
 
-						Context->TotalLandscapeComponentBounds = LandscapeComponentBounds;
-
-						Context->LandscapeComponentExtent = LandscapeComponent->Bounds.BoxExtent.X * 2.0;
-						Context->TexelSizeWorld = Context->LandscapeComponentExtent / static_cast<double>(LandscapeProxy->ComponentSizeQuads);
-
-						if (!ensure(FMath::IsNearlyEqual(LandscapeComponent->Bounds.BoxExtent.X, LandscapeComponent->Bounds.BoxExtent.Y)))
+							if (!ensure(Context->LandscapeComponentExtent > 0.0))
+							{
+								return true;
+							}
+						}
+						else
 						{
-							return true;
+							Context->TotalLandscapeComponentBounds += LandscapeComponentBounds;
+
+							// Currently assuming all landscape have the same extents.
+							ensure(FMath::IsNearlyEqual(Context->LandscapeComponentExtent, LandscapeComponent->Bounds.BoxExtent.X * 2.0));
 						}
 
-						if (!ensure(Context->LandscapeComponentExtent > 0.0))
+	#if WITH_EDITOR
+						if (PCGGenerateLandscapeTextures::CVarDebugDrawGeneratedComponents.GetValueOnGameThread())
 						{
-							return true;
+							DrawDebugBox(ExecutionSource->GetExecutionState().GetWorld(),
+								LandscapeComponentBounds.GetCenter(),
+								LandscapeComponentBounds.GetExtent(),
+								FColor::MakeRandomSeededColor(Context->TaskId),
+								/*bPersistentLines=*/false,
+								/*LifeTime=*/3.0f);
 						}
+	#endif
 					}
-					else
-					{
-						Context->TotalLandscapeComponentBounds += LandscapeComponentBounds;
+				}
+			}
+		}
 
-						const UMaterialInterface* LandscapeProxyMaterial = Context->LandscapeProxy->GetLandscapeMaterial();
-						const UMaterialInterface* LandscapeMaterial = LandscapeComponent->GetLandscapeMaterial();
+		// Filter grass types based on selected/excluded.
+		FString GrassTypeName;
+		GrassTypeName.Reserve(256);
 
-						// We expect all landscape components for a single landscape to have the same grass types.
-						// TODO: In order to support distinct grass types we need to do a dispatch of one landscape grass exporter per LS proxy.
-						if (LandscapeMaterial != LandscapeProxyMaterial && GrassTypes != LandscapeComponent->GetGrassTypes())
-						{
-#if !(UE_BUILD_SHIPPING || UE_BUILD_TEST) || USE_LOGGING_IN_SHIPPING
-							const FString LandscapeProxyMaterialName = LandscapeProxyMaterial ? LandscapeProxyMaterial->GetName() : TEXT("MISSINGMATERIAL");
-							const FString LandscapeMaterialName = LandscapeMaterial ? LandscapeMaterial->GetName() : TEXT("MISSINGMATERIAL");
-							
-							PCGLog::LogErrorOnGraph(
-								FText::Format(LOCTEXT("LandscapeComponentGrassTypeMismatch",
-									"Encountered landscape component '{0}' with material '{1}' that does not match the grass types on material '{2}'. All landscape materials must have the same grass types at this time."),
-									FText::FromString(LandscapeComponent->GetName()),
-									FText::FromString(LandscapeMaterialName),
-									FText::FromString(LandscapeProxyMaterialName)),
-								Context);
-#endif
-							return true;
-						}
+		for (int GrassTypeIndex = 0; GrassTypeIndex < Context->GlobalGrassTypes.Num(); ++GrassTypeIndex)
+		{
+			if (const TObjectPtr<ULandscapeGrassType>& GrassType = Context->GlobalGrassTypes[GrassTypeIndex])
+			{
+				GrassType->GetName(GrassTypeName);
 
-						// Currently assuming all landscape have the same extents.
-						ensure(FMath::IsNearlyEqual(Context->LandscapeComponentExtent, LandscapeComponent->Bounds.BoxExtent.X * 2.0));
-					}
+				const bool bIsSelectedLayer = SelectedGrassTypes.Contains(GrassTypeName);
 
-#if WITH_EDITOR
-					if (PCGGenerateLandscapeTextures::CVarDebugDrawGeneratedComponents.GetValueOnGameThread())
-					{
-						DrawDebugBox(ExecutionSource->GetExecutionState().GetWorld(),
-							LandscapeComponentBounds.GetCenter(),
-							LandscapeComponentBounds.GetExtent(),
-							FColor::MakeRandomSeededColor(Context->TaskId),
-							/*bPersistentLines=*/false,
-							/*LifeTime=*/3.0f);
-					}
-#endif
+				if (bIsSelectedLayer != Settings->bExcludeSelectedGrassTypes)
+				{
+					Context->SelectedGrassTypes.Emplace(GrassType, GrassTypeIndex);
 				}
 			}
 		}
 
-		if (!Context->LandscapeProxy.Get() || Context->LandscapeComponents.IsEmpty() || (Context->SelectedGrassTypes.IsEmpty() && !Settings->bGenerateHeightMap))
+		if (!Context->FirstLandscapeProxy.Get() || Context->LandscapeComponents.IsEmpty() || (Context->SelectedGrassTypes.IsEmpty() && !Settings->bGenerateHeightMap))
 		{
 			return true;
 		}
 
+		// Set up index mapping so that export tile/channel index can be looked up from the component index and global grass map index or height index during unpack.
+		{
+			Context->IndexMapping.Init(OverlappingComponents.Num(), Context->GlobalGrassTypes.Num());
+
+			int32 NumExportedTiles = 0;
+
+			for (int32 OverlapIndex = 0; OverlapIndex < OverlappingComponents.Num(); ++OverlapIndex)
+			{
+				ULandscapeComponent* LandscapeComponent = OverlappingComponents[OverlapIndex];
+				check(LandscapeComponent);
+
+				// Register export data tile index that contains height data for this component.
+				Context->IndexMapping.SetHeightEntry(OverlapIndex, NumExportedTiles);
+
+				for (int32 ComponentGrassTypeIndex = 0; ComponentGrassTypeIndex < LandscapeComponent->GetGrassTypes().Num(); ++ComponentGrassTypeIndex)
+				{
+					// Register export data channel index that contains grass map data for this component & grass channel.
+					const int32 GlobalGrassMapIndex = Context->GlobalGrassTypes.IndexOfByKey(LandscapeComponent->GetGrassTypes()[ComponentGrassTypeIndex]);
+					Context->IndexMapping.SetGrassMapEntry(OverlapIndex, GlobalGrassMapIndex, NumExportedTiles, ComponentGrassTypeIndex);
+				}
+
+				// Account for export tiles added for this component.
+				NumExportedTiles += FMath::DivideAndRoundUp<int32>(LandscapeComponent->GetGrassTypes().Num() + /*Channels for height*/2, 4);
+			}
+		}
+
 		Context->bLandscapeComponentsFiltered = true;
 	}
 
@@ -416,7 +439,7 @@ bool FPCGGenerateLandscapeTexturesElement::ExecuteInternal(FPCGContext* InContex
 						Context->TexturesToStream.Add(HeightMap);
 					}
 
-					if (Context->NumGrassTypes > 0)
+					if (Context->GlobalGrassTypes.Num() > 0)
 					{
 						const ERHIFeatureLevel::Type FeatureLevel = World->GetFeatureLevel();
 						for (UTexture2D* WeightmapTexture : LandscapeComponent->GetRenderedWeightmapTexturesForFeatureLevel(FeatureLevel))
@@ -502,7 +525,7 @@ bool FPCGGenerateLandscapeTexturesElement::ExecuteInternal(FPCGContext* InContex
 	// 3. Schedule grass map generation.
 	if (!Context->bGenerationScheduled)
 	{
-		ALandscapeProxy* LandscapeProxy = Context->LandscapeProxy.IsValid() ? Context->LandscapeProxy.Get() : nullptr;
+		ALandscapeProxy* LandscapeProxy = Context->FirstLandscapeProxy.IsValid() ? Context->FirstLandscapeProxy.Get() : nullptr;
 		if (!LandscapeProxy)
 		{
 			PCGLog::LogErrorOnGraph(LOCTEXT("LandscapeProxyLost", "Reference to landscape proxy actor lost, grass maps will not be generated."), InContext);
@@ -539,7 +562,7 @@ bool FPCGGenerateLandscapeTexturesElement::ExecuteInternal(FPCGContext* InContex
 		Context->LandscapeGrassWeightExporter = new FLandscapeGrassWeightExporter(
 			LandscapeProxy,
 			LandscapeComponents,
-			/*bInNeedsGrassmap=*/Context->NumGrassTypes > 0,
+			/*bInNeedsGrassmap=*/!Context->GlobalGrassTypes.IsEmpty(),
 			Settings->bGenerateHeightMap,
 			/*HeightMips=*/{},
 			/*bExport=*/false,
@@ -612,7 +635,7 @@ bool FPCGGenerateLandscapeTexturesElement::ExecuteInternal(FPCGContext* InContex
 		
 		TWeakPtr<FPCGContextHandle> ContextHandle = Context->GetOrCreateHandle();
 
-		ENQUEUE_RENDER_COMMAND(GenerateGrassMaps)([ContextHandle, LandscapeComponentExtent=Context->LandscapeComponentExtent, LandscapeTileCoords=MoveTemp(LandscapeTileCoords), LandscapeComponentResolution, LandscapeGridScale, LandscapeZ, bGenerateHeightMap=Settings->bGenerateHeightMap, NumTilesX, NumTilesY, OutputResolution, WindowStart, WindowEnd](FRHICommandListImmediate& RHICmdList)
+		ENQUEUE_RENDER_COMMAND(GenerateGrassMaps)([ContextHandle, LandscapeComponentExtent=Context->LandscapeComponentExtent, LandscapeTileCoords=MoveTemp(LandscapeTileCoords), LandscapeComponentResolution, LandscapeGridScale, LandscapeZ, bGenerateHeightMap=Settings->bGenerateHeightMap, NumTilesX, NumTilesY, OutputResolution, WindowStart, WindowEnd, IndexMapping=MoveTemp(Context->IndexMapping)](FRHICommandListImmediate& RHICmdList) // todo ok to move?
 		{
 			LLM_SCOPE_BYTAG(PCG);
 
@@ -637,8 +660,8 @@ bool FPCGGenerateLandscapeTexturesElement::ExecuteInternal(FPCGContext* InContex
 			// Generate grass maps. All will be generated to a single texture.
 			Context->LandscapeGrassWeightExporter->RenderLandscapeComponentToTexture_RenderThread(GraphBuilder, GrassMapTexture);
 
-			const uint32 NumGrassTypes = Context->NumGrassTypes;
-			const uint32 NumOutputSlices = FMath::Max(NumGrassTypes, 1u);
+			// Must have at least one slice to be a valid resource.
+			const uint32 NumOutputSlices = FMath::Max(Context->GlobalGrassTypes.Num(), 1);
 
 			FRDGTextureDesc GrassHeightMapDesc = FRDGTextureDesc::Create2D(
 				bGenerateHeightMap ? OutputResolution : FIntPoint(1, 1), // No default system texture usable for UAV, so always create but with 1x1 resolution
@@ -663,6 +686,8 @@ bool FPCGGenerateLandscapeTexturesElement::ExecuteInternal(FPCGContext* InContex
 			// Unpack the generated results to simple world aligned textures.
 			FPCGGrassMapUnpackerCS::FParameters* Parameters = GraphBuilder.AllocParameters<FPCGGrassMapUnpackerCS::FParameters>();
 			Parameters->InPackedGrassMaps = GrassMapTextureSRV;
+			Parameters->InHeightIndexMap = GraphBuilder.CreateSRV(IndexMapping.CreateHeightIndexMapBuffer(GraphBuilder));
+			Parameters->InGrassIndexMap = GraphBuilder.CreateSRV(IndexMapping.CreateGrassIndexMapBuffer(GraphBuilder));
 			Parameters->OutUnpackedGrassMaps = GrassMapUAV;
 			Parameters->OutUnpackedHeight = GrassHeightMapUAV;
 			Parameters->InNumTiles = FUintVector2(NumTilesX, NumTilesY);
@@ -671,9 +696,7 @@ bool FPCGGenerateLandscapeTexturesElement::ExecuteInternal(FPCGContext* InContex
 			Parameters->InLandscapeComponentResolution = LandscapeComponentResolution;
 			Parameters->InOutputHeight = bGenerateHeightMap ? 1u : 0u;
 			Parameters->InWindow = FUintVector4(WindowStart.X, WindowStart.Y, WindowEnd.X, WindowEnd.Y);
-
-			// The first 2 channels are reserved for height data. See illustration of packing in PCGGrassMapUnpackerCS.usf.
-			Parameters->InNumGrassMapPasses = FMath::DivideAndRoundUp<int>(NumGrassTypes + 2, 4);
+			Parameters->InNumGlobalGrassTypes = IndexMapping.GetNumGlobalGrassTypes();
 
 			// Initialize to invalid component indices.
 			for (int Index = 0; Index < FPCGGrassMapUnpackerCS::MaxNumLandscapeComponents; ++Index)
@@ -859,4 +882,60 @@ bool FPCGGenerateLandscapeTexturesElement::ExecuteInternal(FPCGContext* InContex
 	return true;
 }
 
+void FPCGGenerateLandscapeTexturesContext::FGrassIndexMapping::SetHeightEntry(int32 InComponentIndex, int32 InFirstExportTileIndex)
+{
+	if (ensure(HeightTileIndexMap.IsValidIndex(InComponentIndex)))
+	{
+		HeightTileIndexMap[InComponentIndex] = InFirstExportTileIndex;
+	}
+}
+
+void FPCGGenerateLandscapeTexturesContext::FGrassIndexMapping::SetGrassMapEntry(int32 InComponentIndex, int32 InGlobalGrassMapIndex, int32 InFirstExportTileIndex, int32 InComponentGrassTypeIndex)
+{
+	const int32 NumChannels = 4; // Grass export produces RGBA texture.
+	const int32 FirstGrassMapChannel = 2; // First two channels used for landscape heights.
+	const int32 GlobalChannelIndex = InFirstExportTileIndex * NumChannels + InComponentGrassTypeIndex + FirstGrassMapChannel;
+	const int32 MapEntryIndex = InComponentIndex * NumGlobalGrassTypes + InGlobalGrassMapIndex;
+
+	// Expected that initialize method is called first.
+	if (ensure(GrassChannelIndexMap.IsValidIndex(MapEntryIndex)))
+	{
+		GrassChannelIndexMap[MapEntryIndex] = GlobalChannelIndex;
+	}
+}
+
+FRDGBufferRef FPCGGenerateLandscapeTexturesContext::FGrassIndexMapping::CreateHeightIndexMapBuffer(FRDGBuilder& GraphBuilder) const
+{
+	if (!HeightTileIndexMap.IsEmpty())
+	{
+		FRDGBufferDesc Desc = FRDGBufferDesc::CreateStructuredDesc(HeightTileIndexMap.GetTypeSize(), HeightTileIndexMap.Num());
+		FRDGBufferRef Buffer = GraphBuilder.CreateBuffer(Desc, TEXT("PCGGenerateLandscapeTexturesHeightIndexMap"));
+
+		GraphBuilder.QueueBufferUpload(Buffer, MakeConstArrayView(HeightTileIndexMap));
+
+		return Buffer;
+	}
+	else
+	{
+		return GSystemTextures.GetDefaultStructuredBuffer(GraphBuilder, sizeof(int32));
+	}
+}
+
+FRDGBufferRef FPCGGenerateLandscapeTexturesContext::FGrassIndexMapping::CreateGrassIndexMapBuffer(FRDGBuilder& GraphBuilder) const
+{
+	if (!GrassChannelIndexMap.IsEmpty())
+	{
+		FRDGBufferDesc Desc = FRDGBufferDesc::CreateStructuredDesc(GrassChannelIndexMap.GetTypeSize(), GrassChannelIndexMap.Num());
+		FRDGBufferRef Buffer = GraphBuilder.CreateBuffer(Desc, TEXT("PCGGenerateLandscapeTexturesGrassIndexMap"));
+
+		GraphBuilder.QueueBufferUpload(Buffer, MakeConstArrayView(GrassChannelIndexMap));
+
+		return Buffer;
+	}
+	else
+	{
+		return GSystemTextures.GetDefaultStructuredBuffer(GraphBuilder, sizeof(int32));
+	}
+}
+
 #undef LOCTEXT_NAMESPACE
diff --git a/Engine/Plugins/PCG/Source/PCG/Public/Elements/PCGGenerateLandscapeTextures.h b/Engine/Plugins/PCG/Source/PCG/Public/Elements/PCGGenerateLandscapeTextures.h
--- a/Engine/Plugins/PCG/Source/PCG/Public/Elements/PCGGenerateLandscapeTextures.h
+++ b/Engine/Plugins/PCG/Source/PCG/Public/Elements/PCGGenerateLandscapeTextures.h
@@ -94,10 +94,42 @@ public:
 	 */
 	TArray<TTuple<TWeakObjectPtr<ULandscapeGrassType>, /*TextureIndex=*/int32>> SelectedGrassTypes;
 
-	/** Total number of grass types used by the landscape component. Includes grass types which were not selected for generation. */
-	int32 NumGrassTypes = 0;
+	TArray<TObjectPtr<ULandscapeGrassType>> GlobalGrassTypes;
+	
+	/** Helper for building index map from global height map / grass map indices to the tile & channel in the data from the grass map exporter. */
+	struct FGrassIndexMapping
+	{
+		void Init(int32 InNumComponents, int32 InNumGlobalGrassTypes)
+		{
+			NumGlobalGrassTypes = InNumGlobalGrassTypes;
+			GrassChannelIndexMap.SetNumUninitialized(InNumComponents * InNumGlobalGrassTypes);
+			FMemory::Memset(GrassChannelIndexMap.GetData(), -1, GrassChannelIndexMap.Num() * GrassChannelIndexMap.GetTypeSize());
 
-	TWeakObjectPtr<ALandscapeProxy> LandscapeProxy = nullptr;
+			HeightTileIndexMap.SetNumZeroed(InNumComponents);
+		}
+
+		int32 GetNumGlobalGrassTypes() const { return NumGlobalGrassTypes; }
+
+		void SetHeightEntry(int32 InComponentIndex, int32 InFirstExportTileIndex);
+		void SetGrassMapEntry(int32 InComponentIndex, int32 InGlobalGrassMapIndex, int32 InFirstExportTileIndex, int32 InComponentGrassTypeIndex);
+
+		FRDGBufferRef CreateHeightIndexMapBuffer(FRDGBuilder& GraphBuilder) const;
+		FRDGBufferRef CreateGrassIndexMapBuffer(FRDGBuilder& GraphBuilder) const;
+
+	private:
+		// For each component, global tile index containing height data. Height data always packed into first 2 channels of tile so don't need to track channel index.
+		TArray<int32> HeightTileIndexMap;
+
+		// For each component, for each possible global grass map index, store the global channel index in the export data.
+		// To get export tile index, divide this value by 4.
+		TArray<int32> GrassChannelIndexMap;
+
+		int32 NumGlobalGrassTypes = -1;
+	};
+
+	FGrassIndexMapping IndexMapping;
+
+	TWeakObjectPtr<ALandscapeProxy> FirstLandscapeProxy = nullptr;
 	TArray<TWeakObjectPtr<ULandscapeComponent>> LandscapeComponents;
 
 	/** World-space bounds containing all of the landscape components given to the grass weight exporter. */
@@ -127,6 +159,13 @@ public:
 
 protected:
 	virtual void AddExtraStructReferencedObjects(FReferenceCollector& Collector) override;
+
+	// Deprecated section
+public:
+	UE_DEPRECATED(5.8, "No longer used")
+	int32 NumGrassTypes = 0;
+	UE_DEPRECATED(5.8, "Renamed to FirstLandscapeProxy")
+	TWeakObjectPtr<ALandscapeProxy> LandscapeProxy = nullptr;
 };
 
 class FPCGGenerateLandscapeTexturesElement : public IPCGElement
diff --git a/Engine/Plugins/PCG/Source/PCGCompute/Internal/PCGGrassMapUnpackerCS.h b/Engine/Plugins/PCG/Source/PCGCompute/Internal/PCGGrassMapUnpackerCS.h
--- a/Engine/Plugins/PCG/Source/PCGCompute/Internal/PCGGrassMapUnpackerCS.h
+++ b/Engine/Plugins/PCG/Source/PCGCompute/Internal/PCGGrassMapUnpackerCS.h
@@ -20,12 +20,14 @@ public:
 
 	BEGIN_SHADER_PARAMETER_STRUCT(FParameters, PCGCOMPUTE_API)
 		SHADER_PARAMETER_RDG_TEXTURE_SRV(Texture2D<float4>, InPackedGrassMaps)
+		SHADER_PARAMETER_RDG_BUFFER_SRV(StructuredBuffer<int32>, InHeightIndexMap)
+		SHADER_PARAMETER_RDG_BUFFER_SRV(StructuredBuffer<int32>, InGrassIndexMap)
 		SHADER_PARAMETER_RDG_TEXTURE_UAV(RWTexture2D<float>, OutUnpackedHeight)
 		SHADER_PARAMETER_RDG_TEXTURE_UAV(RWTexture2DArray<float>, OutUnpackedGrassMaps)
 		SHADER_PARAMETER_ARRAY(FIntVector4, InLinearTileIndexToComponentIndex, [MaxNumLandscapeComponents])
 		SHADER_PARAMETER(FUintVector2, InNumTiles)
 		SHADER_PARAMETER(uint32, InLandscapeComponentResolution)
-		SHADER_PARAMETER(uint32, InNumGrassMapPasses)
+		SHADER_PARAMETER(int32, InNumGlobalGrassTypes)
 		SHADER_PARAMETER(FVector3f, InLandscapeGridScale)
 		SHADER_PARAMETER(float, InLandscapeLocationZ)
 		SHADER_PARAMETER(uint32, InOutputHeight)
