[PCG] Refactor PCG Data Asset pins update

* Make sure the pins are cached on save on the data asset, and also marked as AssetRegistrySearchable, so we can query it without loading the asset
* In Load PCG Data Asset Post Load, update the pins, without loading the asset
* React to a PCG Data Asset save to update the node.
* Save PCG Data Asset will not be cacheable anymore but will not save the asset if nothing has changed.

#rnx
#Synced-CL 48852953

diff --git a/Engine/Plugins/PCG/Source/PCG/Private/Elements/IO/PCGLoadAssetElement.cpp b/Engine/Plugins/PCG/Source/PCG/Private/Elements/IO/PCGLoadAssetElement.cpp
--- a/Engine/Plugins/PCG/Source/PCG/Private/Elements/IO/PCGLoadAssetElement.cpp
+++ b/Engine/Plugins/PCG/Source/PCG/Private/Elements/IO/PCGLoadAssetElement.cpp
@@ -15,6 +15,12 @@
 
 #include "Algo/Transform.h"
 #include "AssetRegistry/AssetData.h"
+#include "AssetRegistry/AssetRegistryModule.h"
+#include "UObject/ObjectSaveContext.h"
+
+#if WITH_EDITOR
+#include "UObject/UObjectIterator.h"
+#endif // WITH_EDITOR
 
 #include UE_INLINE_GENERATED_CPP_BY_NAME(PCGLoadAssetElement)
 
@@ -31,6 +37,15 @@ UPCGLoadDataAssetSettings::UPCGLoadDataAssetSettings()
 	bTagOutputsBasedOnOutputPins = true;
 }
 
+void UPCGLoadDataAssetSettings::PostLoad()
+{
+	Super::PostLoad();
+
+#if WITH_EDITOR
+	UpdateFromData(/*bShouldLoadAsset=*/false);
+#endif // WITH_EDITOR
+}
+
 #if WITH_EDITOR
 void UPCGLoadDataAssetSettings::GetStaticTrackedKeys(FPCGSelectionKeyToSettingsMap& OutKeysToSettings, TArray<TObjectPtr<const UPCGGraph>>& OutVisitedGraphs) const
 {
@@ -59,6 +74,7 @@ EPCGChangeType UPCGLoadDataAssetSettings::GetChangeTypeForProperty(const FName&
 	EPCGChangeType ChangeType = Super::GetChangeTypeForProperty(InPropertyName) | EPCGChangeType::Cosmetic;
 
 	if (InPropertyName == GET_MEMBER_NAME_CHECKED(UPCGLoadDataAssetSettings, bLoadFromInput) ||
+		InPropertyName == GET_MEMBER_NAME_CHECKED(UPCGLoadDataAssetSettings, Asset) ||
 		InPropertyName == GET_MEMBER_NAME_CHECKED(UPCGLoadDataAssetSettings, bSetDefaultAttributeOverridesFromInput))
 	{
 		ChangeType |= EPCGChangeType::Structural;
@@ -117,58 +133,31 @@ FPCGDataTypeIdentifier UPCGLoadDataAssetSettings::GetCurrentPinTypesID(const UPC
 
 void UPCGLoadDataAssetSettings::SetFromAsset(const FAssetData& InAsset)
 {
-	Asset = nullptr;
+	// First verify that the asset if the right class, in the very unlikely possibility that the same path would now point to a different asset class.
+	const UClass* AssetClass = InAsset.GetClass();
+	const bool bIsRightClass = AssetClass && AssetClass->IsChildOf(UPCGDataAsset::StaticClass());
 
-	if (const UClass* AssetClass = InAsset.GetClass())
+	FSoftObjectPath InAssetPath = InAsset.GetSoftObjectPath();
+	if (bIsRightClass && InAssetPath == Asset.ToSoftObjectPath())
 	{
-		if (AssetClass->IsChildOf(UPCGDataAsset::StaticClass()))
-		{
-			Asset = TSoftObjectPtr<UPCGDataAsset>(InAsset.GetSoftObjectPath());
-		}
+		// Nothing to do
+		return;
 	}
 
+	// If it is not the right class, reset.
+	Asset = bIsRightClass ? TSoftObjectPtr<UPCGDataAsset>(MoveTemp(InAssetPath)) : TSoftObjectPtr<UPCGDataAsset>{};
+
 	UpdateFromData();
-	// TODO : notify?
+#if WITH_EDITOR
+	OnSettingsChangedDelegate.Broadcast(this, EPCGChangeType::Structural);
+#endif // WITH_EDITOR
 }
 
-void UPCGLoadDataAssetSettings::UpdateFromData()
+void UPCGLoadDataAssetSettings::UpdateFromData(bool bShouldLoadAsset)
 {
-	// Populate pins based on data present, in order, in the data collection.
-	if (UPCGDataAsset* AssetData = Asset.LoadSynchronous())
-	{
-		TArray<FPCGPinProperties> NewPins;
-		const FPCGDataCollection& Data = AssetData->Data;
+	Modify(/*bAlwaysMarkDirty=*/bShouldLoadAsset);
 
-		for (const FPCGTaggedData& TaggedData : Data.TaggedData)
-		{
-			if (!TaggedData.Data)
-			{
-				continue;
-			}
-
-			FPCGPinProperties* MatchingPin = NewPins.FindByPredicate([&TaggedData](const FPCGPinProperties& PinProperty) { return PinProperty.Label == TaggedData.Pin; });
-
-			if (!MatchingPin)
-			{
-				NewPins.Emplace(TaggedData.Pin, TaggedData.Data->GetDataTypeId());
-			}
-			else
-			{
-				MatchingPin->AllowedTypes |= TaggedData.Data->GetDataTypeId();
-			}
-		}
-
-		Pins = NewPins;
-		bTagOutputsBasedOnOutputPins = false;
-
-		// Update rest of cached data (name, tooltip, color, ...)
-		AssetName = AssetData->Name;
-#if WITH_EDITOR
-		AssetDescription = AssetData->Description;
-		AssetColor = AssetData->Color;
-#endif
-	}
-	else
+	if (Asset.IsNull())
 	{
 		Pins = Super::OutputPinProperties();
 		bTagOutputsBasedOnOutputPins = true;
@@ -177,10 +166,44 @@ void UPCGLoadDataAssetSettings::UpdateFromData()
 #if WITH_EDITOR
 		AssetDescription = FText::GetEmpty();
 		AssetColor = FLinearColor::White;
+#endif // WITH_EDITOR
+	}
+	else
+	{
+		PCGDataAsset::FGetAssetDataOutput AssetRegistryData = UPCGDataAsset::GetAssetRegistryData(Asset, bShouldLoadAsset);
+		
+		bTagOutputsBasedOnOutputPins = false;
+
+		// Update rest of cached data (name, tooltip, color, ...)
+		AssetName = MoveTemp(AssetRegistryData.Name);
+		
+#if WITH_EDITOR
+		// We don't have any info, so we can't update.
+		if (AssetRegistryData.CachedPins.bHasCachedPins)
+		{
+			Pins = MoveTemp(AssetRegistryData.CachedPins.CachedPins);
+		}
+		
+		AssetDescription = MoveTemp(AssetRegistryData.Description);
+		if (AssetRegistryData.Color.IsSet())
+		{
+			AssetColor = MoveTemp(AssetRegistryData.Color.GetValue());
+		}
 #endif
 	}
 }
 
+#if WITH_EDITOR
+void UPCGLoadDataAssetSettings::OnPreObjectSaved(UObject* Object, FObjectPreSaveContext SaveContext)
+{
+	if (Object && Asset == Object)
+	{
+		UpdateFromData();
+		OnSettingsChangedDelegate.Broadcast(this, EPCGChangeType::Structural);
+	}
+}
+#endif // WITH_EDITOR
+
 bool FPCGLoadDataAssetElement::PrepareDataInternal(FPCGContext* InContext) const
 {
 	TRACE_CPUPROFILER_EVENT_SCOPE(FPCGLoadDataAssetElement::PrepareData);
@@ -407,4 +430,27 @@ bool FPCGLoadDataAssetElement::ExecuteInternal(FPCGContext* InContext) const
 	return true;
 }
 
+#if WITH_EDITOR
+FPCGLoadDataAssetChangeTracker::FPCGLoadDataAssetChangeTracker()
+{
+	Handle = FCoreUObjectDelegates::OnObjectPreSave.AddStatic(&FPCGLoadDataAssetChangeTracker::OnPreObjectSaved);
+}
+
+FPCGLoadDataAssetChangeTracker::~FPCGLoadDataAssetChangeTracker()
+{
+	FCoreUObjectDelegates::OnObjectPreSave.Remove(Handle);
+}
+
+void FPCGLoadDataAssetChangeTracker::OnPreObjectSaved(UObject* InObject, FObjectPreSaveContext InSaveContext)
+{
+	if (InObject->IsA<UPCGDataAsset>())
+	{
+		for (TObjectIterator<UPCGLoadDataAssetSettings> It; It; ++It)
+		{
+			It->OnPreObjectSaved(InObject, InSaveContext);
+		}
+	}
+}
+#endif // WITH_EDITOR
+
 #undef LOCTEXT_NAMESPACE
diff --git a/Engine/Plugins/PCG/Source/PCG/Private/Elements/IO/PCGSaveAssetElement.cpp b/Engine/Plugins/PCG/Source/PCG/Private/Elements/IO/PCGSaveAssetElement.cpp
--- a/Engine/Plugins/PCG/Source/PCG/Private/Elements/IO/PCGSaveAssetElement.cpp
+++ b/Engine/Plugins/PCG/Source/PCG/Private/Elements/IO/PCGSaveAssetElement.cpp
@@ -11,6 +11,11 @@
 
 bool UPCGDataCollectionExporter::ExportAsset(const FString& PackageName, UPCGDataAsset* Asset)
 {
+	if (!ShouldSave(Asset))
+	{
+		return false;
+	}
+
 	// Relies on default behavior to duplicate if needed
 	Asset->Data = Data;
 #if WITH_EDITOR
@@ -25,6 +30,17 @@ UPackage* UPCGDataCollectionExporter::UpdateAsset(const FAssetData& PCGAsset)
 	return nullptr;
 }
 
+bool UPCGDataCollectionExporter::ShouldSave(const UPCGDataAsset* Asset)
+{
+	Data.ComputeCrcs(/*bFullDataCrc=*/true);
+	
+#if WITH_EDITOR
+	return Data.DataCrcs != Asset->Data.DataCrcs || AssetColor != Asset->Color || AssetDescription != Asset->Description.ToString();
+#else
+	return Data.DataCrcs != Asset->Data.DataCrcs;
+#endif // WITH_EDITOR
+}
+
 UPCGSaveDataAssetSettings::UPCGSaveDataAssetSettings()
 {
 	Pins = Super::InputPinProperties();
diff --git a/Engine/Plugins/PCG/Source/PCG/Private/PCGDataAsset.cpp b/Engine/Plugins/PCG/Source/PCG/Private/PCGDataAsset.cpp
new file mode 100644
--- /dev/null
+++ b/Engine/Plugins/PCG/Source/PCG/Private/PCGDataAsset.cpp
@@ -0,0 +1,110 @@
+﻿// Copyright Epic Games, Inc. All Rights Reserved.
+
+#include "PCGDataAsset.h"
+
+#include "UObject/ObjectSaveContext.h"
+
+#if WITH_EDITOR
+#include "AssetRegistry/AssetData.h"
+#include "AssetRegistry/AssetRegistryModule.h"
+#endif // WITH_EDITOR
+
+void UPCGDataAsset::PreSave(FObjectPreSaveContext SaveContext)
+{
+	Super::PreSave(SaveContext);
+	
+#if WITH_EDITOR
+	CachedPins.CachedPins = GetPins();
+	CachedPins.bHasCachedPins = true;
+#endif // WITH_EDITOR
+}
+
+TArray<FPCGPinProperties> UPCGDataAsset::GetPins() const
+{
+	TArray<FPCGPinProperties> Pins;
+
+	for (const FPCGTaggedData& TaggedData : Data.TaggedData)
+	{
+		if (!TaggedData.Data)
+		{
+			continue;
+		}
+
+		FPCGPinProperties* MatchingPin = Pins.FindByPredicate([&TaggedData](const FPCGPinProperties& PinProperty) { return PinProperty.Label == TaggedData.Pin; });
+
+		if (!MatchingPin)
+		{
+			Pins.Emplace(TaggedData.Pin, TaggedData.Data->GetDataTypeId());
+		}
+		else
+		{
+			MatchingPin->AllowedTypes |= TaggedData.Data->GetDataTypeId();
+		}
+	}
+
+	return Pins;
+}
+
+PCGDataAsset::FGetAssetDataOutput UPCGDataAsset::GetAssetRegistryData(const TSoftObjectPtr<UPCGDataAsset>& InSoftAsset, const bool bShouldLoadAsset)
+{
+	if (InSoftAsset.IsNull())
+	{
+		return {};
+	}
+	
+	if (UPCGDataAsset* Asset = bShouldLoadAsset ? InSoftAsset.LoadSynchronous() : InSoftAsset.Get())
+	{
+		return PCGDataAsset::FGetAssetDataOutput
+		{
+			.Asset = Asset,
+			.Name = Asset->Name,
+			.CachedPins = {.CachedPins=Asset->GetPins(), .bHasCachedPins=true},
+#if WITH_EDITOR
+			.Description = Asset->Description,
+			.Color = Asset->Color
+#endif // WITH_EDITOR
+		};
+	}
+	
+#if WITH_EDITOR
+	const FAssetRegistryModule& AssetRegistryModule = FModuleManager::LoadModuleChecked<FAssetRegistryModule>(TEXT("AssetRegistry"));
+	const FAssetData AssetData = AssetRegistryModule.Get().GetAssetByObjectPath(InSoftAsset.ToSoftObjectPath());
+
+	if (const UClass* AssetClass = AssetData.GetClass(); AssetClass && AssetClass->IsChildOf(UPCGDataAsset::StaticClass()))
+	{
+		PCGDataAsset::FGetAssetDataOutput Result =
+		{
+			.Asset = nullptr,
+			.Name = AssetData.GetTagValueRef<FString>(GET_MEMBER_NAME_CHECKED(UPCGDataAsset, Name)),
+			.CachedPins = {},
+			.Description = AssetData.GetTagValueRef<FText>(GET_MEMBER_NAME_CHECKED(UPCGDataAsset, Description)),
+			.Color = {},
+		};
+
+		FString TempString;
+		AssetData.GetTagValue<FString>(GET_MEMBER_NAME_CHECKED(UPCGDataAsset, CachedPins), TempString);
+		if (!TempString.IsEmpty())
+		{
+			FPCGDataAssetCachedPins::StaticStruct()->ImportText(*TempString, &Result.CachedPins, nullptr, 0, nullptr, FString{});
+		}
+
+		// Since Color was not an AssetRegistrySearchable before, we'll rely on the bHasCachedPins for deprecation
+		if (Result.CachedPins.bHasCachedPins)
+		{
+			// Initialize it with the default value.
+			Result.Color.Emplace(FLinearColor::White);
+			
+			TempString.Reset();
+			AssetData.GetTagValue<FString>(GET_MEMBER_NAME_CHECKED(UPCGDataAsset, Color), TempString);
+			if (!TempString.IsEmpty())
+			{
+				TBaseStructure<FLinearColor>::Get()->ImportText(*TempString, Result.Color.GetPtrOrNull(), nullptr, 0, nullptr, FString{});
+			}
+		}
+
+		return Result;
+	}
+#endif // WITH_EDITOR
+
+	return {};
+}
\ No newline at end of file
diff --git a/Engine/Plugins/PCG/Source/PCG/Public/Elements/IO/PCGLoadAssetElement.h b/Engine/Plugins/PCG/Source/PCG/Public/Elements/IO/PCGLoadAssetElement.h
--- a/Engine/Plugins/PCG/Source/PCG/Public/Elements/IO/PCGLoadAssetElement.h
+++ b/Engine/Plugins/PCG/Source/PCG/Public/Elements/IO/PCGLoadAssetElement.h
@@ -12,8 +12,6 @@ struct FAssetData;
 
 #include "PCGLoadAssetElement.generated.h"
 
-#define UE_API PCG_API
-
 /** Loader/Executor of PCG data assets */
 UCLASS(MinimalAPI, BlueprintType, ClassGroup = (Procedural))
 class UPCGLoadDataAssetSettings : public UPCGSettings
@@ -21,11 +19,12 @@ class UPCGLoadDataAssetSettings : public UPCGSettings
 	GENERATED_BODY()
 
 public:
-	UE_API UPCGLoadDataAssetSettings();
+	PCG_API UPCGLoadDataAssetSettings();
 
 	//~UObject interface implementation
+	PCG_API virtual void PostLoad() override;
 #if WITH_EDITOR
-	UE_API virtual void PostEditChangeProperty(FPropertyChangedEvent& PropertyChangedEvent) override;
+	PCG_API virtual void PostEditChangeProperty(FPropertyChangedEvent& PropertyChangedEvent) override;
 #endif
 	//~End UObject interface implementation
 
@@ -36,29 +35,33 @@ public:
 	virtual FText GetNodeTooltipText() const override { return AssetDescription; }
 	virtual FLinearColor GetNodeTitleColor() const override { return AssetColor; }
 	virtual EPCGSettingsType GetType() const override { return EPCGSettingsType::InputOutput; }
-	UE_API virtual void GetStaticTrackedKeys(FPCGSelectionKeyToSettingsMap& OutKeysToSettings, TArray<TObjectPtr<const UPCGGraph>>& OutVisitedGraphs) const override;
+	PCG_API virtual void GetStaticTrackedKeys(FPCGSelectionKeyToSettingsMap& OutKeysToSettings, TArray<TObjectPtr<const UPCGGraph>>& OutVisitedGraphs) const override;
 	virtual bool CanDynamicallyTrackKeys() const override { return true; }
 #endif // WITH_EDITOR
 
 	virtual bool HasDynamicPins() const override { return true; }
 	virtual bool HasFlippedTitleLines() const override { return true; }
-	UE_API virtual FString GetAdditionalTitleInformation() const override;
-	UE_API virtual FPCGDataTypeIdentifier GetCurrentPinTypesID(const UPCGPin* InPin) const;
+	PCG_API virtual FString GetAdditionalTitleInformation() const override;
+	PCG_API virtual FPCGDataTypeIdentifier GetCurrentPinTypesID(const UPCGPin* InPin) const;
 
 protected:
-	UE_API virtual FPCGElementPtr CreateElement() const override;
+	PCG_API virtual FPCGElementPtr CreateElement() const override;
 #if WITH_EDITOR
-	UE_API virtual EPCGChangeType GetChangeTypeForProperty(const FName& InPropertyName) const override;
+	PCG_API virtual EPCGChangeType GetChangeTypeForProperty(const FName& InPropertyName) const override;
 #endif
 	// TODO tracking if data is stored in external asset
 
-	UE_API virtual TArray<FPCGPinProperties> InputPinProperties() const override;
+	PCG_API virtual TArray<FPCGPinProperties> InputPinProperties() const override;
 	virtual TArray<FPCGPinProperties> OutputPinProperties() const override { return Pins; }
 	// ~End UPCGSettings interface
 
 public:
-	UE_API virtual void SetFromAsset(const FAssetData& InAsset);
-	UE_API virtual void UpdateFromData();
+	PCG_API virtual void SetFromAsset(const FAssetData& InAsset);
+	PCG_API virtual void UpdateFromData(bool bShouldLoadAsset = true);
+
+#if WITH_EDITOR
+	PCG_API void OnPreObjectSaved(UObject* Object, FObjectPreSaveContext SaveContext);
+#endif // WITH_EDITOR
 
 	UPROPERTY(BlueprintReadWrite, EditAnywhere, Category = Data, meta = (PCG_Overridable))
 	TSoftObjectPtr<UPCGDataAsset> Asset;
@@ -131,8 +134,21 @@ public:
 	virtual bool CanExecuteOnlyOnMainThread(FPCGContext* Context) const override { return true; }
 
 protected:
-	UE_API virtual bool PrepareDataInternal(FPCGContext* Context) const override;
-	UE_API virtual bool ExecuteInternal(FPCGContext* Context) const override;
+	PCG_API virtual bool PrepareDataInternal(FPCGContext* Context) const override;
+	PCG_API virtual bool ExecuteInternal(FPCGContext* Context) const override;
 };
 
-#undef UE_API
+// Extra class to handle when a PCG Data Asset is changed to update all loaded Load PCG Data Asset nodes
+// Will be registered by the PCG Editor.
+#if WITH_EDITOR
+class FPCGLoadDataAssetChangeTracker
+{
+public:
+	PCG_API FPCGLoadDataAssetChangeTracker();
+	PCG_API ~FPCGLoadDataAssetChangeTracker();
+
+	PCG_API static void OnPreObjectSaved(UObject* InObject, FObjectPreSaveContext InSaveContext);
+private:
+	FDelegateHandle Handle;
+};
+#endif // WITH_EDITOR
diff --git a/Engine/Plugins/PCG/Source/PCG/Public/Elements/IO/PCGSaveAssetElement.h b/Engine/Plugins/PCG/Source/PCG/Public/Elements/IO/PCGSaveAssetElement.h
--- a/Engine/Plugins/PCG/Source/PCG/Public/Elements/IO/PCGSaveAssetElement.h
+++ b/Engine/Plugins/PCG/Source/PCG/Public/Elements/IO/PCGSaveAssetElement.h
@@ -11,8 +11,6 @@ struct FAssetData;
 
 #include "PCGSaveAssetElement.generated.h"
 
-#define UE_API PCG_API
-
 /** Default exporter to save data collections, with no capacity for update. */
 UCLASS(MinimalAPI, BlueprintType)
 class UPCGDataCollectionExporter : public UPCGAssetExporter
@@ -21,11 +19,14 @@ class UPCGDataCollectionExporter : public UPCGAssetExporter
 
 protected:
 	//~Begin UPCGAssetExporter interface
-	UE_API virtual bool ExportAsset(const FString& PackageName, UPCGDataAsset* Asset);
-	UE_API virtual UPackage* UpdateAsset(const FAssetData& PCGAsset);
+	PCG_API virtual bool ExportAsset(const FString& PackageName, UPCGDataAsset* Asset);
+	PCG_API virtual UPackage* UpdateAsset(const FAssetData& PCGAsset);
 	//~End UPCGAssetExporter interface
 
 public:
+	// @todo_pcg: Should it be a base class function? 
+	PCG_API bool ShouldSave(const UPCGDataAsset* Asset);
+	
 	UPROPERTY(BlueprintReadWrite, VisibleAnywhere, Category = Data)
 	FPCGDataCollection Data;
 
@@ -45,7 +46,7 @@ class UPCGSaveDataAssetSettings : public UPCGSettings
 	GENERATED_BODY()
 
 public:
-	UE_API UPCGSaveDataAssetSettings();
+	PCG_API UPCGSaveDataAssetSettings();
 
 	//~UPCGSettings interface implementation
 #if WITH_EDITOR
@@ -59,7 +60,7 @@ public:
 	virtual TArray<FPCGPinProperties> OutputPinProperties() const override { return TArray<FPCGPinProperties>(); }
 
 protected:
-	UE_API virtual FPCGElementPtr CreateElement() const override;
+	PCG_API virtual FPCGElementPtr CreateElement() const override;
 	// ~End UPCGSettings interface
 
 public:
@@ -83,10 +84,9 @@ class FPCGSaveDataAssetElement : public IPCGElement
 {
 public:
 	virtual bool CanExecuteOnlyOnMainThread(FPCGContext* Context) const override { return true; }
+	virtual bool IsCacheable(const UPCGSettings* InSettings) const override { return false; }
 
 protected:
-	UE_API virtual bool ExecuteInternal(FPCGContext* Context) const override;
+	PCG_API virtual bool ExecuteInternal(FPCGContext* Context) const override;
 	virtual bool SupportsBasePointDataInputs(FPCGContext* InContext) const { return true; }
 };
-
-#undef UE_API
diff --git a/Engine/Plugins/PCG/Source/PCG/Public/PCGDataAsset.h b/Engine/Plugins/PCG/Source/PCG/Public/PCGDataAsset.h
--- a/Engine/Plugins/PCG/Source/PCG/Public/PCGDataAsset.h
+++ b/Engine/Plugins/PCG/Source/PCG/Public/PCGDataAsset.h
@@ -7,6 +7,8 @@
 
 #include "PCGDataAsset.generated.h"
 
+class UPCGDataAsset;
+
 namespace PCGLevelToAssetConstants
 {
 	const FName MaterialAttributeName = TEXT("Material");
@@ -20,6 +22,36 @@ namespace PCGLevelToAssetConstants
 	const FName ExcludedActorTag = TEXT("PCGLevelToAssetExclude");
 }
 
+USTRUCT(BlueprintInternalUseOnly)
+struct FPCGDataAssetCachedPins
+{
+	GENERATED_BODY();
+	
+	/** Cached pins to be able to update the Load PCG Data Asset node, without loading the asset. */
+	UPROPERTY()
+	TArray<FPCGPinProperties> CachedPins;
+
+	/** For deprecation, to know if the asset has cached pins. */
+	UPROPERTY()
+	bool bHasCachedPins = false;
+};
+
+namespace PCGDataAsset
+{
+	struct FGetAssetDataOutput
+	{
+		// Can be null if the asset is not loaded.
+		UPCGDataAsset* Asset = nullptr;
+		FString Name;
+		FPCGDataAssetCachedPins CachedPins;
+		
+#if WITH_EDITOR
+		FText Description;
+		TOptional<FLinearColor> Color;
+#endif // WITH_EDITOR
+	};
+}
+
 /** Container for PCG data exported as standalone objects */
 UCLASS(MinimalAPI, hidecategories=Object, BlueprintType, ClassGroup=(Procedural))
 class UPCGDataAsset : public UObject
@@ -32,6 +64,14 @@ public:
 	// TODO: add array of references to other assets for recursive parsing
 	// TODO: additional embedded graph?
 
+	//~ Begin UObject interface
+	PCG_API virtual void PreSave(FObjectPreSaveContext SaveContext) override;
+	//~ End UObject interface
+
+	PCG_API TArray<FPCGPinProperties> GetPins() const;
+
+	static PCG_API PCGDataAsset::FGetAssetDataOutput GetAssetRegistryData(const TSoftObjectPtr<UPCGDataAsset>& InSoftAsset, const bool bShouldLoadAsset);
+
 	/** Contains direct data owned by this data asset */
 	UPROPERTY(BlueprintReadWrite, VisibleAnywhere, Category = Data, meta = (NoResetToDefault))
 	FPCGDataCollection Data;
@@ -58,7 +98,7 @@ public:
 	FSoftObjectPath ObjectPath;
 
 	/** Custom node color to use for this specific asset */
-	UPROPERTY(BlueprintReadWrite, EditAnywhere, Category = Data)
+	UPROPERTY(BlueprintReadWrite, EditAnywhere, Category = Data, AssetRegistrySearchable)
 	FLinearColor Color = FLinearColor::White;
 
 	/** Controls whether the asset will be visible from the palette & contextual menus. */
@@ -72,5 +112,10 @@ public:
 	/** Additional description (tooltip) shown on the asset/loader node. */
 	UPROPERTY(BlueprintReadWrite, EditAnywhere, Category = AssetInfo, AssetRegistrySearchable)
 	FText Description;
+
+private:
+	/** Cached data to be able to update the pins without loading the asset. */
+	UPROPERTY(AssetRegistrySearchable)
+	FPCGDataAssetCachedPins CachedPins;
 #endif // WITH_EDITORONLY_DATA
-};
+};
\ No newline at end of file
diff --git a/Engine/Plugins/PCG/Source/PCGEditor/Private/PCGEditorModule.cpp b/Engine/Plugins/PCG/Source/PCGEditor/Private/PCGEditorModule.cpp
--- a/Engine/Plugins/PCG/Source/PCGEditor/Private/PCGEditorModule.cpp
+++ b/Engine/Plugins/PCG/Source/PCGEditor/Private/PCGEditorModule.cpp
@@ -23,6 +23,7 @@
 #include "Data/PCGTextureData.h"
 #include "Data/PCGVirtualTextureData.h"
 #include "Data/PCGVolumeData.h"
+#include "Elements/IO/PCGLoadAssetElement.h"
 #include "Grid/PCGPartitionActor.h"
 #include "WorldPartitionBuilder/PCGWorldPartitionBuilder.h"
 
@@ -109,6 +110,8 @@ void FPCGEditorModule::StartupModule()
 	RegisterPinColorAndIcons();
 	RegisterTrackerFactories();
 
+	LoadDataAssetChangeTracker = MakeUnique<FPCGLoadDataAssetChangeTracker>();
+
 	UToolMenus::RegisterStartupCallback(FSimpleMulticastDelegate::FDelegate::CreateRaw(this, &FPCGEditorModule::RegisterMenuExtensions));
 
 	FPCGEditorStyle::Register();
@@ -151,6 +154,8 @@ void FPCGEditorModule::ShutdownModule()
 
 	IPCGEditorModule::SetEditorModule(nullptr);
 
+	LoadDataAssetChangeTracker.Reset();
+
 	UnregisterTrackerFactories();
 	UnregisterPCGDataVisualizations();
 	UnregisterSettings();
diff --git a/Engine/Plugins/PCG/Source/PCGEditor/Public/PCGEditorModule.h b/Engine/Plugins/PCG/Source/PCGEditor/Public/PCGEditorModule.h
--- a/Engine/Plugins/PCG/Source/PCGEditor/Public/PCGEditorModule.h
+++ b/Engine/Plugins/PCG/Source/PCGEditor/Public/PCGEditorModule.h
@@ -13,6 +13,7 @@
 #include "Modules/ModuleInterface.h"
 #include "Toolkits/AssetEditorToolkit.h"
 
+class FPCGLoadDataAssetChangeTracker;
 class ILevelEditor;
 // Logs
 DECLARE_LOG_CATEGORY_EXTERN(LogPCGEditor, Log, All);
@@ -101,4 +102,6 @@ protected:
 
 	EPCGEditorPermissionMode PermissionMode = EPCGEditorPermissionMode::All;
 	FPermissionModeChanged PermissionModeChangedEvent;
+
+	TUniquePtr<FPCGLoadDataAssetChangeTracker> LoadDataAssetChangeTracker;
 };
