[PCG] Provide more information in the spline sample on Interior to better track which spline is faulty

#rnx
#Synced-CL 48852953

diff --git a/Engine/Plugins/PCG/Source/PCG/Private/Elements/PCGSplineSampler.cpp b/Engine/Plugins/PCG/Source/PCG/Private/Elements/PCGSplineSampler.cpp
--- a/Engine/Plugins/PCG/Source/PCG/Private/Elements/PCGSplineSampler.cpp
+++ b/Engine/Plugins/PCG/Source/PCG/Private/Elements/PCGSplineSampler.cpp
@@ -1052,18 +1052,31 @@ namespace PCGSplineSamplerHelpers
 
 		const FPCGSplineStruct* Spline = nullptr;
 
+		// Add more information in order to detect which spline it might be
+		auto LogError = [Context, LineData](const FText& ErrorMessage)
+		{
+			const UPCGMetadata* Metadata = LineData ? LineData->ConstMetadata() : nullptr;
+			const FPCGMetadataDomain* MetadataDomain = Metadata ? Metadata->GetConstMetadataDomain(PCGMetadataDomainID::Data) : nullptr;
+			const FPCGMetadataAttribute<FSoftObjectPath>* ActorRefefrenceAttribute = MetadataDomain ? MetadataDomain->GetConstTypedAttribute<FSoftObjectPath>(PCGAttributeNameConstants::ActorReferenceAttribute) : nullptr;
+			const FText ActorReferenceHint = ActorRefefrenceAttribute ? FText::Format(LOCTEXT("ActorReferenceHint", "\nPossible actor source: {0}"), FText::FromString(ActorRefefrenceAttribute->GetValue(PCGFirstEntryKey).ToString())) : FText{};
+
+			const FText SplineTransformHint = LineData ? FText::Format(LOCTEXT("TransfromHint", "\nSpline world position: {0}"), FText::FromString(LineData->GetTransform().GetLocation().ToString())) : FText{};
+
+			PCGLog::LogErrorOnGraph(FText::Format(INVTEXT("{0}{1}{2}"), ErrorMessage, ActorReferenceHint, SplineTransformHint), Context);
+		};
+
 		if (const UPCGSplineData* SplineData = Cast<UPCGSplineData>(LineData))
 		{
 			Spline = &SplineData->SplineStruct;
 		}
 		else if (const UPCGLandscapeSplineData* LandscapeSplineData = Cast<UPCGLandscapeSplineData>(LineData))
 		{
-			PCGLog::LogErrorOnGraph(LOCTEXT("LandscapeSplinesNotSupported", "Input data of type Landscape Spline are not supported for interior sampling"), Context);
+			LogError(LOCTEXT("LandscapeSplinesNotSupported", "Input data of type Landscape Spline are not supported for interior sampling"));
 			return;
 		}
 		else
 		{
-			PCGLog::LogErrorOnGraph(LOCTEXT("CouldNotCreateSplineData", "Could not create UPCGSplineData from LineData"), Context);
+			LogError(LOCTEXT("CouldNotCreateSplineData", "Could not create UPCGSplineData from LineData"));
 			return;
 		}
 
@@ -1071,7 +1084,7 @@ namespace PCGSplineSamplerHelpers
 
 		if (!Spline->IsClosedLoop())
 		{
-			PCGLog::LogErrorOnGraph(LOCTEXT("ShapeNotClosed", "Interior sampling only generates for closed shapes, enable the 'Closed Loop' setting on the spline"), Context);
+			LogError(LOCTEXT("ShapeNotClosed", "Interior sampling only generates for closed shapes, enable the 'Closed Loop' setting on the spline"));
 			return;
 		}
 
@@ -1406,7 +1419,7 @@ namespace PCGSplineSamplerHelpers
 
 		if (bAnyRowFailedToSample)
 		{
-			PCGLog::LogErrorOnGraph(LOCTEXT("IntersectionTestFailed", "One or more rows of the spline interior failed to sample (intersection test failed). Ensure the spline points are not overlapping."), Context);
+			LogError(LOCTEXT("IntersectionTestFailed", "One or more rows of the spline interior failed to sample (intersection test failed). Ensure the spline points are not overlapping."));
 		}
 
 		// Finally, gather the data and push to the points
