[PCG] Fix Get Texture Data node output pin appearing as multi-data

#rnx
#Synced-CL 48806643

diff --git a/Engine/Plugins/PCG/Source/PCG/Private/Elements/PCGTextureSampler.cpp b/Engine/Plugins/PCG/Source/PCG/Private/Elements/PCGTextureSampler.cpp
--- a/Engine/Plugins/PCG/Source/PCG/Private/Elements/PCGTextureSampler.cpp
+++ b/Engine/Plugins/PCG/Source/PCG/Private/Elements/PCGTextureSampler.cpp
@@ -82,7 +82,8 @@ FText UPCGTextureSamplerSettings::GetNodeTooltipText() const
 TArray<FPCGPinProperties> UPCGTextureSamplerSettings::OutputPinProperties() const
 {
 	TArray<FPCGPinProperties> Properties;
-	Properties.Emplace(PCGPinConstants::DefaultOutputLabel, EPCGDataType::BaseTexture);
+	FPCGPinProperties& Pin = Properties.Emplace_GetRef(PCGPinConstants::DefaultOutputLabel, EPCGDataType::BaseTexture);
+	Pin.bAllowMultipleData = false;
 
 	return Properties;
 }
