[PCG] Extend Data View to include CSV output and a Data View To String node

#Synced-CL 48851180

diff --git a/Engine/Plugins/PCG/Source/PCG/Private/Data/DataView/PCGDataViewCSVConverter.cpp b/Engine/Plugins/PCG/Source/PCG/Private/Data/DataView/PCGDataViewCSVConverter.cpp
new file mode 100644
--- /dev/null
+++ b/Engine/Plugins/PCG/Source/PCG/Private/Data/DataView/PCGDataViewCSVConverter.cpp
@@ -0,0 +1,111 @@
+// Copyright Epic Games, Inc. All Rights Reserved.
+
+#include "Data/DataView/PCGDataViewCSVConverter.h"
+
+#include "Data/DataView/PCGDataView.h"
+#include "Metadata/Accessors/PCGAttributeAccessorHelpers.h"
+
+#define LOCTEXT_NAMESPACE "PCGDataViewCSVConverter"
+
+namespace PCGDataViewCSVConverter::Helpers
+{
+	TValueOrError<FString, FText> DataViewToString(const FPCGDataView& InDataView, const FInstancedStruct& Parameters)
+	{
+		if (!InDataView.IsValid() || !Parameters.IsValid())
+		{
+			return MakeError(LOCTEXT("SerializeStructInvalidData", "Couldn't serialize to target struct: Data View or Parameters were invalid."));
+		}
+
+		if (!ensure(Parameters.GetScriptStruct() == FPCGDataViewCSVParameters::StaticStruct()))
+		{
+			return MakeError(LOCTEXT("ParamStructTypesMismatch", "Struct types between parameters and FPCGDataViewCSVParameters."));
+		}
+
+		const TArray<FPCGAttributePropertySelector> Selection = InDataView.GetResolvedSelection();
+		TStringBuilder<128> StringBuilder;
+
+		const FPCGDataViewCSVParameters& CSVParams = Parameters.Get<FPCGDataViewCSVParameters>();
+		if (!Selection.IsEmpty())
+		{
+			const UPCGData* ViewedData = InDataView.ViewedData;
+
+			if (!ensure(ViewedData))
+			{
+				return MakeError(LOCTEXT("InvalidViewedData", "Viewed data was invalid"));
+			}
+
+			auto AddAttributeToCSV = [&StringBuilder, &CSVParams, ViewedData](const FPCGAttributePropertyInputSelector& InSelector)
+			{
+				if (CSVParams.Output == EPCGDataViewCSVOutput::Name || CSVParams.Output == EPCGDataViewCSVOutput::Both)
+				{
+					StringBuilder += InSelector.ToString(CSVParams.bOmitDomain);
+				}
+
+				if (CSVParams.Output == EPCGDataViewCSVOutput::Both)
+				{
+					StringBuilder += CSVParams.NameValueSeparator;
+				}
+
+				if (CSVParams.Output == EPCGDataViewCSVOutput::Value || CSVParams.Output == EPCGDataViewCSVOutput::Both)
+				{
+					TArray<FString> Values;
+					PCGAttributeAccessorHelpers::ExtractAllValues(ViewedData, InSelector, Values, nullptr, EPCGAttributeAccessorFlags::AllowBroadcast);
+					StringBuilder.Join(Values, *CSVParams.ValueDelimiter);
+				}
+
+				StringBuilder += CSVParams.AttributeDelimiter;
+			};
+
+			FPCGMetadataDomainID DefaultDomainID = InDataView.ViewedData->GetDefaultMetadataDomainID();
+			for (FPCGMetadataDomainID SupportedDomainID : ViewedData->GetAllSupportedMetadataDomainIDs())
+			{
+				const bool bIsDefaultDomain = DefaultDomainID == SupportedDomainID;
+				const TArray<FPCGAttributePropertySelector> DomainAttributes = Selection.FilterByPredicate([ViewedData, &SupportedDomainID, bIsDefaultDomain](const FPCGAttributePropertySelector& Selector)
+				{
+					const FPCGMetadataDomainID DomainID = ViewedData->GetMetadataDomainIDFromSelector(Selector);
+					return (!bIsDefaultDomain && DomainID == SupportedDomainID) || (bIsDefaultDomain && DomainID == PCGMetadataDomainID::Default);
+				});
+
+				for (const FPCGAttributePropertySelector& DomainSelector : DomainAttributes)
+				{
+					const FPCGAttributePropertyInputSelector Selector = FPCGAttributePropertyInputSelector::CreateFromOtherSelector<FPCGAttributePropertyInputSelector>(DomainSelector);
+					AddAttributeToCSV(Selector);
+				}
+			}
+		}
+
+		// Remove the extra delimiter, which should always exist
+		if (StringBuilder.Len() > 0 && !CSVParams.AttributeDelimiter.IsEmpty())
+		{
+			StringBuilder.RemoveSuffix(CSVParams.AttributeDelimiter.Len());
+		}
+
+		return MakeValue(StringBuilder.ToString());
+	}
+} // namespace PCGDataViewCSVConverter::Helpers
+
+TValueOrError<FInstancedStruct, FText> UPCGDataViewCSVConverter::SerializeToTargetStruct(const FPCGDataView& InDataView, const FInstancedStruct& Parameters) const
+{
+	TValueOrError<FString, FText> Result = PCGDataViewCSVConverter::Helpers::DataViewToString(InDataView, Parameters);
+	if (Result.HasError())
+	{
+		return MakeError(Result.StealError());
+	}
+
+	FInstancedStruct InstancedStruct;
+	InstancedStruct.InitializeAs<FPCGDataViewCSVOutput>(Result.StealValue());
+	return MakeValue(MoveTemp(InstancedStruct));
+}
+
+TValueOrError<FString, FText> UPCGDataViewCSVConverter::SerializeToString(const FPCGDataView& InDataView, const FInstancedStruct& Parameters) const
+{
+	return PCGDataViewCSVConverter::Helpers::DataViewToString(InDataView, Parameters);
+}
+
+TValueOrError<UPCGData*, FText> UPCGDataViewCSVConverter::ConstructDataFromTarget(const FInstancedStruct& InTargetStruct, UObject* Outer) const
+{
+	/** @todo_pcg: Construct a PCG Data from a CSV string. Revisit serialization options. */
+	return MakeError(LOCTEXT("DataConstructionNotYetImplemented", "Data construction from CSV not yet implemented."));
+}
+
+#undef LOCTEXT_NAMESPACE
diff --git a/Engine/Plugins/PCG/Source/PCG/Private/Data/DataView/PCGDataViewData.cpp b/Engine/Plugins/PCG/Source/PCG/Private/Data/DataView/PCGDataViewData.cpp
--- a/Engine/Plugins/PCG/Source/PCG/Private/Data/DataView/PCGDataViewData.cpp
+++ b/Engine/Plugins/PCG/Source/PCG/Private/Data/DataView/PCGDataViewData.cpp
@@ -4,9 +4,10 @@
 
 #include "PCGParamData.h"
 #include "Data/PCGSplineData.h"
-#include "Data/DataView/PCGDataViewNativeJsonConverters.h"
 #include "Elements/PCGConvertToDataView.h"
 
+#include "Serialization/ArchiveCrc32.h"
+
 #include UE_INLINE_GENERATED_CPP_BY_NAME(PCGDataViewData)
 
 #define LOCTEXT_NAMESPACE "PCGDataViewData"
@@ -114,6 +115,42 @@ const UPCGMetadata* UPCGDataViewData::ConstMetadata() const
 	return DataView.ViewedData->ConstMetadata();
 }
 
+FPCGCrc UPCGDataViewData::ComputeCrc(bool bFullDataCrc) const
+{
+	FArchiveCrc32 Ar;
+
+	AddToCrc(Ar, bFullDataCrc);
+
+	if (DataView.ViewedData)
+	{
+		Ar << DataView.ViewedData->GetOrComputeCrc(bFullDataCrc);
+	}
+
+	return FPCGCrc(Ar.GetCrc());
+}
+
+void UPCGDataViewData::AddToCrc(FArchiveCrc32& Ar, const bool bFullDataCrc) const
+{
+	Super::AddToCrc(Ar, bFullDataCrc);
+
+	if (!bFullDataCrc)
+	{
+		// Fallback to UID
+		AddUIDToCrc(Ar);
+		return;
+	}
+
+	uint8 AllAttributesByte = DataView.Selection.bAllAttributes ? 1 : 0;
+	Ar << AllAttributesByte;
+
+	for (FPCGAttributePropertySelector Selector : DataView.Selection.Attributes)
+	{
+		Ar << Selector;
+	}
+
+	// @todo_pcg: Add bLimitByDomain, Domain, and bIgnoreProperties when they are live.
+}
+
 TArray<FPCGAttributePropertySelector> UPCGDataViewData::GetSelectedAttributes() const
 {
 	if (!ensure(DataView.ViewedData))
diff --git a/Engine/Plugins/PCG/Source/PCG/Private/Elements/IO/PCGDataViewToStringElement.cpp b/Engine/Plugins/PCG/Source/PCG/Private/Elements/IO/PCGDataViewToStringElement.cpp
new file mode 100644
--- /dev/null
+++ b/Engine/Plugins/PCG/Source/PCG/Private/Elements/IO/PCGDataViewToStringElement.cpp
@@ -0,0 +1,186 @@
+// Copyright Epic Games, Inc. All Rights Reserved.
+
+#include "Elements/IO/PCGDataViewToStringElement.h"
+
+#include "Data/DataView/PCGDataViewCSVConverter.h"
+#include "Data/DataView/PCGDataViewData.h"
+#include "Helpers/PCGSettingsHelpers.h"
+
+#include UE_INLINE_GENERATED_CPP_BY_NAME(PCGDataViewToStringElement)
+
+#define LOCTEXT_NAMESPACE "PCGDataViewToStringElement"
+
+UPCGDataViewToStringSettings::UPCGDataViewToStringSettings()
+{
+	ConverterParameters.InitializeAs(FPCGDataViewCSVParameters::StaticStruct());
+}
+
+#if WITH_EDITOR
+EPCGChangeType UPCGDataViewToStringSettings::GetChangeTypeForProperty(const FName& InPropertyName) const
+{
+	EPCGChangeType ChangeType = Super::GetChangeTypeForProperty(InPropertyName);
+
+	if (InPropertyName == GET_MEMBER_NAME_CHECKED(UPCGDataViewToStringSettings, ConverterClass))
+	{
+		ChangeType |= EPCGChangeType::Structural;
+	}
+
+	return ChangeType;
+}
+
+void UPCGDataViewToStringSettings::PostEditChangeProperty(FPropertyChangedEvent& PropertyChangedEvent)
+{
+	Super::PostEditChangeProperty(PropertyChangedEvent);
+
+	if (PropertyChangedEvent.GetPropertyName() == GET_MEMBER_NAME_CHECKED(UPCGDataViewToStringSettings, ConverterClass))
+	{
+		if (const UPCGDataViewConverterBase* ConverterCDO = ConverterClass ? ConverterClass->GetDefaultObject<UPCGDataViewConverterBase>() : nullptr)
+		{
+			ConverterParameters.InitializeAs(ConverterCDO->GetParameterStruct());
+		}
+		else
+		{
+			ConverterParameters.Reset();
+		}
+	}
+}
+#endif // WITH_EDITOR
+
+TArray<FPCGPinProperties> UPCGDataViewToStringSettings::InputPinProperties() const
+{
+	TArray<FPCGPinProperties> PinProperties;
+	PinProperties.Emplace_GetRef(PCGPinConstants::DefaultInputLabel, FPCGDataTypeInfoDataView::AsId()).SetRequiredPin();
+	return PinProperties;
+}
+
+TArray<FPCGPinProperties> UPCGDataViewToStringSettings::OutputPinProperties() const
+{
+	TArray<FPCGPinProperties> PinProperties;
+	FPCGDataTypeIdentifier StringType = EPCGDataType::Param;
+	StringType.CustomSubtype = static_cast<int32>(EPCGMetadataTypes::String);
+	PinProperties.Emplace(PCGPinConstants::DefaultOutputLabel, StringType);
+	return PinProperties;
+}
+
+FPCGElementPtr UPCGDataViewToStringSettings::CreateElement() const
+{
+	return MakeShared<FPCGDataViewToStringElement>();
+}
+
+#if WITH_EDITOR
+TArray<FPCGSettingsOverridableParam> UPCGDataViewToStringSettings::GatherOverridableParams() const
+{
+	TArray<FPCGSettingsOverridableParam> OverridableParams = Super::GatherOverridableParams();
+
+	if (ConverterParameters.IsValid())
+	{
+		PCGSettingsHelpers::FPCGGetAllOverridableParamsConfig Config{};
+		Config.bExtractArrays = true;
+		return PCGSettingsHelpers::GetAllOverridableParams(ConverterParameters.GetScriptStruct(), Config);
+	}
+
+	return {};
+}
+#endif // WITH_EDITOR
+
+void UPCGDataViewToStringSettings::FixingOverridableParamPropertyClass(FPCGSettingsOverridableParam& Param) const
+{
+	bool bFound = false;
+
+	if (!Param.PropertiesNames.IsEmpty() && !Param.Properties.IsEmpty())
+	{
+		const UScriptStruct* ScriptStruct = ConverterParameters.GetScriptStruct();
+		const FProperty* Property = ScriptStruct ? ScriptStruct->FindPropertyByName(Param.PropertiesNames[0]) : nullptr;
+		if (Property && ensure(Param.Properties[0]->GetOwnerStruct() == ScriptStruct))
+		{
+			Param.PropertyClass = ScriptStruct;
+			bFound = true;
+		}
+	}
+
+	if (!bFound)
+	{
+		Super::FixingOverridableParamPropertyClass(Param);
+	}
+}
+
+FStructView FPCGDataViewToStringExecutionContext::GetExternalStructContainerForOverridableParam(const FPCGSettingsOverridableParam& InParam)
+{
+	UPCGDataViewToStringSettings* Settings = GetMutableInputSettings<UPCGDataViewToStringSettings>();
+	if (Settings && Settings->ConverterParameters.IsValid() && !InParam.PropertiesNames.IsEmpty() && Settings->ConverterParameters.GetScriptStruct()->FindPropertyByName(InParam.PropertiesNames[0]))
+	{
+		return Settings->ConverterParameters;
+	}
+	else
+	{
+		return nullptr;
+	}
+}
+
+bool FPCGDataViewToStringElement::ExecuteInternal(FPCGContext* InContext) const
+{
+	TRACE_CPUPROFILER_EVENT_SCOPE(FPCGDataViewToStringElement::Execute);
+
+	const UPCGDataViewToStringSettings* Settings = InContext->GetInputSettings<UPCGDataViewToStringSettings>();
+	check(Settings);
+
+	const TArray<FPCGTaggedData> Inputs = InContext->InputData.GetInputsByPin(PCGPinConstants::DefaultInputLabel);
+	if (Inputs.IsEmpty())
+	{
+		return true;
+	}
+
+	if (!Settings->ConverterClass)
+	{
+		PCGLog::LogErrorOnGraph(LOCTEXT("InvalidConverterClass", "Converter class is invalid."), InContext);
+		return true;
+	}
+
+	for (const FPCGTaggedData& Input : Inputs)
+	{
+		const UPCGDataViewData* DataViewData = Cast<UPCGDataViewData>(Input.Data);
+		if (!DataViewData)
+		{
+			PCGLog::InputOutput::LogInvalidInputDataError(InContext);
+			return true;
+		}
+
+		const FPCGDataView& DataView = DataViewData->GetDataView();
+		if (!DataView.IsValid())
+		{
+			PCGLog::LogWarningOnGraph(LOCTEXT("InvalidDataView", "Data View selection is empty or internal data is no longer valid."));
+			return true;
+		}
+
+		const UPCGDataViewConverterBase* ConverterCDO = Settings->ConverterClass->GetDefaultObject<UPCGDataViewConverterBase>();
+		TValueOrError<FString, FText> Result = ConverterCDO->SerializeToString(DataView, Settings->ConverterParameters);
+		if (Result.HasError())
+		{
+			PCGLog::LogErrorOnGraph(Result.GetError(), InContext);
+			return true;
+		}
+
+		UPCGParamData* ParamData = FPCGContext::NewObject_AnyThread<UPCGParamData>(InContext);
+		FPCGMetadataAttribute<FString>* Attribute = ParamData->Metadata->CreateAttribute<FString>(DataView.ViewedData->GetFName(), Result.GetValue(), /*bAllowsInterpolation=*/false, /*bOverrideParent=*/false);
+		check(Attribute);
+		ParamData->Metadata->AddEntry();
+
+		FPCGTaggedData& Out = InContext->OutputData.TaggedData.Emplace_GetRef(Input);
+		Out.Data = ParamData;
+	}
+
+	return true;
+}
+
+void FPCGDataViewToStringElement::GetDependenciesCrc(const FPCGGetDependenciesCrcParams& InParams, FPCGCrc& OutCrc) const
+{
+	FPCGCrc Crc;
+	IPCGElement::GetDependenciesCrc(InParams, Crc);
+
+	const UPCGDataViewToStringSettings* Settings = CastChecked<UPCGDataViewToStringSettings>(InParams.Settings);
+	Crc.Combine(Settings->GetOrComputeCrc(/*bFullDataCrc=*/false));
+
+	OutCrc = Crc;
+}
+
+#undef LOCTEXT_NAMESPACE
diff --git a/Engine/Plugins/PCG/Source/PCG/Private/PCGModule.cpp b/Engine/Plugins/PCG/Source/PCG/Private/PCGModule.cpp
--- a/Engine/Plugins/PCG/Source/PCG/Private/PCGModule.cpp
+++ b/Engine/Plugins/PCG/Source/PCG/Private/PCGModule.cpp
@@ -14,6 +14,7 @@
 #include "Data/PCGPolyLineData.h"
 #include "Data/PCGSplineData.h"
 #include "Data/PCGPolygon2DData.h"
+#include "Data/DataView/PCGDataViewCSVConverter.h"
 #include "Data/DataView/PCGDataViewInterface.h"
 #include "Data/DataView/PCGDataViewNativeJsonConverters.h"
 #include "Data/DataView/PCGDataViewNativePropertySelectors.h"
@@ -189,6 +190,7 @@ void FPCGModule::StartupModule()
 
 	DataViewRegistry.RegisterConverter(UPCGParamData::StaticClass(), UPCGDataViewJsonConverter::StaticClass());
 	DataViewRegistry.RegisterConverter(UPCGSpatialData::StaticClass(), UPCGDataViewJsonConverter::StaticClass());
+	DataViewRegistry.RegisterConverter(UPCGData::StaticClass(), UPCGDataViewCSVConverter::StaticClass());
 
 	ChangeTrackingRegistry.RegisterGetExecutionSourcesFromSelectionKey(AActor::StaticClass(), FPCGGetExecutionSourcesFromSelectionKey::CreateStatic(&UPCGActorHelpers::GetExecutionSourcesFromSelectionKey));
 	
@@ -217,6 +219,7 @@ void FPCGModule::ShutdownModule()
 
 void FPCGModule::PreExit()
 {
+	DataViewRegistry.UnregisterConverter(UPCGData::StaticClass(), UPCGDataViewCSVConverter::StaticClass());
 	DataViewRegistry.UnregisterConverter(UPCGSpatialData::StaticClass(), UPCGDataViewJsonConverter::StaticClass());
 	DataViewRegistry.UnregisterConverter(UPCGParamData::StaticClass(), UPCGDataViewJsonConverter::StaticClass());
 
diff --git a/Engine/Plugins/PCG/Source/PCG/Public/Data/DataView/PCGDataViewCSVConverter.h b/Engine/Plugins/PCG/Source/PCG/Public/Data/DataView/PCGDataViewCSVConverter.h
new file mode 100644
--- /dev/null
+++ b/Engine/Plugins/PCG/Source/PCG/Public/Data/DataView/PCGDataViewCSVConverter.h
@@ -0,0 +1,64 @@
+// Copyright Epic Games, Inc. All Rights Reserved.
+
+#pragma once
+
+#include "Data/DataView/PCGDataViewInterface.h"
+
+#include "PCGDataViewCSVConverter.generated.h"
+
+class FJsonObject;
+
+UENUM()
+enum class EPCGDataViewCSVOutput : uint8
+{
+	Name,
+	Value,
+	Both
+};
+
+USTRUCT()
+struct FPCGDataViewCSVParameters
+{
+	GENERATED_BODY()
+
+	// The comma-separated format, i.e. 'name' or 'value' or 'name:value'.
+	UPROPERTY(EditAnywhere, Category = "CSV")
+	EPCGDataViewCSVOutput Output = EPCGDataViewCSVOutput::Both;
+
+	// The separating character(s) between the name and the value.
+	UPROPERTY(EditAnywhere, Category = "CSV", meta = (EditCondition = "Output == EPCGDataViewCSVOutput::Both", EditConditionHides))
+	FString NameValueSeparator = ":";
+
+	// The character(s) between each value. By default, as a CSV, this will be the ',' character.
+	UPROPERTY(EditAnywhere, Category = "CSV", meta = (InlineEditConditionToggle, EditCondition = "Output != EPCGDataViewCSVOutput::Name", EditConditionHides, MultiLine))
+	FString ValueDelimiter = ",";
+
+	// The character(s) between each attribute of the selection. By default--and via reset--this will be a platform-specific line terminator.
+	UPROPERTY(EditAnywhere, Category = "CSV")
+	FString AttributeDelimiter = LINE_TERMINATOR;
+
+	// Omit the domain from the name of the attribute.
+	UPROPERTY(EditAnywhere, Category = "CSV")
+	bool bOmitDomain = true;
+};
+
+USTRUCT()
+struct FPCGDataViewCSVOutput
+{
+	GENERATED_BODY()
+
+	FString CSVFormattedOutput;
+};
+
+UCLASS()
+class UPCGDataViewCSVConverter : public UPCGDataViewConverterBase
+{
+	GENERATED_BODY()
+
+public:
+	virtual TValueOrError<FInstancedStruct, FText> SerializeToTargetStruct(const FPCGDataView& InDataView, const FInstancedStruct& Parameters) const override;
+	virtual TValueOrError<FString, FText> SerializeToString(const FPCGDataView& InDataView, const FInstancedStruct& Parameters) const override;
+	virtual TValueOrError<UPCGData*, FText> ConstructDataFromTarget(const FInstancedStruct& InTargetStruct, UObject* Outer = nullptr) const override;
+	virtual UScriptStruct* GetTargetStruct() const override { return FPCGDataViewCSVOutput::StaticStruct(); }
+	virtual UScriptStruct* GetParameterStruct() const override { return FPCGDataViewCSVParameters::StaticStruct(); }
+};
diff --git a/Engine/Plugins/PCG/Source/PCG/Public/Data/DataView/PCGDataViewData.h b/Engine/Plugins/PCG/Source/PCG/Public/Data/DataView/PCGDataViewData.h
--- a/Engine/Plugins/PCG/Source/PCG/Public/Data/DataView/PCGDataViewData.h
+++ b/Engine/Plugins/PCG/Source/PCG/Public/Data/DataView/PCGDataViewData.h
@@ -61,13 +61,18 @@ public:
 	virtual TArray<FPCGMetadataDomainID> GetAllSupportedMetadataDomainIDs() const override;
 	// Forwarding helper to the underlying metadata.
 	virtual const UPCGMetadata* ConstMetadata() const override;
+
+protected:
+	// Compute the Crc for the Data View and its underlying data.
+	virtual FPCGCrc ComputeCrc(bool bFullDataCrc) const override;
+	// Serialize the Data View selection components into the Crc.
+	virtual void AddToCrc(FArchiveCrc32& Ar, bool bFullDataCrc) const override;
 	//~ End UPCGData interface
 
+public:
 	// The selection is stored lazily. Get Selected Attributes will create the array of selectors. @todo_pcg: Look into caching results.
 	TArray<FPCGAttributePropertySelector> GetSelectedAttributes() const;
 
-	// @todo_pcg: Serialize to/from Archive (Binary)
-
 private:
 	UPROPERTY()
 	FPCGDataView DataView;
diff --git a/Engine/Plugins/PCG/Source/PCG/Public/Elements/IO/PCGDataViewToStringElement.h b/Engine/Plugins/PCG/Source/PCG/Public/Elements/IO/PCGDataViewToStringElement.h
new file mode 100644
--- /dev/null
+++ b/Engine/Plugins/PCG/Source/PCG/Public/Elements/IO/PCGDataViewToStringElement.h
@@ -0,0 +1,63 @@
+// Copyright Epic Games, Inc. All Rights Reserved.
+
+#pragma once
+
+#include "PCGContext.h"
+#include "PCGSettings.h"
+#include "Data/DataView/PCGDataView.h"
+#include "Data/DataView/PCGDataViewCSVConverter.h"
+
+#include "PCGDataViewToStringElement.generated.h"
+
+/** Serialize a PCG Data via a Data View into a formatted FString, i.e. CSV, Json, etc. */
+UCLASS(BlueprintType, MinimalAPI, ClassGroup = (Procedural))
+class UPCGDataViewToStringSettings : public UPCGSettings
+{
+	GENERATED_BODY()
+
+public:
+	UPCGDataViewToStringSettings();
+
+	//~ Begin UPCGSettings interface
+#if WITH_EDITOR
+	virtual FName GetDefaultNodeName() const override { return FName(TEXT("PCGDataViewToString")); }
+	virtual FText GetDefaultNodeTitle() const override { return NSLOCTEXT("PCGDataViewToStringElement", "NodeTitle", "Data View To String"); }
+	virtual FText GetNodeTooltipText() const override { return NSLOCTEXT("PCGDataViewToStringElement", "NodeTooltip", "Convert the Data View selection to a string attribute."); }
+	virtual EPCGSettingsType GetType() const override { return EPCGSettingsType::Generic; }
+	virtual EPCGChangeType GetChangeTypeForProperty(const FName& InPropertyName) const override;
+	virtual void PostEditChangeProperty(FPropertyChangedEvent& PropertyChangedEvent) override;
+#endif // WITH_EDITOR
+
+protected:
+	virtual TArray<FPCGPinProperties> InputPinProperties() const override;
+	virtual TArray<FPCGPinProperties> OutputPinProperties() const override;
+
+	virtual FPCGElementPtr CreateElement() const override;
+#if WITH_EDITOR
+	virtual TArray<FPCGSettingsOverridableParam> GatherOverridableParams() const override;
+#endif // WITH_EDITOR
+	virtual void FixingOverridableParamPropertyClass(FPCGSettingsOverridableParam& Param) const override;
+	//~ End UPCGSettings interface
+
+public:
+	// The converter format to construct from the Data View
+	UPROPERTY(EditAnywhere, BlueprintReadWrite, Category = Settings, NoClear)
+	TSubclassOf<UPCGDataViewConverterBase> ConverterClass = UPCGDataViewCSVConverter::StaticClass();
+
+	// Parameters for the conversion
+	UPROPERTY(EditAnywhere, BlueprintReadWrite, Category = Settings, meta = (ShowOnlyInnerProperties, StructTypeConst))
+	FInstancedStruct ConverterParameters;
+};
+
+class FPCGDataViewToStringExecutionContext : public FPCGContext
+{
+protected:
+	virtual FStructView GetExternalStructContainerForOverridableParam(const FPCGSettingsOverridableParam& InParam) override;
+};
+
+class FPCGDataViewToStringElement final : public IPCGElementWithCustomContext<FPCGDataViewToStringExecutionContext>
+{
+protected:
+	virtual bool ExecuteInternal(FPCGContext* InContext) const override;
+	virtual void GetDependenciesCrc(const FPCGGetDependenciesCrcParams& InParams, FPCGCrc& OutCrc) const override;
+};
