[PCG] Transition Data View attribute layout to consumer elements and add selection by domain

#rnx

#Synced-CL 48851180

diff --git a/Engine/Plugins/PCG/Source/PCG/Private/Data/DataView/PCGDataView.cpp b/Engine/Plugins/PCG/Source/PCG/Private/Data/DataView/PCGDataView.cpp
--- a/Engine/Plugins/PCG/Source/PCG/Private/Data/DataView/PCGDataView.cpp
+++ b/Engine/Plugins/PCG/Source/PCG/Private/Data/DataView/PCGDataView.cpp
@@ -20,7 +20,16 @@ TArray<FPCGMetadataDomainID> FPCGDataView::GetAllMetadataDomainIDs() const
 {
 	if (const UPCGData* Data = Cast<UPCGData>(ViewedData))
 	{
-		if (Selection.bAllAttributes)
+		if (Selection.bLimitByDomain)
+		{
+			const FPCGAttributePropertySelector DummySelector = FPCGAttributePropertySelector::CreateAttributeSelector(FName{}, Selection.Domain);
+			FPCGMetadataDomainID DomainID = Data->GetMetadataDomainIDFromSelector(DummySelector);
+			if (Data->IsSupportedMetadataDomainID(DomainID))
+			{
+				return {DomainID};
+			}
+		}
+		else if (Selection.bAllAttributes)
 		{
 			return Data->GetAllSupportedMetadataDomainIDs();
 		}
@@ -49,14 +58,31 @@ TArray<FPCGAttributePropertySelector> FPCGDataView::GetResolvedSelection() const
 	{
 		TArray<FPCGAttributePropertySelector> SelectedAttributes;
 
-		// Collect all properties defined by the data type
-		const FPCGDataViewRegistry& DataViewRegistry = FPCGModule::GetConstPCGDataViewRegistry();
-		if (const IPCGDataViewPropertySelector* PropertySelection = DataViewRegistry.GetPropertySelector(*this))
+		if (!Selection.bIgnoreProperties)
+		{
+			// Collect all properties defined by the data type
+			const FPCGDataViewRegistry& DataViewRegistry = FPCGModule::GetConstPCGDataViewRegistry();
+			if (const IPCGDataViewPropertySelector* PropertySelection = DataViewRegistry.GetPropertySelector(*this))
+			{
+				SelectedAttributes.Append(PropertySelection->GetSelection(*this));
+			}
+		}
+
+		TOptional<FPCGMetadataDomainID> OptionalDomain;
+
+		if (Selection.bLimitByDomain)
 		{
-			SelectedAttributes.Append(PropertySelection->GetSelection(*this));
+			const FPCGAttributePropertySelector DummySelector = FPCGAttributePropertySelector::CreateAttributeSelector(FName{}, Selection.Domain);
+			OptionalDomain = ViewedData->GetMetadataDomainIDFromSelector(DummySelector);
+
+			SelectedAttributes.RemoveAll([this, &OptionalDomain](const FPCGAttributePropertySelector& InSelector)
+			{
+				const FPCGMetadataDomainID Domain = ViewedData->GetMetadataDomainIDFromSelector(InSelector);
+				return !OptionalDomain || Domain != OptionalDomain.GetValue();
+			});
 		}
 
-		PCGDataView::Helpers::AppendAllAttributeSelectors(ViewedData, SelectedAttributes);
+		PCGDataView::Helpers::AppendAllAttributeSelectors(ViewedData, SelectedAttributes, OptionalDomain);
 
 		return SelectedAttributes;
 	}
diff --git a/Engine/Plugins/PCG/Source/PCG/Private/Data/DataView/PCGDataViewCSVConverter.cpp b/Engine/Plugins/PCG/Source/PCG/Private/Data/DataView/PCGDataViewCSVConverter.cpp
--- a/Engine/Plugins/PCG/Source/PCG/Private/Data/DataView/PCGDataViewCSVConverter.cpp
+++ b/Engine/Plugins/PCG/Source/PCG/Private/Data/DataView/PCGDataViewCSVConverter.cpp
@@ -63,7 +63,7 @@ namespace PCGDataViewCSVConverter::Helpers
 				const TArray<FPCGAttributePropertySelector> DomainAttributes = Selection.FilterByPredicate([ViewedData, &SupportedDomainID, bIsDefaultDomain](const FPCGAttributePropertySelector& Selector)
 				{
 					const FPCGMetadataDomainID DomainID = ViewedData->GetMetadataDomainIDFromSelector(Selector);
-					return (!bIsDefaultDomain && DomainID == SupportedDomainID) || (bIsDefaultDomain && DomainID == PCGMetadataDomainID::Default);
+					return (DomainID == SupportedDomainID) || (bIsDefaultDomain && DomainID == PCGMetadataDomainID::Default);
 				});
 
 				for (const FPCGAttributePropertySelector& DomainSelector : DomainAttributes)
diff --git a/Engine/Plugins/PCG/Source/PCG/Private/Data/DataView/PCGDataViewNativeJsonConverters.cpp b/Engine/Plugins/PCG/Source/PCG/Private/Data/DataView/PCGDataViewNativeJsonConverters.cpp
--- a/Engine/Plugins/PCG/Source/PCG/Private/Data/DataView/PCGDataViewNativeJsonConverters.cpp
+++ b/Engine/Plugins/PCG/Source/PCG/Private/Data/DataView/PCGDataViewNativeJsonConverters.cpp
@@ -10,6 +10,15 @@
 TValueOrError<FInstancedStruct, FText> UPCGDataViewJsonConverter::SerializeToTargetStruct(const FPCGDataView& InDataView, const FInstancedStruct& Parameters) const
 {
 	using namespace PCG::IO;
+	if (!InDataView.IsValid() || !Parameters.IsValid())
+	{
+		return MakeError(LOCTEXT("SerializeStructInvalidData", "Couldn't serialize to target struct: Data View or Parameters were invalid."));
+	}
+
+	if (!ensure(Parameters.GetScriptStruct() == FPCGDataViewJsonParameters::StaticStruct()))
+	{
+		return MakeError(LOCTEXT("ParametersTypeMismatch", "Mismatch between Parameters type and FPCGDataViewJsonParameters."));
+	}
 
 	TSharedPtr<FJsonObject> RootJsonObject = MakeShared<FJsonObject>();
 
@@ -26,38 +35,34 @@ TValueOrError<FInstancedStruct, FText> UPCGDataViewJsonConverter::SerializeToTar
 			return MakeError(LOCTEXT("InvalidViewedData", "Viewed data was invalid"));
 		}
 
-		// @todo_pcg: Currently only support Data Domain and Elements domain, but should be extended eventually to handle any domain.
+		const FPCGDataViewJsonParameters& JsonParameters = Parameters.Get<FPCGDataViewJsonParameters>();
 
-		if (ViewedData->IsSupportedMetadataDomainID(PCGMetadataDomainID::Data))
+		auto GetDomainAttributes = [&Selection, ViewedData, Layout = JsonParameters.AttributeLayout, &RootJsonObject](FPCGMetadataDomainID DomainID)
 		{
-			// Serialize Data Domain attributes first, immediately after the header.
-			const bool bDataIsDefaultDomain = InDataView.ViewedData->GetDefaultMetadataDomainID() == PCGMetadataDomainID::Data;
-			const TArray<FPCGAttributePropertySelector> DataDomainAttributes = Selection.FilterByPredicate([ViewedData, bDataIsDefaultDomain](const FPCGAttributePropertySelector& Selector)
+			const TArray<FPCGAttributePropertySelector> DomainAttributes = Selection.FilterByPredicate([ViewedData, &DomainID, &RootJsonObject](const FPCGAttributePropertySelector& Selector)
 			{
-				const FPCGMetadataDomainID DomainID = ViewedData->GetMetadataDomainIDFromSelector(Selector);
-				return (DomainID == PCGMetadataDomainID::Data) || (bDataIsDefaultDomain && DomainID == PCGMetadataDomainID::Default);
+				const FPCGMetadataDomainID SelectorDomainID = ViewedData->GetMetadataDomainIDFromSelector(Selector);
+				return (SelectorDomainID == DomainID);
 			});
-			Json::Helpers::AppendSelectionByAttribute(RootJsonObject, ViewedData, DataDomainAttributes);
-		}
 
-		// If this data supports the element domain, catalogue the elements and serialize them based on the provided layout.
-		if (InDataView.ViewedData->IsSupportedMetadataDomainID(PCGMetadataDomainID::Elements))
-		{
-			const bool bElementsAreDefaultDomain = InDataView.ViewedData->GetDefaultMetadataDomainID() == PCGMetadataDomainID::Elements;
-			const TArray<FPCGAttributePropertySelector> ElementsDomainAttributes = Selection.FilterByPredicate([&ViewedData = InDataView.ViewedData, bElementsAreDefaultDomain](const FPCGAttributePropertySelector& Selector)
+			if (DomainAttributes.IsEmpty())
 			{
-				const FPCGMetadataDomainID DomainID = ViewedData->GetMetadataDomainIDFromSelector(Selector);
-				return (DomainID == PCGMetadataDomainID::Elements) || (bElementsAreDefaultDomain && DomainID == PCGMetadataDomainID::Default);
-			});
+				return;
+			}
 
-			if (InDataView.AttributeLayout == EPCGDataViewAttributeLayout::ByElement)
+			if (Layout == EPCGDataViewAttributeLayout::ByElement)
 			{
-				Json::Helpers::AppendSelectionByElement(RootJsonObject, InDataView.ViewedData, ElementsDomainAttributes);
+				Json::Helpers::AppendSelectionByElement(RootJsonObject, ViewedData, DomainAttributes);
 			}
-			else if (InDataView.AttributeLayout == EPCGDataViewAttributeLayout::ByAttribute)
+			else if (Layout == EPCGDataViewAttributeLayout::ByAttribute)
 			{
-				Json::Helpers::AppendSelectionByAttribute(RootJsonObject, InDataView.ViewedData, ElementsDomainAttributes);
+				Json::Helpers::AppendSelectionByAttribute(RootJsonObject, ViewedData, DomainAttributes);
 			}
+		};
+
+		for (const FPCGMetadataDomainID DomainID : InDataView.GetAllMetadataDomainIDs())
+		{
+			GetDomainAttributes(DomainID);
 		}
 	}
 
@@ -75,24 +80,15 @@ TValueOrError<FString, FText> UPCGDataViewJsonConverter::SerializeToString(const
 	}
 
 	const FInstancedStruct DataStruct = Result.GetValue();
-	if (!ensure(DataStruct.GetScriptStruct() == FPCGDataViewJsonOutput::StaticStruct()))
+	if (!ensure(DataStruct.IsValid() && DataStruct.GetScriptStruct() == FPCGDataViewJsonOutput::StaticStruct()))
 	{
 		return MakeError(LOCTEXT("DataStructTypesMismatch", "Struct types between result and FPCGDataViewJsonOutput."));
 	}
 
-	if (DataStruct.IsValid() && Parameters.IsValid())
-	{
-		const FPCGDataViewJsonOutput& JsonData = DataStruct.Get<FPCGDataViewJsonOutput>();
-		if (!ensure(Parameters.GetScriptStruct() == FPCGDataViewJsonParameters::StaticStruct()))
-		{
-			return MakeError(LOCTEXT("ParamStructTypesMismatch", "Struct types between parameters and FPCGDataViewJsonParameters."));
-		}
-
-		const FPCGDataViewJsonParameters& JsonParams = Parameters.Get<FPCGDataViewJsonParameters>();
-		return MakeValue(PCG::IO::Json::Helpers::ToJsonString(JsonData.JsonObject, JsonParams.bPrettyJson));
-	}
-
-	return MakeError(LOCTEXT("DataInvalid", "The Data View data or parameters were invalid."));
+	check(Parameters.IsValid()); // already validated in SerializeToTargetStruct call.
+	const FPCGDataViewJsonOutput& JsonData = DataStruct.Get<FPCGDataViewJsonOutput>();
+	const FPCGDataViewJsonParameters& JsonParams = Parameters.Get<FPCGDataViewJsonParameters>();
+	return MakeValue(PCG::IO::Json::Helpers::ToJsonString(JsonData.JsonObject, JsonParams.bPrettyJson));
 }
 
 TValueOrError<UPCGData*, FText> UPCGDataViewJsonConverter::ConstructDataFromTarget(const FInstancedStruct& InTargetStruct, UObject* Outer) const
diff --git a/Engine/Plugins/PCG/Source/PCG/Private/Data/DataView/PCGDataViewNativePropertySelectors.cpp b/Engine/Plugins/PCG/Source/PCG/Private/Data/DataView/PCGDataViewNativePropertySelectors.cpp
--- a/Engine/Plugins/PCG/Source/PCG/Private/Data/DataView/PCGDataViewNativePropertySelectors.cpp
+++ b/Engine/Plugins/PCG/Source/PCG/Private/Data/DataView/PCGDataViewNativePropertySelectors.cpp
@@ -2,6 +2,8 @@
 
 #include "Data/DataView/PCGDataViewNativePropertySelectors.h"
 
+#include "Data/PCGBasePointData.h"
+#include "Data/PCGSplineData.h"
 #include "Data/PCGSplineStruct.h"
 #include "Data/DataView/PCGDataView.h"
 
@@ -38,19 +40,28 @@ namespace PCGDataView
 
 	namespace Helpers
 	{
-		void AppendAllAttributeSelectors(const UPCGData* InData, TArray<FPCGAttributePropertySelector>& InOutSelectors)
+		void AppendAllAttributeSelectors(const UPCGData* InData, TArray<FPCGAttributePropertySelector>& InOutSelectors, TOptional<FPCGMetadataDomainID> OptionalDomain)
 		{
 			check(InData);
 			check(InData->ConstMetadata());
 
-			TArray<FPCGAttributeIdentifier> AttributeNames;
+			TArray<FPCGAttributeIdentifier> AttributeIDs;
 			TArray<EPCGMetadataTypes> AttributeTypes;
-			InData->ConstMetadata()->GetAllAttributes(AttributeNames, AttributeTypes);
+			InData->ConstMetadata()->GetAllAttributes(AttributeIDs, AttributeTypes);
 
-			for (int32 Index = 0; Index < AttributeNames.Num(); ++Index)
+			for (const FPCGAttributeIdentifier& AttributeID : AttributeIDs)
 			{
-				const bool bIsDefaultDomain = AttributeNames[Index].MetadataDomain.IsDefault();
-				InOutSelectors.Add(FPCGAttributePropertySelector::CreateAttributeSelector(AttributeNames[Index].Name, bIsDefaultDomain ? NAME_None : AttributeNames[Index].MetadataDomain.DebugName));
+				FPCGMetadataDomainID MetadataDomain = AttributeID.MetadataDomain.IsDefault() ? InData->GetDefaultMetadataDomainID() : AttributeID.MetadataDomain;
+
+				// Add the selectors if no optional domain is provided, or if provided, it matches the attribute's domain.
+				if (!OptionalDomain.IsSet() || OptionalDomain.GetValue() == MetadataDomain)
+				{
+					FPCGAttributePropertySelector Selector = FPCGAttributePropertySelector::CreateAttributeSelector(AttributeID.Name);
+					InData->SetDomainFromDomainID(MetadataDomain, Selector);
+					check(!MetadataDomain.IsDefault())
+
+					InOutSelectors.Add(Selector);
+				}
 			}
 		}
 
@@ -58,7 +69,7 @@ namespace PCGDataView
 		{
 			for (const EPCGPointProperties Property : Constants::AllPointProperties)
 			{
-				InOutSelectors.Add(FPCGAttributePropertySelector::CreatePointPropertySelector(Property));
+				InOutSelectors.Add(FPCGAttributePropertySelector::CreatePointPropertySelector(Property, PCGPointDataConstants::ElementsDomainName));
 			}
 		}
 
@@ -69,7 +80,7 @@ namespace PCGDataView
 				for (const EPCGSplineStructProperties Attribute : Constants::AllSplineProperties)
 				{
 					const FName DisplayName = FName(SplinePropertiesEnum->GetNameStringByValue(static_cast<int32>(Attribute)));
-					InOutSelectors.Add(FPCGAttributePropertySelector::CreatePropertySelector(DisplayName));
+					InOutSelectors.Add(FPCGAttributePropertySelector::CreatePropertySelector(DisplayName, PCGSplineData::ControlPointDomainName));
 				}
 			}
 		}
diff --git a/Engine/Plugins/PCG/Source/PCG/Private/Elements/PCGConvertToDataView.cpp b/Engine/Plugins/PCG/Source/PCG/Private/Elements/PCGConvertToDataView.cpp
--- a/Engine/Plugins/PCG/Source/PCG/Private/Elements/PCGConvertToDataView.cpp
+++ b/Engine/Plugins/PCG/Source/PCG/Private/Elements/PCGConvertToDataView.cpp
@@ -88,8 +88,6 @@ bool FPCGConvertToDataViewElement::ExecuteInternal(FPCGContext* InContext) const
 			DataView.Selection = Settings->Selection;
 		}
 
-		DataView.AttributeLayout = Settings->AttributeLayout;
-
 		UPCGDataViewData* DataViewData = InContext->NewObject_AnyThread<UPCGDataViewData>(InContext);
 		DataViewData->Initialize(MoveTemp(DataView));
 
diff --git a/Engine/Plugins/PCG/Source/PCG/Private/Helpers/IO/PCGJsonHelpers.cpp b/Engine/Plugins/PCG/Source/PCG/Private/Helpers/IO/PCGJsonHelpers.cpp
--- a/Engine/Plugins/PCG/Source/PCG/Private/Helpers/IO/PCGJsonHelpers.cpp
+++ b/Engine/Plugins/PCG/Source/PCG/Private/Helpers/IO/PCGJsonHelpers.cpp
@@ -5,6 +5,8 @@
 #include "Helpers/PCGHelpers.h"
 #include "Helpers/IO/PCGIOHelpers.h"
 #include "Metadata/Accessors/PCGAttributeAccessorHelpers.h"
+
+#include "Algo/Find.h"
 #include "Policies/CondensedJsonPrintPolicy.h"
 
 #include "Serialization/JsonSerializer.h"
@@ -238,7 +240,32 @@ namespace PCG::IO::Json
 			const Accessor::FCache AccessorCache = Private::CreateAccessorCache(InData, Selectors, InOptionalContext);
 			JsonKeyValuesMap KeyValuesMap;
 			Private::BuildKeyValuesMap(KeyValuesMap, AccessorCache);
-			Private::AppendSelectionByElement(InOutJsonObject, KeyValuesMap, PCGHelpers::GetNumberOfElements(InData));
+
+			FPCGMetadataDomainID DomainID = InData->GetMetadataDomainIDFromSelector(Selectors[0]);
+			const FPCGAttributePropertySelector* MismatchedSelector = Algo::FindByPredicate(Selectors, [InData, &DomainID](const FPCGAttributePropertySelector& InSelector)
+			{
+				return InData->GetMetadataDomainIDFromSelector(InSelector) != DomainID;
+			});
+
+			if (MismatchedSelector)
+			{
+				FString Domain1 = Selectors[0].GetDomainString();
+				Domain1 = Domain1.IsEmpty() ? TEXT("Default") : Domain1;
+				FString Domain2 = MismatchedSelector->GetDomainString();
+				Domain2 = Domain2.IsEmpty() ? TEXT("Default") : Domain2;
+
+				const FText Warning = FText::Format(
+					LOCTEXT("SelectorsFromMultipleDomains",
+							"[Json] Could not append selection by element. All selectors must share the same domain: {0} vs {1}"),
+					FText::FromString(Domain1),
+					FText::FromString(Domain2));
+
+				PCGLog::LogWarningOnGraph(Warning, InOptionalContext);
+
+				return;
+			}
+
+			Private::AppendSelectionByElement(InOutJsonObject, KeyValuesMap, PCGHelpers::GetNumberOfElements(InData, DomainID));
 		}
 
 		FString ToJsonString(const TSharedPtr<FJsonObject>& JsonObject, const bool bPrettyJson)
diff --git a/Engine/Plugins/PCG/Source/PCG/Private/Helpers/PCGHelpers.cpp b/Engine/Plugins/PCG/Source/PCG/Private/Helpers/PCGHelpers.cpp
--- a/Engine/Plugins/PCG/Source/PCG/Private/Helpers/PCGHelpers.cpp
+++ b/Engine/Plugins/PCG/Source/PCG/Private/Helpers/PCGHelpers.cpp
@@ -774,19 +774,20 @@ namespace PCGHelpers
 		}
 	}
 
-	int32 GetNumberOfElements(const UPCGData* InData)
+	int32 GetNumberOfElements(const UPCGData* InData, const FPCGMetadataDomainID& InDomainID)
 	{
-		if (!InData->IsSupportedMetadataDomainID(PCGMetadataDomainID::Elements))
+		if (!InData->IsSupportedMetadataDomainID(InDomainID))
 		{
 			return 0;
 		}
-		
-		FPCGAttributePropertyInputSelector ElementSelector;
-		ElementSelector.SetAttributeName(NAME_None);
-		InData->SetDomainFromDomainID(PCGMetadataDomainID::Elements, ElementSelector);
-		TUniquePtr<const IPCGAttributeAccessorKeys> ElementKeys = PCGAttributeAccessorHelpers::CreateConstKeys(InData, ElementSelector);
-		return ElementKeys ? ElementKeys->GetNum() : 0;
+
+		FPCGAttributePropertyInputSelector DomainSelector;
+		DomainSelector.SetAttributeName(NAME_None);
+		InData->SetDomainFromDomainID(InDomainID, DomainSelector);
+		TUniquePtr<const IPCGAttributeAccessorKeys> DomainKeys = PCGAttributeAccessorHelpers::CreateConstKeys(InData, DomainSelector);
+		return DomainKeys ? DomainKeys->GetNum() : 0;
 	}
+
 #if UE_ENABLE_DEBUG_DRAWING
 
 	void DebugDrawGenerationVolume(FPCGContext* InContext, const FColor* InOverrideColor)
diff --git a/Engine/Plugins/PCG/Source/PCG/Public/Data/DataView/PCGDataView.h b/Engine/Plugins/PCG/Source/PCG/Public/Data/DataView/PCGDataView.h
--- a/Engine/Plugins/PCG/Source/PCG/Public/Data/DataView/PCGDataView.h
+++ b/Engine/Plugins/PCG/Source/PCG/Public/Data/DataView/PCGDataView.h
@@ -3,8 +3,8 @@
 #pragma once
 
 #include "PCGData.h"
-#include "Helpers/IO/PCGIOHelpers.h"
 #include "Metadata/PCGAttributePropertySelector.h"
+#include "Metadata/PCGMetadataCommon.h"
 
 #include "Serialization/CustomVersion.h"
 #include "Templates/SubclassOf.h"
@@ -36,14 +36,6 @@ namespace PCGDataView::Constants
 	};
 } // namespace PCGDataView::Constants
 
-UENUM()
-enum class EPCGDataViewConversion : uint32
-{
-	Json UMETA(Tooltip = "Convert the data view into Json format. See 'EPCGDataViewAttributeLayout' for how properties and attributes can be serialized."),
-
-	// @todo_pcg: Binary (Archive)
-};
-
 UENUM()
 enum class EPCGDataViewAttributeLayout : uint8
 {
@@ -60,9 +52,21 @@ struct FPCGDataViewSelection
 	UPROPERTY(EditAnywhere, Category = Selection)
 	bool bAllAttributes = true;
 
+	// Limit the selection to a specific domain.
+	UPROPERTY(EditAnywhere, Category = Selection, meta = (EditConditionHides, InlineEditConditionToggle))
+	bool bLimitByDomain = false;
+
+	// Ignore default attributes of the viewed data (properties). Ex. $Transform, $Density for Point Data.
+	UPROPERTY(EditAnywhere, Category = Selection, DisplayName = "Ignore Static Attributes", meta = (EditCondition = "bAllAttributes", EditConditionHides))
+	bool bIgnoreProperties = false;
+
 	// @todo_pcg: Extend to use EPCGMetadataFilterMode for more complex selections
 	UPROPERTY(EditAnywhere, Category = Selection, meta = (EditCondition = "!bAllAttributes", EditConditionHides))
 	TArray<FPCGAttributePropertySelector> Attributes;
+
+	// If "Limit By Domain" is enabled and a domain is provided, attributes from other domains will be filtered out.
+	UPROPERTY(EditAnywhere, Category = Selection, DisplayName = "Limit by Domain", meta = (EditCondition = "bLimitByDomain"))
+	FName Domain = TEXT("Data");
 };
 
 /** This wrapper class maintains a reference to unowned "viewed" data and also contains a "selection" of
@@ -96,8 +100,4 @@ struct FPCGDataView
 	// An optionally filtered collection of attributes and properties to perform an operation on, such as data export.
 	UPROPERTY(VisibleAnywhere, Category = Data, meta = (ShowOnlyInnerProperties))
 	FPCGDataViewSelection Selection;
-
-	// How the attribute data should be exported, assuming the data supports the elements domain.
-	UPROPERTY(VisibleAnywhere, Category = Data)
-	EPCGDataViewAttributeLayout AttributeLayout = EPCGDataViewAttributeLayout::ByElement;
 };
diff --git a/Engine/Plugins/PCG/Source/PCG/Public/Data/DataView/PCGDataViewNativeJsonConverters.h b/Engine/Plugins/PCG/Source/PCG/Public/Data/DataView/PCGDataViewNativeJsonConverters.h
--- a/Engine/Plugins/PCG/Source/PCG/Public/Data/DataView/PCGDataViewNativeJsonConverters.h
+++ b/Engine/Plugins/PCG/Source/PCG/Public/Data/DataView/PCGDataViewNativeJsonConverters.h
@@ -2,6 +2,7 @@
 
 #pragma once
 
+#include "Data/DataView/PCGDataView.h"
 #include "Data/DataView/PCGDataViewInterface.h"
 
 #include "PCGDataViewNativeJsonConverters.generated.h"
@@ -13,6 +14,9 @@ struct FPCGDataViewJsonParameters
 {
 	GENERATED_BODY()
 
+	UPROPERTY(EditAnywhere, Category = "Json")
+	EPCGDataViewAttributeLayout AttributeLayout = EPCGDataViewAttributeLayout::ByElement;
+
 	UPROPERTY(EditAnywhere, Category = "Json")
 	bool bPrettyJson = false;
 };
diff --git a/Engine/Plugins/PCG/Source/PCG/Public/Data/DataView/PCGDataViewNativePropertySelectors.h b/Engine/Plugins/PCG/Source/PCG/Public/Data/DataView/PCGDataViewNativePropertySelectors.h
--- a/Engine/Plugins/PCG/Source/PCG/Public/Data/DataView/PCGDataViewNativePropertySelectors.h
+++ b/Engine/Plugins/PCG/Source/PCG/Public/Data/DataView/PCGDataViewNativePropertySelectors.h
@@ -7,7 +7,7 @@
 namespace PCGDataView::Helpers
 {
 	// Append all non-property attributes to the provided array of selectors.
-	void AppendAllAttributeSelectors(const UPCGData* InData, TArray<FPCGAttributePropertySelector>& InOutSelectors);
+	void AppendAllAttributeSelectors(const UPCGData* InData, TArray<FPCGAttributePropertySelector>& InOutSelectors, TOptional<FPCGMetadataDomainID> OptionalDomain = {});
 
 	// Append all relevant point properties avoiding redundancy, i.e. Transform already includes Positions, Rotation, Scale
 	void AppendAllPointPropertySelectors(TArray<FPCGAttributePropertySelector>& InOutSelectors);
diff --git a/Engine/Plugins/PCG/Source/PCG/Public/Elements/PCGConvertToDataView.h b/Engine/Plugins/PCG/Source/PCG/Public/Elements/PCGConvertToDataView.h
--- a/Engine/Plugins/PCG/Source/PCG/Public/Elements/PCGConvertToDataView.h
+++ b/Engine/Plugins/PCG/Source/PCG/Public/Elements/PCGConvertToDataView.h
@@ -46,10 +46,6 @@ public:
 	/** Manually select attributes or properties for this data view. */
 	UPROPERTY(EditAnywhere, BlueprintReadWrite, Category = "Settings", meta = (EditCondition = "!bSelectionAsInput", EditConditionHides))
 	FPCGDataViewSelection Selection;
-
-	/** How the attribute data should be exported, assuming the data supports the elements domain. */
-	UPROPERTY(EditAnywhere, BlueprintReadWrite, Category = "Settings", meta = (PCG_Overridable))
-	EPCGDataViewAttributeLayout AttributeLayout = EPCGDataViewAttributeLayout::ByElement;
 };
 
 class FPCGConvertToDataViewElement : public IPCGElement
diff --git a/Engine/Plugins/PCG/Source/PCG/Public/Helpers/PCGHelpers.h b/Engine/Plugins/PCG/Source/PCG/Public/Helpers/PCGHelpers.h
--- a/Engine/Plugins/PCG/Source/PCG/Public/Helpers/PCGHelpers.h
+++ b/Engine/Plugins/PCG/Source/PCG/Public/Helpers/PCGHelpers.h
@@ -4,6 +4,7 @@
 
 #include "PCGCommon.h"
 #include "PCGPoint.h"
+#include "Metadata/PCGMetadataCommon.h"
 
 #include "EngineDefines.h" // For UE_ENABLE_DEBUG_DRAWING
 #include "Containers/Ticker.h"
@@ -25,6 +26,7 @@ class UPCGMetadata;
 class UPCGSettings;
 class UWorld;
 struct FPCGContext;
+struct FPCGMetadataDomainID;
 struct FPCGPoint;
 
 namespace PCGHelpers
@@ -187,9 +189,11 @@ namespace PCGHelpers
 		}
 	}
 
-	// Create keys for a hypothetical attribute in the 'Elements' domain, get its count, and create a const value accessor for retrieving it
-	// Only available if the data supports the Elements domain 
-	PCG_API int32 GetNumberOfElements(const UPCGData* InData);
+	/**
+	 * Create keys for a hypothetical attribute in the provided domain ('Elements' by default), get its count, and create
+	 * a const value accessor for retrieving it. Only available if the data supports the provided domain.
+	 */
+	PCG_API int32 GetNumberOfElements(const UPCGData* InData, const FPCGMetadataDomainID& InDomainID = PCGMetadataDomainID::Elements);
 
 #if UE_ENABLE_DEBUG_DRAWING 
 	PCG_API void DebugDrawGenerationVolume(FPCGContext* InContext, const FColor* InOverrideColor = nullptr);
