[PCG] Make sure aliases for PCG Overridable are taken into account for override pin deprecation rename

* When a property name changes in a settings, we can provide an alias for the overrides by name to still work
* But also use that information to know how to rename a override pin, so we don't break existing pins.

#rnx
#Synced-CL 48852953

diff --git a/Engine/Plugins/PCG/Source/PCG/Private/PCGSettings.cpp b/Engine/Plugins/PCG/Source/PCG/Private/PCGSettings.cpp
--- a/Engine/Plugins/PCG/Source/PCG/Private/PCGSettings.cpp
+++ b/Engine/Plugins/PCG/Source/PCG/Private/PCGSettings.cpp
@@ -1106,9 +1106,9 @@ void UPCGSettings::InitializeCachedOverridableParams(bool bReset)
 		// Deprecation: Previously params labels were using the GetDisplayTextName, which is localized and will break existing pins.
 		// Now it will be the authored name, but we still need to make sure the pins are not broken. 
 		// For that matter, we will match names for properties between serialized and newly found ones. If their label don't match, we add the previous label as alias.
-		// FIXME: Should be depending on a version, but we can't do that as an hotfix.
+		// Also we keep a list of previously serialized labels. If we have an alias that is the same as one of those labels, it will also mark this label up for deprecation to rename and not breaking pins.
 		TMap<FString, FName> PropertyPathToLabel;
-
+		TSet<FName> PreUpdatedLabels;
 		if (!CachedOverridableParams.IsEmpty())
 		{
 			for (const FPCGSettingsOverridableParam& Param : CachedOverridableParams)
@@ -1119,6 +1119,7 @@ void UPCGSettings::InitializeCachedOverridableParams(bool bReset)
 				}
 
 				PropertyPathToLabel.Emplace(Param.GetPropertyPath(), Param.Label);
+				PreUpdatedLabels.Add(Param.Label);
 			}
 		}
 
@@ -1139,6 +1140,16 @@ void UPCGSettings::InitializeCachedOverridableParams(bool bReset)
 					if (*OldLabel != Param.Label)
 					{
 						Param.MapOfAliases.FindOrAdd(PCGSettings::DeprecationAliasIndex).Aliases.Emplace(*OldLabel);
+						continue;
+					}
+				}
+
+				for (const FName Alias : Param.GenerateAllPossibleAliases())
+				{
+					if (PreUpdatedLabels.Contains(Alias))
+					{
+						Param.MapOfAliases.FindOrAdd(PCGSettings::DeprecationAliasIndex).Aliases.Emplace(Alias);
+						break;
 					}
 				}
 			}
