[PCG] Reduce scheduling contention overhead by adding a hard limit on worker threads (cvar default to 8)
- unused threads are added to the PCGAsync task budget

#Synced-CL 48946867

diff --git a/Engine/Plugins/PCG/Source/PCG/Private/Graph/PCGGraphExecutor.cpp b/Engine/Plugins/PCG/Source/PCG/Private/Graph/PCGGraphExecutor.cpp
--- a/Engine/Plugins/PCG/Source/PCG/Private/Graph/PCGGraphExecutor.cpp
+++ b/Engine/Plugins/PCG/Source/PCG/Private/Graph/PCGGraphExecutor.cpp
@@ -60,6 +60,11 @@ static TAutoConsoleVariable<float> CVarMaxPercentageOfExecutingThreads(
 	0.3f,
 	TEXT("Maximum percentage of number of threads for processing scheduled pcg elements (percentage of pcg.MaxPercentageOfThreadsToUse)"));
 
+static TAutoConsoleVariable<int32> CVarMaxExecutingThreads(
+	TEXT("pcg.MaxExecutingThreads"),
+	8,
+	TEXT("Hard cap on number of threads for processing scheduled pcg elements calculated from pcg.MaxPercentageOfExecutingThreads"));
+
 namespace PCGGraphExecutor
 {
 	TAutoConsoleVariable<float> CVarTimePerFrame(
@@ -1553,7 +1558,10 @@ bool FPCGGraphExecutor::ExecuteScheduling(double EndTime, TSharedPtr<FPCGGraphAc
 	const int32 MaxNumThreads = FApp::ShouldUseThreadingForPerformance() ? FMath::Max(0, (int32)(LowLevelTasks::FScheduler::Get().GetNumWorkers() * MaxPercentageOfThreadsToUse)) : 0;
 
 	// Number of threads to use for PCG Element execution (outside of main thread tick)
-	const int32 MaxExecutingThreads = PCGGraphExecutor::CVarGraphMultithreading.GetValueOnAnyThread() ? FMath::Max(0, (int32)(MaxNumThreads * MaxPercentageOfExecutingThreads)) : 0;
+	const int32 MaxExecutingThreadsPriorHardCap = PCGGraphExecutor::CVarGraphMultithreading.GetValueOnAnyThread() ? FMath::Max(0, (int32)(MaxNumThreads * MaxPercentageOfExecutingThreads)) : 0;
+
+	// Final number of threads to use for PCG Element execution (outside of main thread tick)
+	const int32 MaxExecutingThreads = FMath::Min(MaxExecutingThreadsPriorHardCap, FMath::Max(0, CVarMaxExecutingThreads.GetValueOnAnyThread()));
 
 	// Number of threads to use for PCG Async operations 
 	const int32 MaxPCGAsyncThreads = FMath::Min(MaxNumThreads - MaxExecutingThreads, CVarMaxNumTasks.GetValueOnAnyThread());
