[PCG] Fix Cache result raw stack ptr outliving its source Task. Use FPCGStackHandle instead (ref counted)

#Synced-CL 48946867

diff --git a/Engine/Plugins/PCG/Source/PCG/Private/Graph/PCGGraphExecutor.cpp b/Engine/Plugins/PCG/Source/PCG/Private/Graph/PCGGraphExecutor.cpp
--- a/Engine/Plugins/PCG/Source/PCG/Private/Graph/PCGGraphExecutor.cpp
+++ b/Engine/Plugins/PCG/Source/PCG/Private/Graph/PCGGraphExecutor.cpp
@@ -1443,7 +1443,7 @@ void FPCGGraphExecutor::PrepareForExecute(FPCGGraphTask& Task, FCachedResult*& O
 		if (PCGGraphExecutor::CVarDynamicTaskCulling.GetValueOnAnyThread() && TaskSettings && TaskSettings->OutputPinsCanBeDeactivated() && OutCachedResult->Output.InactiveOutputPinBitmask != 0)
 		{
 			OutCachedResult->bDoDynamicTaskCulling = true;
-			OutCachedResult->Stack = Task.GetStack();
+			OutCachedResult->Stack = Task.GetStackHandle();
 			OutCachedResult->Node = Task.Node;
 		}
 
@@ -2081,7 +2081,7 @@ void FPCGGraphExecutor::ProcessCachedResultsInternal(TConstArrayView<FCachedResu
 		{
 			CullInactiveDownstreamNodes(CachedResult->TaskId, CachedResult->Output.InactiveOutputPinBitmask);
 #if WITH_EDITOR
-			SendInactivePinNotification(CachedResult->Node, CachedResult->Stack, CachedResult->Output.InactiveOutputPinBitmask);
+			SendInactivePinNotification(CachedResult->Node, CachedResult->Stack.GetStack(), CachedResult->Output.InactiveOutputPinBitmask);
 #endif
 		}
 
diff --git a/Engine/Plugins/PCG/Source/PCG/Private/Graph/PCGGraphExecutor.h b/Engine/Plugins/PCG/Source/PCG/Private/Graph/PCGGraphExecutor.h
--- a/Engine/Plugins/PCG/Source/PCG/Private/Graph/PCGGraphExecutor.h
+++ b/Engine/Plugins/PCG/Source/PCG/Private/Graph/PCGGraphExecutor.h
@@ -192,7 +192,7 @@ private:
 	{
 		FPCGTaskId TaskId = InvalidPCGTaskId;
 		FPCGDataCollection Output;
-		const FPCGStack* Stack = nullptr;
+		FPCGStackHandle Stack;
 		const UPCGNode* Node = nullptr;
 		bool bDoDynamicTaskCulling = false;
 		bool bIsPostGraphTask = false;		
