[PCG] Move Hit/Overlap result handling to Physics Registry

#Synced-CL 48958696

diff --git a/Engine/Plugins/PCG/Source/PCG/Private/Data/PCGWorldData.cpp b/Engine/Plugins/PCG/Source/PCG/Private/Data/PCGWorldData.cpp
--- a/Engine/Plugins/PCG/Source/PCG/Private/Data/PCGWorldData.cpp
+++ b/Engine/Plugins/PCG/Source/PCG/Private/Data/PCGWorldData.cpp
@@ -29,7 +29,7 @@
 namespace PCGWorldCommonQueryParams
 {
 	template<class T, class Container, typename Func>
-	bool ExtractActorFiltersIfNeeded(const FPCGWorldCommonQueryParams& InParams, const UPCGData* InData, Container& OutContainer, FPCGContext* InContext, Func TransformLambda)
+	bool ExtractObjectFiltersIfNeeded(const FPCGWorldCommonQueryParams& InParams, const UPCGData* InData, Container& OutContainer, FPCGContext* InContext, Func TransformLambda)
 	{
 		if (InParams.ActorFilterFromInput != EPCGWorldQueryFilter::None)
 		{
@@ -59,18 +59,39 @@ void FPCGWorldCommonQueryParams::AddFilterPinIfNeeded(TArray<FPCGPinProperties>&
 	}
 }
 
+// deprecated
 bool FPCGWorldCommonQueryParams::ExtractActorFiltersIfNeeded(const UPCGData* InData, TArray<TSoftObjectPtr<AActor>>& OutArray, FPCGContext* InContext) const
 {
-	return PCGWorldCommonQueryParams::ExtractActorFiltersIfNeeded<TSoftObjectPtr<AActor>>(*this, InData, OutArray, InContext,
+	return PCGWorldCommonQueryParams::ExtractObjectFiltersIfNeeded<TSoftObjectPtr<AActor>>(*this, InData, OutArray, InContext,
 		[](const FSoftObjectPath& ActorPath) -> TSoftObjectPtr<AActor> { return TSoftObjectPtr<AActor>(ActorPath); });
 }
 
+// deprecated
 bool FPCGWorldCommonQueryParams::ExtractLoadedActorFiltersIfNeeded(const UPCGData* InData, TSet<TObjectKey<AActor>>& OutSet, FPCGContext* InContext) const
 {
-	return PCGWorldCommonQueryParams::ExtractActorFiltersIfNeeded<TObjectKey<AActor>>(*this, InData, OutSet, InContext,
+	return PCGWorldCommonQueryParams::ExtractObjectFiltersIfNeeded<TObjectKey<AActor>>(*this, InData, OutSet, InContext,
 	[](const FSoftObjectPath& ActorPath) -> TObjectKey<AActor> { return TObjectKey<AActor>(Cast<AActor>(ActorPath.ResolveObject())); });
 }
 
+bool FPCGWorldCommonQueryParams::ExtractObjectFiltersIfNeeded(const UPCGData* InData, TArray<FSoftObjectPath>& OutArray, FPCGContext* InContext) const
+{
+	if (ActorFilterFromInput != EPCGWorldQueryFilter::None)
+	{
+		if (!PCGAttributeAccessorHelpers::ExtractAllValues(InData, ActorFilterInputSource, OutArray, InContext))
+		{
+			return false;
+		}
+	}
+
+	return true;
+}
+
+bool FPCGWorldCommonQueryParams::ExtractLoadedObjectFiltersIfNeeded(const UPCGData* InData, TSet<FObjectKey>& OutSet, FPCGContext* InContext) const
+{
+	return PCGWorldCommonQueryParams::ExtractObjectFiltersIfNeeded<FObjectKey>(*this, InData, OutSet, InContext,
+		[](const FSoftObjectPath& ObjectPath) -> FObjectKey { return FObjectKey(ObjectPath.ResolveObject()); });
+}
+
 void FPCGWorldCommonQueryParams::Initialize()
 {
 	if (ActorTagFilter == EPCGWorldQueryFilter::None)
@@ -144,35 +165,86 @@ void FPCGWorldRayHitQueryParams::PostSerialize(const FArchive& Ar)
 #endif
 }
 
-const TSet<TObjectKey<AActor>>& FPCGWorldQueryActorFilterCache::GetCachedFilterActors() const
+void UPCGWorldVolumetricData::Initialize(IPCGGraphExecutionSource* InExecutionSource, const FBox& InBounds)
+{
+	Super::Initialize(InBounds);
+	SetOriginatingSource(InExecutionSource);
+	check(InExecutionSource && InExecutionSource->GetExecutionState().GetWorld());
+}
+
+IPCGGraphExecutionSource* UPCGWorldVolumetricData::GetOriginatingSource() const
+{ 
+	return Cast<IPCGGraphExecutionSource>(OriginatingObject.ResolveObject()); 
+}
+
+UWorld* UPCGWorldVolumetricData::GetWorld() const
+{
+	if (IPCGGraphExecutionSource* ExecutionSource = GetOriginatingSource())
+	{
+		return ExecutionSource->GetExecutionState().GetWorld();
+	}
+	else
+	{
+		return nullptr;
+	}
+}
+
+void UPCGWorldVolumetricData::SetOriginatingSource(IPCGGraphExecutionSource* InExecutionSource) 
+{ 
+	OriginatingObject = Cast<UObject>(InExecutionSource); 
+}
+
+const TSet<FObjectKey>& UPCGWorldVolumetricData::GetCachedFilteredObjects() const
 {
-	if (!FilterActors.IsEmpty() && CachedFilterActorsDirty)
+	if (!FilteredObjects.IsEmpty() && CachedFilteredObjectsDirty)
 	{
-		PCG::TScopeLock Lock(CachedFilterActorsLock);
-		if (CachedFilterActorsDirty)
+		PCG::TScopeLock Lock(CachedFilteredObjectsLock);
+		if (CachedFilteredObjectsDirty)
 		{
-			CachedFilterActors.Empty(FilterActors.Num());
-			Algo::Transform(FilterActors, CachedFilterActors, [](const TSoftObjectPtr<AActor>& SoftPtr) { return TObjectKey<AActor>(SoftPtr.Get()); });
-			CachedFilterActorsDirty = false;
+			CachedFilteredObjects.Empty(FilteredObjects.Num());
+			Algo::Transform(FilteredObjects, CachedFilteredObjects, [](const FSoftObjectPath& SoftObjectPath) { return FObjectKey(SoftObjectPath.ResolveObject()); });
+			CachedFilteredObjectsDirty = false;
 		}
 	}
 
-	return CachedFilterActors;
+	return CachedFilteredObjects;
 }
 
-void UPCGWorldVolumetricData::Initialize(UWorld* InWorld, const FBox& InBounds)
+void UPCGWorldVolumetricData::PostLoad()
 {
-	Super::Initialize(InBounds);
-	World = InWorld;
-	check(World.IsValid());
+	Super::PostLoad();
+
+PRAGMA_DISABLE_DEPRECATION_WARNINGS
+	
+#if WITH_EDITOR
+	// Move data into new member
+	if (!ActorFilter.FilterActors_DEPRECATED.IsEmpty())
+	{
+		Algo::Transform(ActorFilter.FilterActors_DEPRECATED, GetMutableFilteredObjects(), [](const TSoftObjectPtr<AActor> InActor) { return InActor.ToSoftObjectPath(); });
+	}
+
+	if(OriginatingComponent.IsValid())
+	{
+		SetOriginatingSource(OriginatingComponent.Get());
+		OriginatingComponent = nullptr;
+	}
+#endif // WITH_EDITOR
+
+PRAGMA_ENABLE_DEPRECATION_WARNINGS
+
+	// Cache is dirty if filtered objects is not empty on load
+	CachedFilteredObjectsDirty = GetConstFilteredObjects().Num() > 0;
 }
 
 bool UPCGWorldVolumetricData::SamplePoint(const FTransform& InTransform, const FBox& InBounds, FPCGPoint& OutPoint, UPCGMetadata* OutMetadata) const
 {
 	// This is a pure implementation.
-
-	check(World.IsValid());
-
+	UWorld* ResolvedWorld = GetWorld();
+	if (!ensure(ResolvedWorld))
+	{
+		return false;
+	}
+	
 	FPCGMetadataAttribute<FSoftObjectPath>* ActorOverlappedAttribute = ((OutMetadata && QueryParams.bGetReferenceToActorHit) ? OutMetadata->FindOrCreateAttribute<FSoftObjectPath>(PCGPointDataConstants::ActorReferenceAttribute) : nullptr);
 
 	FCollisionObjectQueryParams ObjectQueryParams(QueryParams.CollisionChannel);
@@ -181,9 +253,9 @@ bool UPCGWorldVolumetricData::SamplePoint(const FTransform& InTransform, const F
 	Params.bTraceComplex = QueryParams.bTraceComplex;
 
 	TArray<FOverlapResult> Overlaps;
-	/*bool bOverlaps =*/ World->OverlapMultiByObjectType(Overlaps, InTransform.TransformPosition(InBounds.GetCenter()), InTransform.GetRotation(), ObjectQueryParams, CollisionShape, Params);
+	/*bool bOverlaps =*/ ResolvedWorld->OverlapMultiByObjectType(Overlaps, InTransform.TransformPosition(InBounds.GetCenter()), InTransform.GetRotation(), ObjectQueryParams, CollisionShape, Params);
 
-	TOptional<FOverlapResult> Overlap = PCGWorldQueryHelpers::FilterOverlapResults(&QueryParams, OriginatingComponent, Overlaps, ActorFilter.GetCachedFilterActors());
+	TOptional<FOverlapResult> Overlap = PCGWorldQueryHelpers::FilterOverlapResults(QueryParams, GetOriginatingSource(), Overlaps, GetCachedFilteredObjects());
 
 	// If searched for overlap and found one, or didn't search and didn't find one, set the point and return true. Otherwise, return false.
 	if (Overlap.IsSet() == QueryParams.bSearchForOverlap)
@@ -265,24 +337,88 @@ UPCGSpatialData* UPCGWorldVolumetricData::CopyInternal(FPCGContext* Context) con
 
 	CopyBaseVolumeData(NewVolumetricData);
 
-	NewVolumetricData->World = World;
-	NewVolumetricData->OriginatingComponent = OriginatingComponent;
+	NewVolumetricData->SetOriginatingSource(GetOriginatingSource());
 	NewVolumetricData->QueryParams = QueryParams;
 	NewVolumetricData->QueryParams.Initialize();
-	NewVolumetricData->ActorFilter.GetFilterActorsMutable() = ActorFilter.GetFilterActorsConst();
+	NewVolumetricData->GetMutableFilteredObjects() = GetConstFilteredObjects();
 
 	return NewVolumetricData;
 }
 
 /** World Ray Hit data implementation */
-void UPCGWorldRayHitData::Initialize(UWorld* InWorld, const FTransform& InTransform, const FBox& InBounds, const FBox& InLocalBounds)
+void UPCGWorldRayHitData::Initialize(IPCGGraphExecutionSource* InExecutionSource, const FTransform& InTransform, const FBox& InBounds, const FBox& InLocalBounds)
 {
-	World = InWorld;
+	// To save object path
+	SetOriginatingSource(InExecutionSource);
 	Transform = InTransform;
 	Bounds = InBounds;
 	LocalBounds = InLocalBounds;
 }
 
+IPCGGraphExecutionSource* UPCGWorldRayHitData::GetOriginatingSource() const
+{
+	return Cast<IPCGGraphExecutionSource>(OriginatingObject.ResolveObject());
+}
+
+UWorld* UPCGWorldRayHitData::GetWorld() const
+{
+	if (IPCGGraphExecutionSource* ExecutionSource = GetOriginatingSource())
+	{
+		return ExecutionSource->GetExecutionState().GetWorld();
+	}
+	else
+	{
+		return nullptr;
+	}
+}
+
+void UPCGWorldRayHitData::SetOriginatingSource(IPCGGraphExecutionSource* InExecutionSource)
+{
+	OriginatingObject = Cast<UObject>(InExecutionSource);
+}
+
+const TSet<FObjectKey>& UPCGWorldRayHitData::GetCachedFilteredObjects() const
+{
+	if (!FilteredObjects.IsEmpty() && CachedFilteredObjectsDirty)
+	{
+		PCG::TScopeLock Lock(CachedFilteredObjectsLock);
+		if (CachedFilteredObjectsDirty)
+		{
+			CachedFilteredObjects.Empty(FilteredObjects.Num());
+			Algo::Transform(FilteredObjects, CachedFilteredObjects, [](const FSoftObjectPath& SoftObjectPath) { return FObjectKey(SoftObjectPath.ResolveObject()); });
+			CachedFilteredObjectsDirty = false;
+		}
+	}
+
+	return CachedFilteredObjects;
+}
+
+void UPCGWorldRayHitData::PostLoad()
+{
+	Super::PostLoad();
+
+PRAGMA_DISABLE_DEPRECATION_WARNINGS
+
+#if WITH_EDITOR
+	// Move data into new member
+	if (!ActorFilter.FilterActors_DEPRECATED.IsEmpty())
+	{
+		Algo::Transform(ActorFilter.FilterActors_DEPRECATED, GetMutableFilteredObjects(), [](const TSoftObjectPtr<AActor> InActor) { return InActor.ToSoftObjectPath(); });
+	}
+
+	if(OriginatingComponent.IsValid())
+	{
+		SetOriginatingSource(OriginatingComponent.Get());
+		OriginatingComponent = nullptr;
+	}
+#endif // WITH_EDITOR	
+
+PRAGMA_ENABLE_DEPRECATION_WARNINGS
+
+	// Cache is dirty if filtered objects is not empty on load
+	CachedFilteredObjectsDirty = GetConstFilteredObjects().Num() > 0;
+}
+
 void UPCGWorldRayHitData::AddToCrc(FArchiveCrc32& Ar, bool bFullDataCrc) const
 {
 	Super::AddToCrc(Ar, bFullDataCrc);
@@ -294,7 +430,11 @@ void UPCGWorldRayHitData::AddToCrc(FArchiveCrc32& Ar, bool bFullDataCrc) const
 bool UPCGWorldRayHitData::SamplePoint(const FTransform& InTransform, const FBox& InBounds, FPCGPoint& OutPoint, UPCGMetadata* OutMetadata) const
 {
 	// TODO: This seems to be a projection - along a direction. I suspect that UPCGWorldVolumetricData is the SamplePoint(), and this is the ProjectPoint() (in a direction)?
-	check(World.IsValid());
+	UWorld* ResolvedWorld = GetWorld();
+	if (!ensure(ResolvedWorld))
+	{
+		return false;
+	}
 
 	// Todo: consider prebuilding this
 	const FCollisionObjectQueryParams ObjectQueryParams(QueryParams.CollisionChannel);
@@ -304,7 +444,7 @@ bool UPCGWorldRayHitData::SamplePoint(const FTransform& InTransform, const FBox&
 	FVector RayEnd = RayOrigin + QueryParams.RayDirection * QueryParams.RayLength;
 
 	TArray<FHitResult> Hits;
-	if (World->SweepMultiByObjectType(
+	if (ResolvedWorld->SweepMultiByObjectType(
 				Hits,
 				RayOrigin,
 				RayEnd,
@@ -313,13 +453,13 @@ bool UPCGWorldRayHitData::SamplePoint(const FTransform& InTransform, const FBox&
 				CollisionShape.ToCollisionShape(),
 				CollisionQueryParams))
 	{
-		TOptional<FHitResult> HitResult = PCGWorldQueryHelpers::FilterRayHitResults(&QueryParams, OriginatingComponent, Hits, ActorFilter.GetCachedFilterActors());
+		TOptional<FHitResult> HitResult = PCGWorldQueryHelpers::FilterRayHitResults(QueryParams, GetOriginatingSource(), Hits, GetCachedFilteredObjects());
 
 		if (HitResult.IsSet())
 		{
 			const FHitResult& Hit = HitResult.GetValue();
 			OutPoint = FPCGPoint(PCGWorldQueryHelpers::GetOrthonormalImpactTransform(Hit), 1.0f, UPCGBlueprintHelpers::ComputeSeedFromPosition(Hit.Location));
-			PCGWorldQueryHelpers::ApplyRayHitMetadata(Hit, QueryParams, QueryParams.RayDirection, OutPoint.Transform, OutPoint.MetadataEntry, OutMetadata, World);
+			PCGWorldQueryHelpers::ApplyRayHitMetadata(Hit, QueryParams, QueryParams.RayDirection, OutPoint.Transform, OutPoint.MetadataEntry, OutMetadata, ResolvedWorld);
 			return true;
 		}
 	}
@@ -400,12 +540,11 @@ UPCGSpatialData* UPCGWorldRayHitData::CopyInternal(FPCGContext* Context) const
 
 	CopyBaseSurfaceData(NewData);
 
-	NewData->World = World;
-	NewData->OriginatingComponent = OriginatingComponent;
+	NewData->SetOriginatingSource(GetOriginatingSource());
 	NewData->Bounds = Bounds;
 	NewData->QueryParams = QueryParams;
 	NewData->QueryParams.Initialize();
-	NewData->ActorFilter.GetFilterActorsMutable() = ActorFilter.GetFilterActorsConst();
+	NewData->GetMutableFilteredObjects() = GetConstFilteredObjects();
 	NewData->CachedLandscapeLayerNames = CachedLandscapeLayerNames;
 
 	return NewData;
diff --git a/Engine/Plugins/PCG/Source/PCG/Private/Elements/PCGPathfindingElement.cpp b/Engine/Plugins/PCG/Source/PCG/Private/Elements/PCGPathfindingElement.cpp
--- a/Engine/Plugins/PCG/Source/PCG/Private/Elements/PCGPathfindingElement.cpp
+++ b/Engine/Plugins/PCG/Source/PCG/Private/Elements/PCGPathfindingElement.cpp
@@ -343,8 +343,8 @@ bool FPCGPathfindingElement::PrepareDataInternal(FPCGContext* InContext) const
 			}
 
 			// --- Gather filtered actors ---
-			TSet<TObjectKey<AActor>> CachedFilterActors;
-			if (!Settings->PathTraceParams.ExtractLoadedActorFiltersIfNeeded(FilterActorsData, CachedFilterActors, Context))
+			TSet<FObjectKey> CachedFilterObjects;
+			if (!Settings->PathTraceParams.ExtractLoadedObjectFiltersIfNeeded(FilterActorsData, CachedFilterObjects, Context))
 			{
 				return EPCGTimeSliceInitResult::NoOperation;
 			}
@@ -354,17 +354,17 @@ bool FPCGPathfindingElement::PrepareDataInternal(FPCGContext* InContext) const
 				FPCGWorldRaycastQueryParams PathTraceParams = Settings->PathTraceParams;
 				PathTraceParams.Initialize();
 
-				TWeakObjectPtr<UPCGComponent> OriginatingComponent = Cast<UPCGComponent>(Context->ExecutionSource.Get());
+				TWeakInterfacePtr<IPCGGraphExecutionSource> OriginatingSource = Context->ExecutionSource.Get();
 				FCollisionObjectQueryParams ObjectQueryParams(PathTraceParams.CollisionChannel);
 				FCollisionQueryParams Params;
 				Params.bTraceComplex = PathTraceParams.bTraceComplex;
 
-				LineTraceTest = [World, ObjectQueryParams = MoveTemp(ObjectQueryParams), Params = MoveTemp(Params), OriginatingComponent, PathTraceParams = MoveTemp(PathTraceParams), CachedFilterActors = MoveTemp(CachedFilterActors)](const FVector& StartPosition, const FVector& EndPosition) -> bool
+				LineTraceTest = [World, ObjectQueryParams = MoveTemp(ObjectQueryParams), Params = MoveTemp(Params), OriginatingSource, PathTraceParams = MoveTemp(PathTraceParams), CachedFilterObjects = MoveTemp(CachedFilterObjects)](const FVector& StartPosition, const FVector& EndPosition) -> bool
 				{
 					TArray<FHitResult> OutHits;
 					if (World->LineTraceMultiByObjectType(OutHits, StartPosition, EndPosition, ObjectQueryParams, Params))
 					{
-						TOptional<FHitResult> HitResult = PCGWorldQueryHelpers::FilterRayHitResults(&PathTraceParams, OriginatingComponent, OutHits, CachedFilterActors);
+						TOptional<FHitResult> HitResult = PCGWorldQueryHelpers::FilterRayHitResults(PathTraceParams, OriginatingSource, OutHits, CachedFilterObjects);
 						return !HitResult.IsSet();
 					}
 					else
diff --git a/Engine/Plugins/PCG/Source/PCG/Private/Elements/PCGWorldQuery.cpp b/Engine/Plugins/PCG/Source/PCG/Private/Elements/PCGWorldQuery.cpp
--- a/Engine/Plugins/PCG/Source/PCG/Private/Elements/PCGWorldQuery.cpp
+++ b/Engine/Plugins/PCG/Source/PCG/Private/Elements/PCGWorldQuery.cpp
@@ -19,7 +19,7 @@
 namespace PCGWorldQuery
 {
 	template <typename WorldQueryDataType>
-	void ExtractActorFiltersIfNeeded(const FPCGWorldCommonQueryParams& QueryParams, FPCGContext* Context, WorldQueryDataType* WorldQueryData)
+	void ExtractObjectFiltersIfNeeded(const FPCGWorldCommonQueryParams& QueryParams, FPCGContext* Context, WorldQueryDataType* WorldQueryData)
 	{
 		check(WorldQueryData);
 		
@@ -32,7 +32,7 @@ namespace PCGWorldQuery
 			}
 
 			const UPCGData* ActorFilterData = ActorFilterTaggedData.IsEmpty() ? nullptr : ActorFilterTaggedData[0].Data;
-			if (ActorFilterData && !QueryParams.ExtractActorFiltersIfNeeded(ActorFilterData, WorldQueryData->ActorFilter.GetFilterActorsMutable(), Context))
+			if (ActorFilterData && !QueryParams.ExtractObjectFiltersIfNeeded(ActorFilterData, WorldQueryData->GetMutableFilteredObjects(), Context))
 			{
 				PCGLog::LogWarningOnGraph(LOCTEXT("FailExtractActorFilters", "Failed to extract actor filters."));
 			}
@@ -95,15 +95,20 @@ bool FPCGWorldVolumetricQueryElement::ExecuteInternal(FPCGContext* Context) cons
 	// TODO: Add params pin + Apply param data overrides
 
 	check(Context->ExecutionSource.IsValid());
-	UWorld* World = Context->ExecutionSource->GetExecutionState().GetWorld();
+
+	if (!Context->ExecutionSource->GetExecutionState().GetWorld())
+	{
+		PCGLog::LogErrorOnGraph(LOCTEXT("NoWorld_VolumetricQuery", "Execution source does not have a valid world."), Context);
+		return true;
+	}
 
 	UPCGWorldVolumetricData* Data = FPCGContext::NewObject_AnyThread<UPCGWorldVolumetricData>(Context);
-	Data->Initialize(World);
+	Data->Initialize(Context->ExecutionSource.Get());
 	Data->QueryParams = QueryParams;
 	Data->QueryParams.Initialize();
-	Data->OriginatingComponent = Cast<UPCGComponent>(Context->ExecutionSource.Get());
+	Data->SetOriginatingSource(Context->ExecutionSource.Get());
 
-	PCGWorldQuery::ExtractActorFiltersIfNeeded(Settings->QueryParams, Context, Data);
+	PCGWorldQuery::ExtractObjectFiltersIfNeeded(Settings->QueryParams, Context, Data);
 
 	FPCGTaggedData& Output = Context->OutputData.TaggedData.Emplace_GetRef();
 	Output.Data = Data;
@@ -176,10 +181,16 @@ bool FPCGWorldRayHitQueryElement::ExecuteInternal(FPCGContext* Context) const
 	// TODO: Might want to revisit this for 3D partitioning. The reasoning here is that ray origin should be the same
 	// if we are partitioned or non-partitioned. But with 3D partitioning we might want something different.
 	// Since it is not yet widely used, we'll stick with same behavior for all cases.
-	UPCGComponent* SourceComponent = Cast<UPCGComponent>(Context->ExecutionSource.Get());
-	if (!SourceComponent)
+	IPCGGraphExecutionSource* ExecutionSource = Context->ExecutionSource.Get();
+	if (!ExecutionSource)
 	{
-		PCGLog::LogErrorOnGraph(LOCTEXT("NoSourceComponent", "Execution source is not a PCG Component."), Context);
+		return true;
+	}
+
+	UWorld* World = ExecutionSource->GetExecutionState().GetWorld();
+	if (!World)
+	{
+		PCGLog::LogErrorOnGraph(LOCTEXT("NoWorld_RayHitQuery", "Execution source does not have a valid world."), Context);
 		return true;
 	}
 
@@ -213,18 +224,15 @@ bool FPCGWorldRayHitQueryElement::ExecuteInternal(FPCGContext* Context) const
 		}
 	}
 
-	UWorld* World = SourceComponent->GetWorld();
-
 	UPCGWorldRayHitData* Data = FPCGContext::NewObject_AnyThread<UPCGWorldRayHitData>(Context);
 	Data->QueryParams = QueryParams;
 	Data->QueryParams.Initialize();
-	Data->OriginatingComponent = SourceComponent;
 	Data->CollisionShape = Settings->CollisionShape;
-	Data->Initialize(World, Transform, /*InBounds=*/FBox(EForceInit::ForceInit), LocalBounds);
+	Data->Initialize(Context->ExecutionSource.Get(), Transform, /*InBounds=*/FBox(EForceInit::ForceInit), LocalBounds);
 
 	if (QueryParams.bApplyMetadataFromLandscape && Data->Metadata)
 	{
-		UPCGSubsystem* PCGSubsystem = UPCGSubsystem::GetInstance(World);
+		UPCGSubsystem* PCGSubsystem = Cast<UPCGSubsystem>(ExecutionSource->GetExecutionState().GetSubsystem());
 		UPCGLandscapeCache* LandscapeCache = PCGSubsystem ? PCGSubsystem->GetLandscapeCache() : nullptr;
 
 		if (LandscapeCache)
@@ -240,14 +248,15 @@ bool FPCGWorldRayHitQueryElement::ExecuteInternal(FPCGContext* Context) const
 
 			if (Data->Bounds.IsValid)
 			{
-				BoundsCheck = [Data, SourceComponent](const AActor* OtherActor) -> bool
+				BoundsCheck = [Data, ExecutionSource](const AActor* OtherActor) -> bool
 				{
-					const FBox OtherActorBounds = OtherActor ? PCGHelpers::GetGridBounds(OtherActor, SourceComponent) : FBox(EForceInit::ForceInit);
+					// @todo_pcg: adapt to IPCGGraphExecutionSource (currently only for partitioning which isn't supported by other types than UPCGComponent)
+					const FBox OtherActorBounds = OtherActor ? PCGHelpers::GetGridBounds(OtherActor, Cast<UPCGComponent>(ExecutionSource)) : FBox(EForceInit::ForceInit);
 					return OtherActorBounds.Intersect(Data->Bounds);
 				};
 			}
 
-			TArray<AActor*> LandscapeActors = PCGActorSelector::FindActors(ActorSelector, SourceComponent, BoundsCheck, SelfIgnoreCheck);
+			TArray<AActor*> LandscapeActors = PCGActorSelector::FindActors(ActorSelector, ExecutionSource, BoundsCheck, SelfIgnoreCheck);
 			for (AActor* Landscape : LandscapeActors)
 			{
 				if (ALandscapeProxy* LandscapeProxy = Cast<ALandscapeProxy>(Landscape))
@@ -263,7 +272,7 @@ bool FPCGWorldRayHitQueryElement::ExecuteInternal(FPCGContext* Context) const
 		Data->QueryParams.bApplyMetadataFromLandscape = false;
 	}
 
-	PCGWorldQuery::ExtractActorFiltersIfNeeded(Settings->QueryParams, Context, Data);
+	PCGWorldQuery::ExtractObjectFiltersIfNeeded(Settings->QueryParams, Context, Data);
 
 	FPCGTaggedData& Output = Context->OutputData.TaggedData.Emplace_GetRef();
 	Output.Data = Data;
diff --git a/Engine/Plugins/PCG/Source/PCG/Private/Elements/PCGWorldRaycast.cpp b/Engine/Plugins/PCG/Source/PCG/Private/Elements/PCGWorldRaycast.cpp
--- a/Engine/Plugins/PCG/Source/PCG/Private/Elements/PCGWorldRaycast.cpp
+++ b/Engine/Plugins/PCG/Source/PCG/Private/Elements/PCGWorldRaycast.cpp
@@ -290,7 +290,7 @@ bool FPCGWorldRaycastElement::PrepareDataInternal(FPCGContext* InContext) const
 		}
 
 		// --- Gather filtered actors ---
-		if (!Settings->WorldQueryParams.ExtractLoadedActorFiltersIfNeeded(FilterActorsData, OutState.CachedFilterActors, Context))
+		if (!Settings->WorldQueryParams.ExtractLoadedObjectFiltersIfNeeded(FilterActorsData, OutState.CachedFilterObjects, Context))
 		{
 			return EPCGTimeSliceInitResult::NoOperation;
 		}
@@ -485,7 +485,7 @@ bool FPCGWorldRaycastElement::ExecuteInternal(FPCGContext* InContext) const
 				Settings->CollisionShape.ToCollisionShape(),
 				ExecState.CollisionQueryParams);
 
-			TOptional<FHitResult> HitResult = PCGWorldQueryHelpers::FilterRayHitResults(&WorldQueryParams, Cast<UPCGComponent>(Context->ExecutionSource.Get()), Hits, IterState.CachedFilterActors);
+			TOptional<FHitResult> HitResult = PCGWorldQueryHelpers::FilterRayHitResults(WorldQueryParams, Context->ExecutionSource.Get(), Hits, IterState.CachedFilterObjects);
 			if (HitResult.IsSet() && (Settings->bUnbounded || PCGHelpers::IsInsideBounds(ExecState.Bounds, HitResult.GetValue().ImpactPoint)))
 			{
 				HitResults.Add({ MoveTemp(HitResult), Origin, Direction });
diff --git a/Engine/Plugins/PCG/Source/PCG/Private/Helpers/PCGWorldQueryHelpers.cpp b/Engine/Plugins/PCG/Source/PCG/Private/Helpers/PCGWorldQueryHelpers.cpp
--- a/Engine/Plugins/PCG/Source/PCG/Private/Helpers/PCGWorldQueryHelpers.cpp
+++ b/Engine/Plugins/PCG/Source/PCG/Private/Helpers/PCGWorldQueryHelpers.cpp
@@ -40,13 +40,52 @@ namespace PCGWorldQueryHelpers
 		return FTransform(XAxis, YAxis, Hit.ImpactNormal, Hit.ImpactPoint);
 	}
 
+	bool FilterRayHitResult(const FPCGFilterHitResultParams& InParams, bool& OutPassesFilter)
+	{
+		const UPrimitiveComponent* HitComponent = InParams.HitResult.GetComponent();
+		if (!HitComponent)
+		{
+			// Not handled
+			return false;
+		}
+
+		OutPassesFilter = FilterCommonQueryResults(InParams.QueryParams, HitComponent, InParams.OriginatingSource, InParams.FilteredObjectReferences);
+
+		if (OutPassesFilter && InParams.QueryParams.bIgnoreBackfaceHits)
+		{
+			// If it's a landscape, we cull if the normal is negative in Z direction (landscape normal is always the +Z axis). If not, then we cull if the impact normal and the ray are headed in the same direction
+			if (InParams.HitResult.bStartPenetrating || (InParams.HitResult.GetActor() && InParams.HitResult.GetActor()->IsA<ALandscapeProxy>() && InParams.HitResult.ImpactNormal.Z < 0) || (InParams.HitResult.TraceEnd - InParams.HitResult.TraceStart).Dot(InParams.HitResult.ImpactNormal) > 0)
+			{
+				OutPassesFilter = false;
+			}
+		}
+
+		// Handled
+		return true;
+	}
+
+	bool FilterOverlapResult(const FPCGFilterOverlapResultParams& InParams, bool& OutPassesFilter)
+	{
+		const UPrimitiveComponent* OverlapComponent = InParams.OverlapResult.GetComponent();
+		if (!OverlapComponent)
+		{
+			// Not handled
+			return false;
+		}
+
+		OutPassesFilter = FilterCommonQueryResults(InParams.QueryParams, OverlapComponent, InParams.OriginatingSource, InParams.FilteredObjectReferences);
+
+		// Handled
+		return true;
+	}
+
 	bool FilterCommonQueryResults(
-		const FPCGWorldCommonQueryParams* QueryParams,
+		const FPCGWorldCommonQueryParams& QueryParams,
 		const UPrimitiveComponent* TriggeredComponent,
-		const TWeakObjectPtr<UPCGComponent> OriginatingComponent,
-		const TSet<TObjectKey<AActor>>& FilteredActorReferences)
+		const TWeakInterfacePtr<IPCGGraphExecutionSource> OriginatingSource,
+		const TSet<FObjectKey>& FilteredObjectReferences)
 	{
-		check(QueryParams && TriggeredComponent);
+		check(TriggeredComponent);
 
 		// Skip invisible walls / triggers / volumes
 		if (TriggeredComponent->IsA<UBrushComponent>())
@@ -55,7 +94,7 @@ namespace PCGWorldQueryHelpers
 		}
 
 		// Skip "No collision" type actors
-		if (!TriggeredComponent->IsQueryCollisionEnabled() || TriggeredComponent->GetCollisionResponseToChannel(QueryParams->CollisionChannel) != ECR_Block)
+		if (!TriggeredComponent->IsQueryCollisionEnabled() || TriggeredComponent->GetCollisionResponseToChannel(QueryParams.CollisionChannel) != ECR_Block)
 		{
 			return false;
 		}
@@ -67,21 +106,22 @@ namespace PCGWorldQueryHelpers
 		}
 
 		// Optionally skip all PCG created objects
-		if (QueryParams->bIgnorePCGHits && (TriggeredComponent->ComponentHasTag(PCGHelpers::DefaultPCGTag) || (TriggeredComponent->GetOwner() && TriggeredComponent->GetOwner()->ActorHasTag(PCGHelpers::DefaultPCGActorTag))))
+		if (QueryParams.bIgnorePCGHits && (TriggeredComponent->ComponentHasTag(PCGHelpers::DefaultPCGTag) || (TriggeredComponent->GetOwner() && TriggeredComponent->GetOwner()->ActorHasTag(PCGHelpers::DefaultPCGActorTag))))
 		{
 			return false;
 		}
 
 		// Skip self-generated PCG objects optionally
-		if (QueryParams->bIgnoreSelfHits && OriginatingComponent.IsValid() && TriggeredComponent->ComponentTags.Contains(OriginatingComponent->GetFName()))
+		const UObject* OriginatingObject = Cast<UObject>(OriginatingSource.Get());
+		if (QueryParams.bIgnoreSelfHits && OriginatingObject && TriggeredComponent->ComponentTags.Contains(OriginatingObject->GetFName()))
 		{
 			return false;
 		}
 
 		// Additional filter as provided in the QueryParams base class
-		if (QueryParams->ActorTagFilter != EPCGWorldQueryFilter::None
-			|| QueryParams->ActorClassFilter != EPCGWorldQueryFilter::None
-			|| QueryParams->ActorFilterFromInput != EPCGWorldQueryFilter::None)
+		if (QueryParams.ActorTagFilter != EPCGWorldQueryFilter::None
+			|| QueryParams.ActorClassFilter != EPCGWorldQueryFilter::None
+			|| QueryParams.ActorFilterFromInput != EPCGWorldQueryFilter::None)
 		{
 			AActor* Actor = TriggeredComponent->GetOwner();
 			if (!Actor)
@@ -105,23 +145,23 @@ namespace PCGWorldQueryHelpers
 				bHasIncludeFilter |= bIsIncludeFilter;
 			};
 
-			if (QueryParams->ActorTagFilter != EPCGWorldQueryFilter::None)
+			if (QueryParams.ActorTagFilter != EPCGWorldQueryFilter::None)
 			{
-				const bool bFoundMatch = Algo::AnyOf(Actor->Tags, [QueryParams](const FName Tag) { return QueryParams->ParsedActorTagsList.Contains(Tag); });
-				UpdateIncludeExcludeStates(QueryParams->ActorTagFilter, bFoundMatch);
+				const bool bFoundMatch = Algo::AnyOf(Actor->Tags, [QueryParams](const FName Tag) { return QueryParams.ParsedActorTagsList.Contains(Tag); });
+				UpdateIncludeExcludeStates(QueryParams.ActorTagFilter, bFoundMatch);
 			}
 
 			// No need to check if we already know we are going to discard it.
-			if (!bHardExclude && QueryParams->ActorClassFilter != EPCGWorldQueryFilter::None)
+			if (!bHardExclude && QueryParams.ActorClassFilter != EPCGWorldQueryFilter::None)
 			{
-				const bool bFoundMatch = Actor->GetClass() && Actor->GetClass()->IsChildOf(QueryParams->ActorClass);
-				UpdateIncludeExcludeStates(QueryParams->ActorClassFilter, bFoundMatch);
+				const bool bFoundMatch = Actor->GetClass() && Actor->GetClass()->IsChildOf(QueryParams.ActorClass);
+				UpdateIncludeExcludeStates(QueryParams.ActorClassFilter, bFoundMatch);
 			}
 
-			if (!bHardExclude && QueryParams->ActorFilterFromInput != EPCGWorldQueryFilter::None)
+			if (!bHardExclude && QueryParams.ActorFilterFromInput != EPCGWorldQueryFilter::None)
 			{
-				const bool bFoundMatch = FilteredActorReferences.Contains(TObjectKey<AActor>(Actor));
-				UpdateIncludeExcludeStates(QueryParams->ActorFilterFromInput, bFoundMatch);
+				const bool bFoundMatch = FilteredObjectReferences.Contains(FObjectKey(Actor));
+				UpdateIncludeExcludeStates(QueryParams.ActorFilterFromInput, bFoundMatch);
 			}
 
 			if (bHardExclude || (!bSoftInclude && bHasIncludeFilter))
@@ -131,7 +171,7 @@ namespace PCGWorldQueryHelpers
 		}
 
 		// Landscape or not, include it if this is true
-		if (QueryParams->SelectLandscapeHits == EPCGWorldQuerySelectLandscapeHits::Include)
+		if (QueryParams.SelectLandscapeHits == EPCGWorldQuerySelectLandscapeHits::Include)
 		{
 			return true;
 		}
@@ -139,8 +179,8 @@ namespace PCGWorldQueryHelpers
 		bool bTriggeredOnLandscape = TriggeredComponent->GetOwner() && TriggeredComponent->GetOwner()->IsA<ALandscapeProxy>();
 
 		// If excluding, skip. If requiring, and it's not a landscape, skip.
-		if ((bTriggeredOnLandscape && QueryParams->SelectLandscapeHits == EPCGWorldQuerySelectLandscapeHits::Exclude) ||
-			(!bTriggeredOnLandscape && QueryParams->SelectLandscapeHits == EPCGWorldQuerySelectLandscapeHits::Require))
+		if ((bTriggeredOnLandscape && QueryParams.SelectLandscapeHits == EPCGWorldQuerySelectLandscapeHits::Exclude) ||
+			(!bTriggeredOnLandscape && QueryParams.SelectLandscapeHits == EPCGWorldQuerySelectLandscapeHits::Require))
 		{
 			return false;
 		}
@@ -149,56 +189,42 @@ namespace PCGWorldQueryHelpers
 	}
 
 	TOptional<FHitResult> FilterRayHitResults(
-		const FPCGWorldRaycastQueryParams* QueryParams,
-		const TWeakObjectPtr<UPCGComponent> OriginatingComponent,
+		const FPCGWorldRaycastQueryParams& QueryParams,
+		const TWeakInterfacePtr<IPCGGraphExecutionSource> OriginatingSource,
 		const TArray<FHitResult>& HitResults,
-		const TSet<TObjectKey<AActor>>& FilteredActorReferences)
+		const TSet<FObjectKey>& FilteredObjectReferences)
 	{
-		check(QueryParams);
+		const FPCGPhysicsRegistry& PhysicsRegistry = FPCGModule::GetConstPhysicsRegistry();
 
 		for (const FHitResult& Hit : HitResults)
 		{
-			const UPrimitiveComponent* HitComponent = Hit.GetComponent();
+			const FPCGFilterHitResultParams FilterHitParams(Hit, QueryParams, OriginatingSource, FilteredObjectReferences);
 
-			if (!FilterCommonQueryResults(QueryParams, HitComponent, OriginatingComponent, FilteredActorReferences))
+			if (PhysicsRegistry.FilterHitResult(FilterHitParams))
 			{
-				continue;
+				return Hit;
 			}
-
-			// Optionally skip backface hits
-			if (QueryParams->bIgnoreBackfaceHits)
-			{
-				// If it's a landscape, we cull if the normal is negative in Z direction (landscape normal is always the +Z axis). If not, then we cull if the impact normal and the ray are headed in the same direction
-				if (Hit.bStartPenetrating || (Hit.GetActor() && Hit.GetActor()->IsA<ALandscapeProxy>() && Hit.ImpactNormal.Z < 0) || (Hit.TraceEnd - Hit.TraceStart).Dot(Hit.ImpactNormal) > 0)
-				{
-					continue;
-				}
-			}
-
-			return Hit;
 		}
 
 		return {};
 	}
 
 	TOptional<FOverlapResult> FilterOverlapResults(
-		const FPCGWorldVolumetricQueryParams* QueryParams,
-		const TWeakObjectPtr<UPCGComponent> OriginatingComponent,
+		const FPCGWorldVolumetricQueryParams& QueryParams,
+		const TWeakInterfacePtr<IPCGGraphExecutionSource> OriginatingSource,
 		const TArray<FOverlapResult>& OverlapResults,
-		const TSet<TObjectKey<AActor>>& FilteredActorReferences)
+		const TSet<FObjectKey>& FilteredObjectReferences)
 	{
-		check(QueryParams);
+		const FPCGPhysicsRegistry& PhysicsRegistry = FPCGModule::GetConstPhysicsRegistry();
 
 		for (const FOverlapResult& Overlap : OverlapResults)
 		{
-			const UPrimitiveComponent* OverlappedComponent = Overlap.GetComponent();
+			const FPCGFilterOverlapResultParams FilterOverlapParams(Overlap, QueryParams, OriginatingSource, FilteredObjectReferences);
 
-			if (!FilterCommonQueryResults(QueryParams, OverlappedComponent, OriginatingComponent, FilteredActorReferences))
+			if (PhysicsRegistry.FilterOverlapResult(FilterOverlapParams))
 			{
-				continue;
+				return Overlap;
 			}
-
-			return Overlap;
 		}
 
 		return {};
@@ -264,86 +290,68 @@ namespace PCGWorldQueryHelpers
 		return false;
 	}
 
-	// TODO: Add an option to create and apply ranges
-	bool ApplyRayHitMetadata(
-		const FHitResult& HitResult,
-		const FPCGWorldRaycastQueryParams& QueryParams,
-		const FVector& RayDirection,
-		const FTransform& InTransform, 
-		int64& OutMetadataEntry,
-		UPCGMetadata* OutMetadata,
-		TWeakObjectPtr<UWorld> World)
+	template<class T>
+	bool ApplyAttribute(const FPCGApplyHitResultAttributesParams& InParams, FName AttributeName, const T& Value, bool bShouldApply = true)
 	{
-		if (!OutMetadata)
+		if (!bShouldApply)
 		{
-			return false;
+			return true;
 		}
 
-		auto ApplyAttribute = [&OutMetadataEntry, OutMetadata]<typename Type>(FName AttributeName, const Type& Value, bool bShouldApply = true)
+		if (FPCGMetadataAttribute<T>* Attribute = InParams.Metadata.FindOrCreateAttribute<T>(AttributeName))
 		{
-			if (!bShouldApply)
-			{
-				return true;
-			}
+			InParams.Metadata.InitializeOnSet(InParams.MetadataEntry);
+			Attribute->SetValue(InParams.MetadataEntry, Value);
+			return true;
+		}
 
-			if (FPCGMetadataAttribute<Type>* Attribute = OutMetadata->FindOrCreateAttribute<Type>(AttributeName))
-			{
-				OutMetadata->InitializeOnSet(OutMetadataEntry);
-				Attribute->SetValue(OutMetadataEntry, Value);
-				return true;
-			}
+		return false;
+	};
 
-			return false;
-		};
+	bool ApplyHitResultAttributes(const FPCGApplyHitResultAttributesParams& InParams, bool& bOutSuccess)
+	{
+		bOutSuccess = true;
 
-		const FHitResult& Hit = HitResult;
+		AActor* HitActor = InParams.HitResult.GetActor();
+		UPrimitiveComponent* HitComponent = InParams.HitResult.GetComponent();
 
-		FVector ReflectionVector = FVector::ZeroVector;
-		if (QueryParams.bGetReflection)
+		if (!HitActor && !HitComponent)
 		{
-			ReflectionVector = RayDirection - 2.0 * (FVector(Hit.ImpactNormal) | RayDirection) * FVector(Hit.ImpactNormal);
-			ReflectionVector.Normalize();
+			// Unhandled
+			return false;
 		}
 
-		bool bResult = true;
-		// Note: The T/F Impact attribute is true by default, so no need to set it directly.
-		bResult &= ApplyAttribute(PCGWorldQueryConstants::ImpactPointAttribute, FVector(Hit.ImpactPoint), QueryParams.bGetImpactPoint);
-		bResult &= ApplyAttribute(PCGWorldQueryConstants::ImpactNormalAttribute, FVector(Hit.ImpactNormal), QueryParams.bGetImpactNormal);
-		bResult &= ApplyAttribute(PCGWorldQueryConstants::ImpactReflectionAttribute, ReflectionVector, QueryParams.bGetReflection);
-		bResult &= ApplyAttribute(PCGWorldQueryConstants::ImpactDistanceAttribute, (Hit.ImpactPoint - Hit.TraceStart).Length(), QueryParams.bGetDistance);
-		bResult &= ApplyAttribute(PCGPointDataConstants::ActorReferenceAttribute, FSoftObjectPath(Hit.GetActor()), QueryParams.bGetReferenceToActorHit);
-		bResult &= ApplyAttribute(PCGWorldQueryConstants::PhysicalMaterialReferenceAttribute, FSoftObjectPath(Hit.PhysMaterial.Get()), QueryParams.bGetReferenceToPhysicalMaterial);
-		bResult &= ApplyAttribute(PCGWorldQueryConstants::ElementIndexAttribute, static_cast<int32>(Hit.ElementIndex), QueryParams.bGetElementIndex);
+		bOutSuccess &= ApplyAttribute(InParams, PCGPointDataConstants::ActorReferenceAttribute, FSoftObjectPath(HitActor), InParams.QueryParams.bGetReferenceToActorHit);
 
 		// Handle landscape related hit result
-		if (ALandscapeProxy* Landscape = Cast<ALandscapeProxy>(Hit.GetActor()))
+		if (ALandscapeProxy* Landscape = Cast<ALandscapeProxy>(HitActor))
 		{
 			if (const UMaterialInterface* RenderMaterial = Landscape->GetLandscapeMaterial())
 			{
-				bResult &= ApplyAttribute(PCGWorldQueryConstants::RenderMaterialReferenceAttribute, FSoftObjectPath(RenderMaterial), QueryParams.bGetReferenceToRenderMaterial);
+				bOutSuccess &= ApplyAttribute(InParams, PCGWorldQueryConstants::RenderMaterialReferenceAttribute, FSoftObjectPath(RenderMaterial), InParams.QueryParams.bGetReferenceToRenderMaterial);
 			}
 
-			if (QueryParams.bApplyMetadataFromLandscape && World.IsValid() && World->GetSubsystem<UPCGSubsystem>())
+			if (InParams.QueryParams.bApplyMetadataFromLandscape && InParams.World && InParams.World->GetSubsystem<UPCGSubsystem>())
 			{
-				if (UPCGLandscapeCache* LandscapeCache = World->GetSubsystem<UPCGSubsystem>()->GetLandscapeCache())
+				if (UPCGLandscapeCache* LandscapeCache = InParams.World->GetSubsystem<UPCGSubsystem>()->GetLandscapeCache())
 				{
 					const TArray<FName> Layers = LandscapeCache->GetLayerNames(Landscape);
 					for (const FName& Layer : Layers)
 					{
-						OutMetadata->FindOrCreateAttribute<float>(Layer);
+						InParams.Metadata.FindOrCreateAttribute<float>(Layer);
 					}
 
-					LandscapeCache->SampleMetadataOnPoint(Landscape, InTransform, OutMetadataEntry, OutMetadata);
+					LandscapeCache->SampleMetadataOnPoint(Landscape, InParams.Transform, InParams.MetadataEntry, &InParams.Metadata);
 				}
 			}
 		}
 		// Handle primitive component related hit result
-		else if (UPrimitiveComponent* HitComponent = Hit.GetComponent())
+		else if (HitComponent)
 		{
-			if (QueryParams.bGetLocalImpactPoint)
+			if (InParams.QueryParams.bGetLocalImpactPoint)
 			{
-				const FVector LocalHitLocation = HitComponent->GetComponentToWorld().InverseTransformPosition(Hit.ImpactPoint);
-				bResult &= ApplyAttribute(PCGWorldQueryConstants::LocalImpactPointAttribute, LocalHitLocation);
+				const FVector LocalHitLocation = HitComponent->GetComponentToWorld().InverseTransformPosition(InParams.HitResult.ImpactPoint);
+				bOutSuccess &= ApplyAttribute(InParams, PCGWorldQueryConstants::LocalImpactPointAttribute, LocalHitLocation);
 			}
 
 			// Handle static mesh component related hit result
@@ -351,24 +359,24 @@ namespace PCGWorldQueryHelpers
 			{
 				if (UStaticMesh* StaticMesh = StaticMeshComponent->GetStaticMesh())
 				{
-					bResult &= ApplyAttribute(PCGWorldQueryConstants::StaticMeshReferenceAttribute, FSoftObjectPath(StaticMesh), QueryParams.bGetReferenceToStaticMesh);
+					bOutSuccess &= ApplyAttribute(InParams, PCGWorldQueryConstants::StaticMeshReferenceAttribute, FSoftObjectPath(StaticMesh), InParams.QueryParams.bGetReferenceToStaticMesh);
 				}
 
 				// Implementation note: FaceIndex will return -1 if complex queries are disabled
-				bResult &= ApplyAttribute(PCGWorldQueryConstants::FaceIndexAttribute, Hit.FaceIndex, QueryParams.bTraceComplex && QueryParams.bGetFaceIndex);
+				bOutSuccess &= ApplyAttribute(InParams, PCGWorldQueryConstants::FaceIndexAttribute, InParams.HitResult.FaceIndex, InParams.QueryParams.bTraceComplex && InParams.QueryParams.bGetFaceIndex);
 
-				if (QueryParams.bTraceComplex)
+				if (InParams.QueryParams.bTraceComplex)
 				{
 					int32 SectionIndex = INDEX_NONE;
 					int32 MaterialIndex = INDEX_NONE;
 					UMaterialInterface* RenderMaterial = nullptr;
 
 					// Attempt to get these values from LOD0, only if needed.
-					if (QueryParams.bGetSectionIndex
-						|| QueryParams.bGetRenderMaterialIndex
-						|| (QueryParams.bGetReferenceToRenderMaterial && !QueryParams.bUseRenderMaterialIndex))
+					if (InParams.QueryParams.bGetSectionIndex
+						|| InParams.QueryParams.bGetRenderMaterialIndex
+						|| (InParams.QueryParams.bGetReferenceToRenderMaterial && !InParams.QueryParams.bUseRenderMaterialIndex))
 					{
-						if (Hit.FaceIndex != INDEX_NONE
+						if (InParams.HitResult.FaceIndex != INDEX_NONE
 							&& StaticMeshComponent->GetStaticMesh()
 							&& StaticMeshComponent->GetStaticMesh()->GetRenderData()
 							&& !StaticMeshComponent->GetStaticMesh()->GetRenderData()->LODResources.IsEmpty())
@@ -377,7 +385,7 @@ namespace PCGWorldQueryHelpers
 							for (int32 CurrentIndex = 0; CurrentIndex < LOD.Sections.Num(); ++CurrentIndex)
 							{
 								const FStaticMeshSection& Section = LOD.Sections[CurrentIndex];
-								const uint32 HitFaceIndex = static_cast<uint32>(Hit.FaceIndex);
+								const uint32 HitFaceIndex = static_cast<uint32>(InParams.HitResult.FaceIndex);
 
 								// Find the section by iterating through the indices 3 at a time.
 								const uint32 FirstFaceIndex = Section.FirstIndex / 3u;
@@ -394,40 +402,81 @@ namespace PCGWorldQueryHelpers
 						}
 					}
 
-					if (QueryParams.bUseRenderMaterialIndex)  // Use explicit Render Material Index
+					if (InParams.QueryParams.bUseRenderMaterialIndex)  // Use explicit Render Material Index
 					{
-						RenderMaterial = HitComponent->GetMaterial(QueryParams.RenderMaterialIndex);
-						bResult &= ApplyAttribute(PCGWorldQueryConstants::RenderMaterialIndexAttribute, QueryParams.RenderMaterialIndex, QueryParams.bGetRenderMaterialIndex);
+						RenderMaterial = HitComponent->GetMaterial(InParams.QueryParams.RenderMaterialIndex);
+						bOutSuccess &= ApplyAttribute(InParams, PCGWorldQueryConstants::RenderMaterialIndexAttribute, InParams.QueryParams.RenderMaterialIndex, InParams.QueryParams.bGetRenderMaterialIndex);
 					}
 					else
 					{
-						RenderMaterial = HitComponent->GetMaterialFromCollisionFaceIndex(Hit.FaceIndex, SectionIndex);
-						bResult &= ApplyAttribute(PCGWorldQueryConstants::RenderMaterialIndexAttribute, MaterialIndex, QueryParams.bGetRenderMaterialIndex);
+						RenderMaterial = HitComponent->GetMaterialFromCollisionFaceIndex(InParams.HitResult.FaceIndex, SectionIndex);
+						bOutSuccess &= ApplyAttribute(InParams, PCGWorldQueryConstants::RenderMaterialIndexAttribute, MaterialIndex, InParams.QueryParams.bGetRenderMaterialIndex);
 					}
 
-					bResult &= ApplyAttribute(PCGWorldQueryConstants::SectionIndexAttribute, SectionIndex, QueryParams.bGetSectionIndex);
-					bResult &= ApplyAttribute(PCGWorldQueryConstants::RenderMaterialReferenceAttribute, FSoftObjectPath(RenderMaterial), QueryParams.bGetReferenceToRenderMaterial);
+					bOutSuccess &= ApplyAttribute(InParams, PCGWorldQueryConstants::SectionIndexAttribute, SectionIndex, InParams.QueryParams.bGetSectionIndex);
+					bOutSuccess &= ApplyAttribute(InParams, PCGWorldQueryConstants::RenderMaterialReferenceAttribute, FSoftObjectPath(RenderMaterial), InParams.QueryParams.bGetReferenceToRenderMaterial);
 
-					if (QueryParams.bGetUVCoords && UPhysicsSettings::Get()->bSupportUVFromHitResults)
+					if (InParams.QueryParams.bGetUVCoords && UPhysicsSettings::Get()->bSupportUVFromHitResults)
 					{
 						FVector2D UVCoords;
-						if (UGameplayStatics::FindCollisionUV(Hit, QueryParams.UVChannel, UVCoords))
+						if (UGameplayStatics::FindCollisionUV(InParams.HitResult, InParams.QueryParams.UVChannel, UVCoords))
 						{
-							bResult &= ApplyAttribute(PCGWorldQueryConstants::UVCoordAttribute, UVCoords);
+							bOutSuccess &= ApplyAttribute(InParams, PCGWorldQueryConstants::UVCoordAttribute, UVCoords);
 						}
 					}
 				}
 				else // @todo_pcg: Address other primitive types here
 				{
-					if (QueryParams.bGetReferenceToRenderMaterial)
+					if (InParams.QueryParams.bGetReferenceToRenderMaterial)
 					{
-						const UMaterialInterface* RenderMaterial = HitComponent->GetMaterial(QueryParams.RenderMaterialIndex);
-						bResult &= RenderMaterial && ApplyAttribute(PCGWorldQueryConstants::RenderMaterialReferenceAttribute, FSoftObjectPath(RenderMaterial));
+						const UMaterialInterface* RenderMaterial = HitComponent->GetMaterial(InParams.QueryParams.RenderMaterialIndex);
+						bOutSuccess &= RenderMaterial && ApplyAttribute(InParams, PCGWorldQueryConstants::RenderMaterialReferenceAttribute, FSoftObjectPath(RenderMaterial));
 					}
 				}
 			}
 		}
 
+		// Handled
+		return true;
+	}
+
+	// TODO: Add an option to create and apply ranges
+	bool ApplyRayHitMetadata(
+		const FHitResult& HitResult,
+		const FPCGWorldRaycastQueryParams& QueryParams,
+		const FVector& RayDirection,
+		const FTransform& InTransform, 
+		int64& OutMetadataEntry,
+		UPCGMetadata* OutMetadata,
+		TWeakObjectPtr<UWorld> World)
+	{
+		if (!OutMetadata)
+		{
+			return false;
+		}
+
+		FVector ReflectionVector = FVector::ZeroVector;
+		if (QueryParams.bGetReflection)
+		{
+			ReflectionVector = RayDirection - 2.0 * (FVector(HitResult.ImpactNormal) | RayDirection) * FVector(HitResult.ImpactNormal);
+			ReflectionVector.Normalize();
+		}
+
+		bool bResult = true;
+		
+		FPCGApplyHitResultAttributesParams Params(HitResult, QueryParams, RayDirection, InTransform, World.Get(), OutMetadataEntry, *OutMetadata);
+		
+		// Note: The T/F Impact attribute is true by default, so no need to set it directly.
+		bResult &= ApplyAttribute(Params, PCGWorldQueryConstants::ImpactPointAttribute, FVector(HitResult.ImpactPoint), QueryParams.bGetImpactPoint);
+		bResult &= ApplyAttribute(Params, PCGWorldQueryConstants::ImpactNormalAttribute, FVector(HitResult.ImpactNormal), QueryParams.bGetImpactNormal);
+		bResult &= ApplyAttribute(Params, PCGWorldQueryConstants::ImpactReflectionAttribute, ReflectionVector, QueryParams.bGetReflection);
+		bResult &= ApplyAttribute(Params, PCGWorldQueryConstants::ImpactDistanceAttribute, (HitResult.ImpactPoint - HitResult.TraceStart).Length(), QueryParams.bGetDistance);
+		bResult &= ApplyAttribute(Params, PCGWorldQueryConstants::PhysicalMaterialReferenceAttribute, FSoftObjectPath(HitResult.PhysMaterial.Get()), QueryParams.bGetReferenceToPhysicalMaterial);
+		bResult &= ApplyAttribute(Params, PCGWorldQueryConstants::ElementIndexAttribute, static_cast<int32>(HitResult.ElementIndex), QueryParams.bGetElementIndex);
+
+		// Extra Apply based on hit type
+		bResult &= FPCGModule::GetConstPhysicsRegistry().ApplyHitResultAttributes(Params);
+
 		return bResult;
 	}
 }
diff --git a/Engine/Plugins/PCG/Source/PCG/Private/PCGModule.cpp b/Engine/Plugins/PCG/Source/PCG/Private/PCGModule.cpp
--- a/Engine/Plugins/PCG/Source/PCG/Private/PCGModule.cpp
+++ b/Engine/Plugins/PCG/Source/PCG/Private/PCGModule.cpp
@@ -19,6 +19,7 @@
 #include "Data/DataView/PCGDataViewNativeJsonConverters.h"
 #include "Data/DataView/PCGDataViewNativePropertySelectors.h"
 #include "Helpers/PCGActorHelpers.h"
+#include "Helpers/PCGWorldQueryHelpers.h"
 
 #include "Engine/Engine.h"
 #include "Interfaces/IPluginManager.h"
@@ -125,6 +126,10 @@ namespace PCGModule
 		}
 #endif
 	};
+
+	constexpr FGuid PCGHitResultFilter_PrimitiveComponent(0xF6ABE3BA, 0x41C34DA0, 0x86D96E82, 0x313893BC);
+	constexpr FGuid PCGOverlapResultFilter_PrimitiveComponent(0x719F00B8, 0x729E4EF0, 0xA45FF00F, 0x3C8B75FF);
+	constexpr FGuid PCGApplyHitResultAttributes_PrimitiveComponent(0xA1F51F3D, 0x0B78435C, 0x8AAC7FF6, 0xF1CF6582);
 }
 
 FPCGModule& FPCGModule::GetPCGModuleChecked()
@@ -194,6 +199,10 @@ void FPCGModule::StartupModule()
 
 	ChangeTrackingRegistry.RegisterGetExecutionSourcesFromSelectionKey(AActor::StaticClass(), FPCGGetExecutionSourcesFromSelectionKey::CreateStatic(&UPCGActorHelpers::GetExecutionSourcesFromSelectionKey));
 	
+	PhysicsRegistry.RegisterHitResultFilter(PCGModule::PCGHitResultFilter_PrimitiveComponent, FPCGFilterHitResult::CreateStatic(&PCGWorldQueryHelpers::FilterRayHitResult));
+	PhysicsRegistry.RegisterOverlapResultFilter(PCGModule::PCGOverlapResultFilter_PrimitiveComponent, FPCGFilterOverlapResult::CreateStatic(&PCGWorldQueryHelpers::FilterOverlapResult));
+	PhysicsRegistry.RegisterApplyHitResultAttributes(PCGModule::PCGApplyHitResultAttributes_PrimitiveComponent, FPCGApplyHitResultAttributes::CreateStatic(&PCGWorldQueryHelpers::ApplyHitResultAttributes));
+
 #if WITH_EDITOR
 	ObjectHashFactory.RegisterObjectHashContextFactory(UPCGGraphInterface::StaticClass(), FPCGOnCreateObjectHashContext::CreateStatic(&FPCGGraphHashContext::MakeInstance));
 	ObjectHashFactory.RegisterObjectHashContextFactory(UPCGSettingsInterface::StaticClass(), FPCGOnCreateObjectHashContext::CreateStatic(&FPCGSettingsHashContext::MakeInstance));
@@ -240,6 +249,10 @@ void FPCGModule::PreExit()
 
 	ChangeTrackingRegistry.UnregisterGetExecutionSourcesFromSelectionKey(AActor::StaticClass());
 
+	PhysicsRegistry.UnregisterHitResultFilter(PCGModule::PCGHitResultFilter_PrimitiveComponent);
+	PhysicsRegistry.UnregisterOverlapResultFilter(PCGModule::PCGOverlapResultFilter_PrimitiveComponent);
+	PhysicsRegistry.UnregisterApplyHitResultAttributes(PCGModule::PCGApplyHitResultAttributes_PrimitiveComponent);
+
 #if WITH_EDITOR
 	DeregisterNativeElementDeterminismTests();
 
diff --git a/Engine/Plugins/PCG/Source/PCG/Private/Physics/PCGPhysicsRegistry.cpp b/Engine/Plugins/PCG/Source/PCG/Private/Physics/PCGPhysicsRegistry.cpp
new file mode 100644
--- /dev/null
+++ b/Engine/Plugins/PCG/Source/PCG/Private/Physics/PCGPhysicsRegistry.cpp
@@ -0,0 +1,69 @@
+// Copyright Epic Games, Inc. All Rights Reserved.
+
+#include "Physics/PCGPhysicsRegistry.h"
+
+namespace PCGPhysicsRegistry
+{
+	template<typename DelegateType, typename ParamType>
+	bool CallFunc(const TMap<FGuid, DelegateType>& InFuncs, const ParamType& InParams, bool bDefaultReturnValue)
+	{
+		for (const auto& [FuncID, Func] : InFuncs)
+		{
+			bool bReturnValue = bDefaultReturnValue;
+			if (Func.IsBound() && Func.Execute(InParams, bReturnValue))
+			{
+				return bReturnValue;
+			}
+		}
+
+		return bDefaultReturnValue;
+	}
+}
+
+void FPCGPhysicsRegistry::RegisterHitResultFilter(const FGuid& InID, FPCGFilterHitResult InFilter)
+{
+	check(!HitResultFilters.Contains(InID));
+	HitResultFilters.Add(InID, MoveTemp(InFilter));
+}
+
+void FPCGPhysicsRegistry::UnregisterHitResultFilter(const FGuid& InID)
+{
+	HitResultFilters.Remove(InID);
+}
+
+void FPCGPhysicsRegistry::RegisterOverlapResultFilter(const FGuid& InID, FPCGFilterOverlapResult InFilter)
+{
+	check(!OverlapResultFilters.Contains(InID));
+	OverlapResultFilters.Add(InID, MoveTemp(InFilter));
+}
+
+void FPCGPhysicsRegistry::UnregisterOverlapResultFilter(const FGuid& InID)
+{
+	OverlapResultFilters.Remove(InID);
+}
+
+void FPCGPhysicsRegistry::RegisterApplyHitResultAttributes(const FGuid& InID, FPCGApplyHitResultAttributes InApplyHitResultAttributes)
+{
+	check(!ApplyHitResultAttributesFuncs.Contains(InID));
+	ApplyHitResultAttributesFuncs.Add(InID, MoveTemp(InApplyHitResultAttributes));
+}
+
+void FPCGPhysicsRegistry::UnregisterApplyHitResultAttributes(const FGuid& InID)
+{
+	ApplyHitResultAttributesFuncs.Remove(InID);
+}
+
+bool FPCGPhysicsRegistry::FilterOverlapResult(const FPCGFilterOverlapResultParams& InParams) const
+{
+	return PCGPhysicsRegistry::CallFunc(OverlapResultFilters, InParams, /*bDefaultReturnValue=*/false);
+}
+
+bool FPCGPhysicsRegistry::FilterHitResult(const FPCGFilterHitResultParams& InParams) const
+{
+	return PCGPhysicsRegistry::CallFunc(HitResultFilters, InParams, /*bDefaultReturnValue=*/false);
+}
+
+bool FPCGPhysicsRegistry::ApplyHitResultAttributes(const FPCGApplyHitResultAttributesParams& InParams) const
+{
+	return PCGPhysicsRegistry::CallFunc(ApplyHitResultAttributesFuncs, InParams, /*bDefaultReturnValue*/true);
+}
diff --git a/Engine/Plugins/PCG/Source/PCG/Public/Data/PCGWorldData.h b/Engine/Plugins/PCG/Source/PCG/Public/Data/PCGWorldData.h
--- a/Engine/Plugins/PCG/Source/PCG/Public/Data/PCGWorldData.h
+++ b/Engine/Plugins/PCG/Source/PCG/Public/Data/PCGWorldData.h
@@ -121,10 +121,16 @@ struct FPCGWorldCommonQueryParams
 	/** Utility function to add a filter pin with the right label and tooltip to the pin properties. */
 	PCG_API void AddFilterPinIfNeeded(TArray<FPCGPinProperties>& PinProperties) const;
 
-	/** Utility function to extract the Actor Filters from the incoming data using the selector in this struct. Return true if the extraction succeeded. */
+	/** Utility function to extract the Object Filters from the incoming data using the selector in this struct. Return true if the extraction succeeded. */
+	PCG_API bool ExtractObjectFiltersIfNeeded(const UPCGData* InData, TArray<FSoftObjectPath>& OutArray, FPCGContext* InContext = nullptr) const;
+
+	/** Utility function to extract the Object Filters (from loaded actors)  from the incoming data using the selector in this struct. Return true if the extraction succeeded. */
+	PCG_API bool ExtractLoadedObjectFiltersIfNeeded(const UPCGData* InData, TSet<FObjectKey>& OutSet, FPCGContext* InContext = nullptr) const;
+
+	UE_DEPRECATED(5.8, "Use ExtractObjectFiltersIfNeeded instead")
 	PCG_API bool ExtractActorFiltersIfNeeded(const UPCGData* InData, TArray<TSoftObjectPtr<AActor>>& OutArray, FPCGContext* InContext = nullptr) const;
 
-	/** Utility function to extract the Actor Filters (from loaded actors) from the incoming data using the selector in this struct. Return true if the extraction succeeded. */
+	UE_DEPRECATED(5.8, "Use ExtractLoadedObjectFiltersIfNeeded instead")
 	PCG_API bool ExtractLoadedActorFiltersIfNeeded(const UPCGData* InData, TSet<TObjectKey<AActor>>& OutSet, FPCGContext* InContext = nullptr) const;
 
 protected:
@@ -296,26 +302,29 @@ struct TStructOpsTypeTraits<FPCGWorldRayHitQueryParams> : public TStructOpsTypeT
 };
 
 USTRUCT()
-struct FPCGWorldQueryActorFilterCache
+struct UE_DEPRECATED(5.8, "Use FPCGWorldQueryFilterCache instead") FPCGWorldQueryActorFilterCache
 {
 	GENERATED_BODY();
-	
+
 public:
-	PCG_API const TSet<TObjectKey<AActor>>& GetCachedFilterActors() const;
-	TArray<TSoftObjectPtr<AActor>>& GetFilterActorsMutable() { CachedFilterActorsDirty = true; return FilterActors; }
-	const TArray<TSoftObjectPtr<AActor>>& GetFilterActorsConst() const { return FilterActors; }
-	
+	// No longer computed, use GetCachedFilteredObjects()
+	PCG_API const TSet<TObjectKey<AActor>>& GetCachedFilterActors() const { return CachedFilterActors; }
+	TArray<TSoftObjectPtr<AActor>>& GetFilterActorsMutable() { CachedFilterActorsDirty = true; return FilterActors_DEPRECATED; }
+	const TArray<TSoftObjectPtr<AActor>>& GetFilterActorsConst() const { return FilterActors_DEPRECATED; }
+
 private:
 	UPROPERTY()
-	TArray<TSoftObjectPtr<AActor>> FilterActors;
+	TArray<TSoftObjectPtr<AActor>> FilterActors_DEPRECATED;
 
 	// Transient data to cache filter actors
 	mutable TSet<TObjectKey<AActor>> CachedFilterActors;
-	mutable PCG::FLock CachedFilterActorsLock;
 	mutable bool CachedFilterActorsDirty = false;
+
+	friend class UPCGWorldVolumetricData;
+	friend class UPCGWorldRayHitData;
 };
 
-// Not copyable because of the lock, but FilterActors will be copied in the Copy Internal of the data.
+PRAGMA_DISABLE_DEPRECATION_WARNINGS
 template<>
 struct TStructOpsTypeTraits<FPCGWorldQueryActorFilterCache> : public TStructOpsTypeTraitsBase2<FPCGWorldQueryActorFilterCache>
 {
@@ -324,6 +333,7 @@ struct TStructOpsTypeTraits<FPCGWorldQueryActorFilterCache> : public TStructOpsT
 		WithCopy = false,
 	};
 };
+PRAGMA_ENABLE_DEPRECATION_WARNINGS
 
 /** Queries volume for presence of world collision or not. Can be used to voxelize environment. */
 UCLASS(MinimalAPI, BlueprintType, ClassGroup = (Procedural))
@@ -332,7 +342,19 @@ class UPCGWorldVolumetricData : public UPCGVolumeData
 	GENERATED_BODY()
 
 public:
-	PCG_API void Initialize(UWorld* InWorld, const FBox& InBounds = FBox(EForceInit::ForceInit));
+	PCG_API void Initialize(IPCGGraphExecutionSource* InOriginatingSource, const FBox& InBounds = FBox(EForceInit::ForceInit));
+
+	PCG_API IPCGGraphExecutionSource* GetOriginatingSource() const;
+	PCG_API void SetOriginatingSource(IPCGGraphExecutionSource* InExecutionSource);
+	PCG_API UWorld* GetWorld() const;
+	
+	TArray<FSoftObjectPath>& GetMutableFilteredObjects() { CachedFilteredObjectsDirty = true; return FilteredObjects; }
+	const TArray<FSoftObjectPath>& GetConstFilteredObjects() const { return FilteredObjects; }
+	PCG_API const TSet<FObjectKey>& GetCachedFilteredObjects() const;
+
+	//~Begin UObject interface
+	virtual void PostLoad() override;
+	//~End UObject interface
 
 	//~Begin UPCGSpatialData interface
 	virtual bool IsBounded() const override { return !!Bounds.IsValid; }
@@ -354,17 +376,38 @@ private:
 	const UPCGBasePointData* CreateBasePointData(FPCGContext* Context, const FBox& InBounds, TSubclassOf<UPCGBasePointData> PointDataClass) const;
 
 public:
+	UPROPERTY(BlueprintReadWrite, EditAnywhere, Category = Data, meta = (ShowOnlyInnerProperties))
+	FPCGWorldVolumetricQueryParams QueryParams;
+	
+private:
+	UPROPERTY()
+	FSoftObjectPath OriginatingObject;
+
 	UPROPERTY()
+	TArray<FSoftObjectPath> FilteredObjects;
+
+	// transient
+	mutable bool CachedFilteredObjectsDirty = false;
+	mutable PCG::FLock CachedFilteredObjectsLock;
+	mutable TSet<FObjectKey> CachedFilteredObjects;
+
+public:
+	UE_DEPRECATED(5.8, "Use Initialize taking a IPCGGraphExecutionSource parameter instead")
+	PCG_API void Initialize(UWorld* InWorld, const FBox& InBounds = FBox(EForceInit::ForceInit)) {}
+
+PRAGMA_DISABLE_DEPRECATION_WARNINGS
+	UE_DEPRECATED(5.8, "Use GetWorld() instead")
+	UPROPERTY(Transient)
 	TWeakObjectPtr<UWorld> World = nullptr;
 
+	UE_DEPRECATED(5.8, "Use OriginatingObject instead")
 	UPROPERTY()
 	TWeakObjectPtr<UPCGComponent> OriginatingComponent = nullptr;
 
-	UPROPERTY(BlueprintReadWrite, EditAnywhere, Category = Data, meta = (ShowOnlyInnerProperties))
-	FPCGWorldVolumetricQueryParams QueryParams;
-	
+	UE_DEPRECATED(5.8, "Use GetObjectFilter methods instead")
 	UPROPERTY()
 	FPCGWorldQueryActorFilterCache ActorFilter;
+PRAGMA_ENABLE_DEPRECATION_WARNINGS
 };
 
 /** Executes collision queries against world collision. */
@@ -374,7 +417,19 @@ class UPCGWorldRayHitData : public UPCGSurfaceData
 	GENERATED_BODY()
 
 public:
-	PCG_API void Initialize(UWorld* InWorld, const FTransform& InTransform, const FBox& InBounds = FBox(EForceInit::ForceInit), const FBox& InLocalBounds = FBox(EForceInit::ForceInit));
+	PCG_API void Initialize(IPCGGraphExecutionSource* InExecutionSource, const FTransform& InTransform, const FBox& InBounds = FBox(EForceInit::ForceInit), const FBox& InLocalBounds = FBox(EForceInit::ForceInit));
+
+	PCG_API IPCGGraphExecutionSource* GetOriginatingSource() const;
+	PCG_API void SetOriginatingSource(IPCGGraphExecutionSource* InExecutionSource);
+	PCG_API UWorld* GetWorld() const;
+	
+	TArray<FSoftObjectPath>& GetMutableFilteredObjects() { CachedFilteredObjectsDirty = true; return FilteredObjects; }
+	const TArray<FSoftObjectPath>& GetConstFilteredObjects() const { return FilteredObjects; }
+	PCG_API const TSet<FObjectKey>& GetCachedFilteredObjects() const;
+
+	//~Begin UObject interface
+	virtual void PostLoad() override;
+	//~End UObject interface
 
 	// ~Begin UPCGData interface
 	virtual void AddToCrc(FArchiveCrc32& Ar, bool bFullDataCrc) const override;
@@ -401,12 +456,6 @@ public:
 	virtual const UPCGPointArrayData* CreatePointArrayData(FPCGContext* Context, const FBox& InBounds) const override;
 	// ~End UPCGSpatialDataWithPointCache interface
 
-	UPROPERTY()
-	TWeakObjectPtr<UWorld> World = nullptr;
-
-	UPROPERTY()
-	TWeakObjectPtr<UPCGComponent> OriginatingComponent = nullptr;
-
 	UPROPERTY()
 	FBox Bounds = FBox(EForceInit::ForceInit);
 
@@ -417,10 +466,37 @@ public:
 	UPROPERTY(BlueprintReadWrite, EditAnywhere, Category = Data, meta = (ShowOnlyInnerProperties))
 	FPCGWorldRayHitQueryParams QueryParams;
 	
-	UPROPERTY()
-	FPCGWorldQueryActorFilterCache ActorFilter;
-	
 	/** Attributes related to landscape layers are added when we initialize the target metadata, so we need to cache them when this data gets initialized. */
 	UPROPERTY()
 	TSet<FName> CachedLandscapeLayerNames;
+
+private:
+	UPROPERTY()
+	FSoftObjectPath OriginatingObject = nullptr;
+
+	UPROPERTY()
+	TArray<FSoftObjectPath> FilteredObjects;
+
+	// transient
+	mutable bool CachedFilteredObjectsDirty = false;
+	mutable PCG::FLock CachedFilteredObjectsLock;
+	mutable TSet<FObjectKey> CachedFilteredObjects;
+
+public:
+	UE_DEPRECATED(5.8, "Use Initialize taking a IPCGGraphExecutionSource parameter instead")
+	PCG_API void Initialize(UWorld* InWorld, const FTransform& InTransform, const FBox& InBounds = FBox(EForceInit::ForceInit), const FBox& InLocalBounds = FBox(EForceInit::ForceInit)) {}
+
+PRAGMA_DISABLE_DEPRECATION_WARNINGS
+	UE_DEPRECATED(5.8, "Use GetWorld() instead")
+	UPROPERTY(Transient)
+	TWeakObjectPtr<UWorld> World = nullptr;
+
+	UE_DEPRECATED(5.8, "Use OriginatingObject instead")
+	UPROPERTY()
+	TWeakObjectPtr<UPCGComponent> OriginatingComponent = nullptr;
+
+	UE_DEPRECATED(5.8, "Use ObjectFilter instead")
+	UPROPERTY()
+	FPCGWorldQueryActorFilterCache ActorFilter;
+PRAGMA_ENABLE_DEPRECATION_WARNINGS
 };
diff --git a/Engine/Plugins/PCG/Source/PCG/Public/Elements/PCGWorldRaycast.h b/Engine/Plugins/PCG/Source/PCG/Public/Elements/PCGWorldRaycast.h
--- a/Engine/Plugins/PCG/Source/PCG/Public/Elements/PCGWorldRaycast.h
+++ b/Engine/Plugins/PCG/Source/PCG/Public/Elements/PCGWorldRaycast.h
@@ -109,9 +109,21 @@ namespace FPCGWorldRaycast
 
 	struct FIterationState
 	{
+		// Note: We need to explicitly disable warnings on these constructors/operators for clang to be happy with deprecated variables
+PRAGMA_DISABLE_DEPRECATION_WARNINGS
+		FIterationState() = default;
+		FIterationState(const FIterationState&) = default;
+		FIterationState(FIterationState&&) = default;
+		FIterationState& operator=(const FIterationState&) = default;
+		FIterationState& operator=(FIterationState&&) = default;
+PRAGMA_ENABLE_DEPRECATION_WARNINGS
+
 		TArray<FVector> CachedRayOrigins;
 		TArray<FVector> CachedRayVectors;
 		TArray<double> CachedRayLengths;
+		TSet<FObjectKey> CachedFilterObjects;
+
+		UE_DEPRECATED(5.8, "Use CachedFilterObjects instead")
 		TSet<TObjectKey<AActor>> CachedFilterActors;
 	};
 }
diff --git a/Engine/Plugins/PCG/Source/PCG/Public/Helpers/PCGWorldQueryHelpers.h b/Engine/Plugins/PCG/Source/PCG/Public/Helpers/PCGWorldQueryHelpers.h
--- a/Engine/Plugins/PCG/Source/PCG/Public/Helpers/PCGWorldQueryHelpers.h
+++ b/Engine/Plugins/PCG/Source/PCG/Public/Helpers/PCGWorldQueryHelpers.h
@@ -2,15 +2,19 @@
 
 #pragma once
 
+#include "Physics/PCGPhysicsRegistry.h"
+
 #include "Containers/ArrayView.h"
 #include "Engine/HitResult.h"
 #include "Misc/Optional.h"
 #include "UObject/ObjectKey.h"
 #include "UObject/SoftObjectPtr.h"
+#include "UObject/WeakInterfacePtr.h"
 
 #include "PCGWorldQueryHelpers.generated.h"
 
 class AActor;
+class IPCGGraphExecutionSource;
 class FPCGMetadataDomain;
 class UPCGComponent;
 class UPCGMetadata;
@@ -52,26 +56,36 @@ namespace PCGWorldQueryHelpers
 {
 	FTransform GetOrthonormalImpactTransform(const FHitResult& Hit);
 
+	
+	/** Returns true if hit was handled, in this case bOutPassesFilter will be set. */
+	bool FilterRayHitResult(const FPCGFilterHitResultParams& InParams, bool& bOutPassesFilter);
+	
+	/** Returns true if overlap was handled, in this case bOutPassesFilter will be set. */
+	bool FilterOverlapResult(const FPCGFilterOverlapResultParams& InParams, bool& bOutPassesFilter);
+	
+	/** Returns true if hit was handled, in this case bOutSuccess will be set. */
+	bool ApplyHitResultAttributes(const FPCGApplyHitResultAttributesParams& InParams, bool& bOutSuccess);
+
 	/** Returns false if a physics query doesn't meet the requirements of common query parameters. */
 	bool FilterCommonQueryResults(
-		const FPCGWorldCommonQueryParams* QueryParams,
+		const FPCGWorldCommonQueryParams& QueryParams,
 		const UPrimitiveComponent* TriggeredComponent,
-		const TWeakObjectPtr<UPCGComponent> OriginatingComponent,
-		const TSet<TObjectKey<AActor>>& FilteredActorReferences);
+		TWeakInterfacePtr<IPCGGraphExecutionSource> OriginatingExecutionSource,
+		const TSet<FObjectKey>& FilteredObjectReferences);
 
 	/** Filters through an array of hit results, testing them against raycast query parameters. Optional extra filter predicate. */
 	TOptional<FHitResult> FilterRayHitResults(
-		const FPCGWorldRaycastQueryParams* QueryParams,
-		TWeakObjectPtr<UPCGComponent> OriginatingComponent,
+		const FPCGWorldRaycastQueryParams& QueryParams,
+		TWeakInterfacePtr<IPCGGraphExecutionSource> OriginatingExecutionSource,
 		const TArray<FHitResult>& HitResults,
-		const TSet<TObjectKey<AActor>>& FilteredActorReferences);
+		const TSet<FObjectKey>& FilteredObjectReferences);
 
 	/** Filters through an array of overlap results, testing them against raycast query parameters. Optional extra filter predicate. */
 	TOptional<FOverlapResult> FilterOverlapResults(
-		const FPCGWorldVolumetricQueryParams* QueryParams,
-		TWeakObjectPtr<UPCGComponent> OriginatingComponent,
+		const FPCGWorldVolumetricQueryParams& QueryParams,
+		TWeakInterfacePtr<IPCGGraphExecutionSource> OriginatingExecutionSource,
 		const TArray<FOverlapResult>& OverlapResults,
-		const TSet<TObjectKey<AActor>>& FilteredActorReferences);
+		const TSet<FObjectKey>& FilteredObjectReferences);
 
 	/** Creates hit result attributes based off query params. Can be called before ApplyRayHitMetadata. */
 	bool CreateRayHitAttributes(const FPCGWorldRaycastQueryParams& QueryParams, UPCGMetadata* OutMetadata);
diff --git a/Engine/Plugins/PCG/Source/PCG/Public/PCGModule.h b/Engine/Plugins/PCG/Source/PCG/Public/PCGModule.h
--- a/Engine/Plugins/PCG/Source/PCG/Public/PCGModule.h
+++ b/Engine/Plugins/PCG/Source/PCG/Public/PCGModule.h
@@ -10,6 +10,7 @@
 #include "Data/Registry/PCGGetDataFunctionRegistry.h"
 #include "Hash/PCGObjectHash.h"
 #include "Metadata/Accessors/PCGAttributeAccessorFactory.h"
+#include "Physics/PCGPhysicsRegistry.h"
 #include "Utils/PCGLogErrors.h"
 
 #include "Containers/Ticker.h"
@@ -75,6 +76,9 @@ public:
 	static const FPCGChangeTrackingRegistry& GetConstChangeTrackingRegistry() { return GetPCGModuleChecked().ChangeTrackingRegistry; }
 	static FPCGChangeTrackingRegistry& GetMutableChangeTrackingRegistry() { return GetPCGModuleChecked().ChangeTrackingRegistry; }
 
+	static const FPCGPhysicsRegistry& GetConstPhysicsRegistry() { return GetPCGModuleChecked().PhysicsRegistry; }
+	static FPCGPhysicsRegistry& GetMutablePhysicsRegistry() { return GetPCGModuleChecked().PhysicsRegistry; }
+
 #if WITH_EDITOR
 	static const FPCGObjectHashFactory& GetConstObjectHashFactory() { return GetPCGModuleChecked().ObjectHashFactory; }
 	static FPCGObjectHashFactory& GetMutableObjectHashFactory() { return GetPCGModuleChecked().ObjectHashFactory; }
@@ -102,6 +106,7 @@ private:
 	FPCGGraphExecutionRegistry GraphExecutionRegistry;
 	FPCGDataViewRegistry DataViewRegistry;
 	FPCGChangeTrackingRegistry ChangeTrackingRegistry;
+	FPCGPhysicsRegistry PhysicsRegistry;
 
 #if WITH_EDITOR
 	FPCGObjectHashFactory ObjectHashFactory;
diff --git a/Engine/Plugins/PCG/Source/PCG/Public/Physics/PCGPhysicsRegistry.h b/Engine/Plugins/PCG/Source/PCG/Public/Physics/PCGPhysicsRegistry.h
new file mode 100644
--- /dev/null
+++ b/Engine/Plugins/PCG/Source/PCG/Public/Physics/PCGPhysicsRegistry.h
@@ -0,0 +1,107 @@
+// Copyright Epic Games, Inc. All Rights Reserved.
+
+#pragma once
+
+#include "PCGCommon.h"
+
+#include "Delegates/Delegate.h"
+#include "Misc/Guid.h"
+#include "UObject/ObjectKey.h"
+#include "UObject/WeakInterfacePtr.h"
+
+class IPCGGraphExecutionSource;
+class FPCGMetadataDomain;
+struct FPCGWorldRaycastQueryParams;
+struct FPCGWorldVolumetricQueryParams;
+class UPCGMetadata;
+class UWorld;
+
+struct FHitResult;
+struct FOverlapResult;
+
+struct FPCGFilterHitResultParams
+{
+	explicit FPCGFilterHitResultParams(const FHitResult& InHitResult, const FPCGWorldRaycastQueryParams& InQueryParams, const TWeakInterfacePtr<IPCGGraphExecutionSource> InOriginatingSource, const TSet<FObjectKey>& InFilteredObjectReferences)
+		: HitResult(InHitResult), QueryParams(InQueryParams), OriginatingSource(InOriginatingSource), FilteredObjectReferences(InFilteredObjectReferences)
+	{
+	}
+
+	const FHitResult& HitResult;
+
+	const FPCGWorldRaycastQueryParams& QueryParams;
+
+	const TWeakInterfacePtr<IPCGGraphExecutionSource> OriginatingSource;
+
+	const TSet<FObjectKey> FilteredObjectReferences;
+};
+
+struct FPCGFilterOverlapResultParams
+{
+	explicit FPCGFilterOverlapResultParams(const FOverlapResult& InOverlapResult, const FPCGWorldVolumetricQueryParams& InQueryParams, const TWeakInterfacePtr<IPCGGraphExecutionSource> InOriginatingSource, const TSet<FObjectKey>& InFilteredObjectReferences)
+		: OverlapResult(InOverlapResult), QueryParams(InQueryParams), OriginatingSource(InOriginatingSource), FilteredObjectReferences(InFilteredObjectReferences)
+	{
+	}
+
+	const FOverlapResult& OverlapResult;
+
+	const FPCGWorldVolumetricQueryParams& QueryParams;
+
+	const TWeakInterfacePtr<IPCGGraphExecutionSource> OriginatingSource;
+
+	const TSet<FObjectKey> FilteredObjectReferences;
+};
+
+struct FPCGApplyHitResultAttributesParams
+{
+	explicit FPCGApplyHitResultAttributesParams(const FHitResult& InHitResult, const FPCGWorldRaycastQueryParams& InQueryParams, const FVector& InRayDirection, const FTransform& InTransform, const UWorld* InWorld, int64& InOutMetaDataEntry, UPCGMetadata& InOutMetadata)
+		: HitResult(InHitResult), QueryParams(InQueryParams), RayDirection(InRayDirection), Transform(InTransform), World(InWorld), MetadataEntry(InOutMetaDataEntry), Metadata(InOutMetadata)
+	{
+
+	}
+
+	const FHitResult& HitResult;
+	
+	const FPCGWorldRaycastQueryParams& QueryParams;
+	
+	const FVector& RayDirection;
+	
+	const FTransform& Transform;
+	
+	const UWorld* World = nullptr;
+
+	int64& MetadataEntry;
+
+	UPCGMetadata& Metadata;
+};
+
+
+// Returns true if filter is handling this call, the bool param is only valid if delegate returns true and indicates if Hit passes the filter.
+DECLARE_DELEGATE_RetVal_TwoParams(bool, FPCGFilterHitResult, const FPCGFilterHitResultParams&, bool&);
+
+// Returns true if filter is handling this call, the bool param is only valid if delegate returns true and indicates if Overlap passes the filter.
+DECLARE_DELEGATE_RetVal_TwoParams(bool, FPCGFilterOverlapResult, const FPCGFilterOverlapResultParams&, bool&);
+
+// Returns true if filter is handling this call, the bool param is only valid if delegate returns true and indicates success (true) or failure (false).
+DECLARE_DELEGATE_RetVal_TwoParams(bool, FPCGApplyHitResultAttributes, const FPCGApplyHitResultAttributesParams&, bool&);
+
+class FPCGPhysicsRegistry
+{
+public:
+	PCG_API void RegisterHitResultFilter(const FGuid& InID, FPCGFilterHitResult InFilter);
+	PCG_API void UnregisterHitResultFilter(const FGuid& InID);
+
+	PCG_API void RegisterOverlapResultFilter(const FGuid& InID, FPCGFilterOverlapResult InFilter);
+	PCG_API void UnregisterOverlapResultFilter(const FGuid& InID);
+
+	PCG_API void RegisterApplyHitResultAttributes(const FGuid& InID, FPCGApplyHitResultAttributes InApplyHitResultAttributes);
+	PCG_API void UnregisterApplyHitResultAttributes(const FGuid& InID);
+
+	PCG_API bool FilterOverlapResult(const FPCGFilterOverlapResultParams& InParams) const;
+	PCG_API bool FilterHitResult(const FPCGFilterHitResultParams& InParams) const;
+	PCG_API bool ApplyHitResultAttributes(const FPCGApplyHitResultAttributesParams& InParams) const;
+
+private:
+	TMap<FGuid, FPCGFilterHitResult> HitResultFilters;
+	TMap<FGuid, FPCGFilterOverlapResult> OverlapResultFilters;
+	TMap<FGuid, FPCGApplyHitResultAttributes> ApplyHitResultAttributesFuncs;
+};
\ No newline at end of file
