[PCG] Bughawk fix

#rnx
#Synced-CL 48986798

diff --git a/Engine/Plugins/PCG/Source/PCG/Private/Helpers/PCGWorldQueryHelpers.cpp b/Engine/Plugins/PCG/Source/PCG/Private/Helpers/PCGWorldQueryHelpers.cpp
--- a/Engine/Plugins/PCG/Source/PCG/Private/Helpers/PCGWorldQueryHelpers.cpp
+++ b/Engine/Plugins/PCG/Source/PCG/Private/Helpers/PCGWorldQueryHelpers.cpp
@@ -389,7 +389,7 @@ namespace PCGWorldQueryHelpers
 
 								// Find the section by iterating through the indices 3 at a time.
 								const uint32 FirstFaceIndex = Section.FirstIndex / 3u;
-								const uint32 LastFaceIndex = FirstFaceIndex + Section.NumTriangles;
+								const uint32 LastFaceIndex = FirstFaceIndex + Section.NumTriangles - 1;
 
 								// Find the section index
 								if (HitFaceIndex >= FirstFaceIndex && HitFaceIndex <= LastFaceIndex)
