[PCG] Fixed normal assumption by using the input transform (from spline) if any or from the settings otherwise, which will ensure consistent winding.

#rnx
#Synced-CL 48910779

diff --git a/Engine/Plugins/PCG/Source/PCG/Private/Elements/Polygon/PCGCreatePolygon2D.cpp b/Engine/Plugins/PCG/Source/PCG/Private/Elements/Polygon/PCGCreatePolygon2D.cpp
--- a/Engine/Plugins/PCG/Source/PCG/Private/Elements/Polygon/PCGCreatePolygon2D.cpp
+++ b/Engine/Plugins/PCG/Source/PCG/Private/Elements/Polygon/PCGCreatePolygon2D.cpp
@@ -47,6 +47,7 @@ bool FPCGCreatePolygonElement::ExecuteInternal(FPCGContext* Context) const
 	for (const FPCGTaggedData& Input : Inputs)
 	{
 		const UPCGSpatialData* SpatialData = nullptr;
+		FVector DataNormal = Settings->DefaultPolygonNormal.GetSafeNormal(UE_SMALL_NUMBER, FVector::ZAxisVector);
 
 		bool bInputIsOpen = false;
 		bool bInputFromClosedSpline = false;
@@ -60,6 +61,7 @@ bool FPCGCreatePolygonElement::ExecuteInternal(FPCGContext* Context) const
 		if (const UPCGSplineData* SplineData = Cast<UPCGSplineData>(Input.Data))
 		{
 			SpatialData = SplineData;
+			DataNormal = SplineData->GetTransform().GetUnitAxis(EAxis::Z);
 			bInputFromClosedSpline = SplineData->IsClosed();
 			bInputIsOpen = !bInputFromClosedSpline;
 
@@ -153,6 +155,12 @@ bool FPCGCreatePolygonElement::ExecuteInternal(FPCGContext* Context) const
 			PlaneNormal = Settings->DefaultPolygonNormal.GetSafeNormal(UE_SMALL_NUMBER, FVector::ZAxisVector);
 		}
 
+		// Ensure normal is aligned to expectation otherwise it will break sampling, conversion to points, etc.
+		if ((PlaneNormal | DataNormal) < 0)
+		{
+			PlaneNormal = -PlaneNormal;
+		}
+
 		FTransform PlaneTransform;
 		PlaneTransform.SetRotation(FRotationMatrix::MakeFromZ(PlaneNormal).ToQuat());
 		PlaneTransform.SetLocation(PlaneCentroid);
