[PCG] Add On PCG Link Cleared callback on PCGComponent and add option to Clear PCG Link to all PCG Components

* Also fixing the API on Dialog confirmation.

#rnx
#Synced-CL 48998949

diff --git a/Engine/Plugins/PCG/Source/PCG/Private/PCGComponent.cpp b/Engine/Plugins/PCG/Source/PCG/Private/PCGComponent.cpp
--- a/Engine/Plugins/PCG/Source/PCG/Private/PCGComponent.cpp
+++ b/Engine/Plugins/PCG/Source/PCG/Private/PCGComponent.cpp
@@ -60,13 +60,13 @@
 
 namespace PCGComponent
 {
-	template <typename DelegateType>
-	static void BroadcastDynamicDelegate(const DelegateType& Delegate, UPCGComponent* PCGComponent)
+	template <typename DelegateType, typename ...Args>
+	static void BroadcastDynamicDelegate(const DelegateType& Delegate, UPCGComponent* PCGComponent, Args&& ...InArgs)
 	{
 #if WITH_EDITOR
 		const TGuardValue ScriptExecutionGuard(GAllowActorScriptExecutionInEditor, true);
 #endif // WITH_EDITOR
-		Delegate.Broadcast(PCGComponent);
+		Delegate.Broadcast(PCGComponent, std::forward<Args>(InArgs)...);
 	}
 
 #if WITH_EDITOR
@@ -1229,6 +1229,8 @@ AActor* UPCGComponent::ClearPCGLink(const FPCGMoveResourceParams& InParams)
 	}
 #endif
 
+	PCGComponent::BroadcastDynamicDelegate(OnPCGLinkClearedExternal, this, NewActor);
+
 	return NewActor;
 }
 
diff --git a/Engine/Plugins/PCG/Source/PCG/Private/Subsystems/PCGSubsystem.cpp b/Engine/Plugins/PCG/Source/PCG/Private/Subsystems/PCGSubsystem.cpp
--- a/Engine/Plugins/PCG/Source/PCG/Private/Subsystems/PCGSubsystem.cpp
+++ b/Engine/Plugins/PCG/Source/PCG/Private/Subsystems/PCGSubsystem.cpp
@@ -940,25 +940,25 @@ TStatId UPCGSubsystem::GetStatId() const
 	RETURN_QUICK_DECLARE_CYCLE_STAT(UPCGSubsystem, STATGROUP_Tickables);
 }
 
-FPCGTaskId UPCGSubsystem::ScheduleGeneric(TFunction<bool()> InOperation, IPCGGraphExecutionSource* ExecutionSource, const TArray<FPCGTaskId>& TaskExecutionDependencies)
+FPCGTaskId UPCGSubsystem::ScheduleGeneric(TFunction<bool()> InOperation, IPCGGraphExecutionSource* ExecutionSource, const TArray<FPCGTaskId>& TaskExecutionDependencies) const
 {
 	check(GraphExecutor);
 	return GraphExecutor->ScheduleGeneric(InOperation, ExecutionSource, TaskExecutionDependencies);
 }
 
-FPCGTaskId UPCGSubsystem::ScheduleGeneric(TFunction<bool()> InOperation, TFunction<void()> InAbortOperation, IPCGGraphExecutionSource* ExecutionSource, const TArray<FPCGTaskId>& TaskExecutionDependencies)
+FPCGTaskId UPCGSubsystem::ScheduleGeneric(TFunction<bool()> InOperation, TFunction<void()> InAbortOperation, IPCGGraphExecutionSource* ExecutionSource, const TArray<FPCGTaskId>& TaskExecutionDependencies) const
 {
 	check(GraphExecutor);
 	return GraphExecutor->ScheduleGeneric(InOperation, InAbortOperation, ExecutionSource, TaskExecutionDependencies);
 }
 
-FPCGTaskId UPCGSubsystem::ScheduleGenericWithContext(TFunction<bool(FPCGContext*)> InOperation, IPCGGraphExecutionSource* ExecutionSource, const TArray<FPCGTaskId>& TaskExecutionDependencies, const TArray<FPCGTaskId>& TaskDataDependencies)
+FPCGTaskId UPCGSubsystem::ScheduleGenericWithContext(TFunction<bool(FPCGContext*)> InOperation, IPCGGraphExecutionSource* ExecutionSource, const TArray<FPCGTaskId>& TaskExecutionDependencies, const TArray<FPCGTaskId>& TaskDataDependencies) const
 {
 	check(GraphExecutor);
 	return GraphExecutor->ScheduleGenericWithContext(InOperation, ExecutionSource, TaskExecutionDependencies, TaskDataDependencies);
 }
 
-FPCGTaskId UPCGSubsystem::ScheduleGenericWithContext(TFunction<bool(FPCGContext*)> InOperation, TFunction<void(FPCGContext*)> InAbortOperation, IPCGGraphExecutionSource* ExecutionSource, const TArray<FPCGTaskId>& TaskExecutionDependencies, const TArray<FPCGTaskId>& TaskDataDependencies)
+FPCGTaskId UPCGSubsystem::ScheduleGenericWithContext(TFunction<bool(FPCGContext*)> InOperation, TFunction<void(FPCGContext*)> InAbortOperation, IPCGGraphExecutionSource* ExecutionSource, const TArray<FPCGTaskId>& TaskExecutionDependencies, const TArray<FPCGTaskId>& TaskDataDependencies) const
 {
 	check(GraphExecutor);
 	return GraphExecutor->ScheduleGenericWithContext(InOperation, InAbortOperation, ExecutionSource, TaskExecutionDependencies, TaskDataDependencies);
@@ -1766,16 +1766,40 @@ void UPCGSubsystem::ClearLandscapeCache()
 
 void UPCGSubsystem::GenerateAllPCGComponents(bool bForce) const
 {
+	TArray<FPCGTaskId> GenerateTasks;
+
 	for (IPCGGraphExecutionSource* ExecutionSource : TrackingManager->GetAllRegisteredExecutionSources())
 	{
 		IPCGGraphExecutionState::FGenerateParams GenerateParams;
 		GenerateParams.bEvenIfAlreadyGenerated = bForce;
-		ExecutionSource->GetExecutionState().Generate(GenerateParams);
+		const FPCGTaskId TaskId = ExecutionSource->GetExecutionState().Generate(GenerateParams);
+		if (TaskId != InvalidPCGTaskId)
+		{
+			GenerateTasks.Add(TaskId);
+		}
+	}
+
+	auto BroadcastGenerated = [this]()
+	{
+		OnAllComponentsGenerated.Broadcast();
+		return true;
+	};
+
+	if (!GenerateTasks.IsEmpty())
+	{
+		ScheduleGeneric(BroadcastGenerated, /*ExecutionSource=*/nullptr,	GenerateTasks);
+	}
+	else
+	{
+		// Broadcast immediately if nothing happened.
+		BroadcastGenerated();
 	}
 }
 
 void UPCGSubsystem::CleanupAllPCGComponents(bool bPurge) const
 {
+	TArray<FPCGTaskId> CleanupTasks;
+
 	for (IPCGGraphExecutionSource* ExecutionSource : TrackingManager->GetAllRegisteredExecutionSources())
 	{
 		UPCGComponent* Component = Cast<UPCGComponent>(ExecutionSource);
@@ -1787,9 +1811,42 @@ void UPCGSubsystem::CleanupAllPCGComponents(bool bPurge) const
 		{
 			IPCGGraphExecutionState::FCleanupParams CleanupParams;
 			CleanupParams.bReleaseManagedResources = true;
-			ExecutionSource->GetExecutionState().Cleanup(CleanupParams);
+			const FPCGTaskId TaskId = ExecutionSource->GetExecutionState().Cleanup(CleanupParams);
+			if (TaskId != InvalidPCGTaskId)
+			{
+				CleanupTasks.Add(TaskId);
+			}
 		}
 	}
+
+	auto BroadcastCleanup = [this]()
+	{
+		OnAllComponentsCleanedup.Broadcast();
+		return true;
+	};
+
+	if (bPurge || CleanupTasks.IsEmpty())
+	{
+		// Broadcast immediately if nothing happened, or we purged (immediate cleanup)
+		BroadcastCleanup();
+	}
+	else
+	{
+		ScheduleGeneric(BroadcastCleanup, /*ExecutionSource=*/nullptr, CleanupTasks);
+	}
+}
+
+void UPCGSubsystem::ClearLinkForAllPCGComponents(const FPCGMoveResourceParams& InParams) const
+{
+	for (IPCGGraphExecutionSource* ExecutionSource : TrackingManager->GetAllRegisteredExecutionSources())
+	{
+		if (UPCGComponent* Component = Cast<UPCGComponent>(ExecutionSource))
+		{
+			Component->ClearPCGLink(InParams);
+		}
+	}
+
+	OnAllComponentsClearedLink.Broadcast();
 }
 
 void UPCGSubsystem::NotifySelectionKeyChanged(const FPCGSelectionKey& InSelectionKey, const UObject* InOriginatingChangeObject, TArrayView<const FBox> InChangedBounds)
diff --git a/Engine/Plugins/PCG/Source/PCG/Public/PCGComponent.h b/Engine/Plugins/PCG/Source/PCG/Public/PCGComponent.h
--- a/Engine/Plugins/PCG/Source/PCG/Public/PCGComponent.h
+++ b/Engine/Plugins/PCG/Source/PCG/Public/PCGComponent.h
@@ -60,6 +60,7 @@ DECLARE_DYNAMIC_MULTICAST_DELEGATE_OneParam(FOnPCGGraphStartGeneratingExternal,
 DECLARE_DYNAMIC_MULTICAST_DELEGATE_OneParam(FOnPCGGraphCancelledExternal, UPCGComponent*, PCGComponent);
 DECLARE_DYNAMIC_MULTICAST_DELEGATE_OneParam(FOnPCGGraphGeneratedExternal, UPCGComponent*, PCGComponent);
 DECLARE_DYNAMIC_MULTICAST_DELEGATE_OneParam(FOnPCGGraphCleanedExternal, UPCGComponent*, PCGComponent);
+DECLARE_DYNAMIC_MULTICAST_DELEGATE_TwoParams(FOnPCGLinkClearedExternal, UPCGComponent*, PCGComponent, UObject*, NewOwner);
 
 UENUM(Blueprintable)
 enum class EPCGComponentInput : uint8
@@ -385,6 +386,9 @@ public:
 	/** Event dispatched when a graph cleans on this component. */
 	UPROPERTY(BlueprintAssignable, meta = (DisplayName = "On Graph Cleaned"))
 	FOnPCGGraphCleanedExternal OnPCGGraphCleanedExternal;
+	/** Event dispatched when a component has finished to do its Clear PCG Link. */
+	UPROPERTY(BlueprintAssignable, meta = (DisplayName = "On PCG Link Cleared"))
+	FOnPCGLinkClearedExternal OnPCGLinkClearedExternal;
 
 	/** Flag to indicate whether this component has run in the editor. Note that for partitionable actors, this will always be false. */
 	UPROPERTY(BlueprintReadOnly, VisibleAnywhere, Category = Debug, NonTransactional, meta = (NoResetToDefault))
diff --git a/Engine/Plugins/PCG/Source/PCG/Public/Subsystems/PCGSubsystem.h b/Engine/Plugins/PCG/Source/PCG/Public/Subsystems/PCGSubsystem.h
--- a/Engine/Plugins/PCG/Source/PCG/Public/Subsystems/PCGSubsystem.h
+++ b/Engine/Plugins/PCG/Source/PCG/Public/Subsystems/PCGSubsystem.h
@@ -55,6 +55,15 @@ DECLARE_MULTICAST_DELEGATE_OneParam(FPCGOnPCGComponentUnregistered, UPCGComponen
 
 /** Deprecated - use FPCGOnPCGSourceGenerationDone */
 DECLARE_MULTICAST_DELEGATE_ThreeParams(FPCGOnPCGComponentGenerationDone, UPCGSubsystem*, UPCGComponent*, EPCGGenerationStatus);
+
+// Event triggered when GenerateAllPCGComponents is done
+DECLARE_MULTICAST_DELEGATE(FPCGOnAllComponentsGenerated);
+
+// Event triggered when CleanupAllPCGComponents is done
+DECLARE_MULTICAST_DELEGATE(FPCGOnAllComponentsCleanedup);
+
+// Event triggered when ClearLinkForAllPCGComponents is done
+DECLARE_MULTICAST_DELEGATE(FPCGOnAllComponentsClearedLink);
 #endif // WITH_EDITOR
 
 /**
@@ -163,7 +172,7 @@ public:
 	*  @param ExecutionSource:           Execution source associated with this task. Can be null.
 	*  @param TaskExecutionDependencies: Task will wait on these tasks to execute and won't take their output data as input.
 	*/
-	PCG_API FPCGTaskId ScheduleGeneric(TFunction<bool()> InOperation, IPCGGraphExecutionSource* ExecutionSource, const TArray<FPCGTaskId>& TaskExecutionDependencies);
+	PCG_API FPCGTaskId ScheduleGeneric(TFunction<bool()> InOperation, IPCGGraphExecutionSource* ExecutionSource, const TArray<FPCGTaskId>& TaskExecutionDependencies) const;
 
 	/** General job scheduling
 	*  @param InOperation:               Callback that returns true if the task is done, false otherwise.
@@ -171,7 +180,7 @@ public:
 	*  @param ExecutionSource:           Execution source associated with this task. Can be null.
 	*  @param TaskExecutionDependencies: Task will wait on these tasks to execute and won't take their output data as input.
 	*/
-	PCG_API FPCGTaskId ScheduleGeneric(TFunction<bool()> InOperation, TFunction<void()> InAbortOperation, IPCGGraphExecutionSource* ExecutionSource, const TArray<FPCGTaskId>& TaskExecutionDependencies);
+	PCG_API FPCGTaskId ScheduleGeneric(TFunction<bool()> InOperation, TFunction<void()> InAbortOperation, IPCGGraphExecutionSource* ExecutionSource, const TArray<FPCGTaskId>& TaskExecutionDependencies) const;
 
 	/** General job scheduling with context
 	*  @param InOperation:               Callback that takes a Context as argument and returns true if the task is done, false otherwise.
@@ -179,7 +188,7 @@ public:
 	*  @param TaskExecutionDependencies: Task will wait on these tasks to execute and won't take their output data as input.
 	*  @param TaskDataDependencies:      Task will wait on these tasks to execute and will take their output data as input.
 	*/
-	PCG_API FPCGTaskId ScheduleGenericWithContext(TFunction<bool(FPCGContext*)> InOperation, IPCGGraphExecutionSource* ExecutionSource, const TArray<FPCGTaskId>& TaskExecutionDependencies, const TArray<FPCGTaskId>& TaskDataDependencies);
+	PCG_API FPCGTaskId ScheduleGenericWithContext(TFunction<bool(FPCGContext*)> InOperation, IPCGGraphExecutionSource* ExecutionSource, const TArray<FPCGTaskId>& TaskExecutionDependencies, const TArray<FPCGTaskId>& TaskDataDependencies) const;
 
 	/** General job scheduling with context
 	*  @param InOperation:               Callback that takes a Context as argument and returns true if the task is done, false otherwise.
@@ -188,7 +197,7 @@ public:
 	*  @param TaskExecutionDependencies: Task will wait on these tasks to execute and won't take their output data as input.
 	*  @param TaskDataDependencies:      Task will wait on these tasks to execute and will take their output data as input.
 	*/
-	PCG_API FPCGTaskId ScheduleGenericWithContext(TFunction<bool(FPCGContext*)> InOperation, TFunction<void(FPCGContext*)> InAbortOperation, IPCGGraphExecutionSource* ExecutionSource, const TArray<FPCGTaskId>& TaskExecutionDependencies, const TArray<FPCGTaskId>& TaskDataDependencies);
+	PCG_API FPCGTaskId ScheduleGenericWithContext(TFunction<bool(FPCGContext*)> InOperation, TFunction<void(FPCGContext*)> InAbortOperation, IPCGGraphExecutionSource* ExecutionSource, const TArray<FPCGTaskId>& TaskExecutionDependencies, const TArray<FPCGTaskId>& TaskDataDependencies) const;
 
 	/** Asks the runtime generation scheduler to refresh a given GenerateAtRuntime component. ChangeType should be 'GenerationGrid' to perform a full cleanup of PAs and local components. */
 	PCG_API void RefreshRuntimeGenComponent(UPCGComponent* RuntimeComponent, EPCGChangeType ChangeType = EPCGChangeType::None);
@@ -335,6 +344,9 @@ public:
 	/** Will gather all the components registered, and ask for cleanup. */
 	PCG_API void CleanupAllPCGComponents(bool bPurge) const;
 
+	/** Will gather all the components registered, and ask for Clear PCG Link */
+	PCG_API void ClearLinkForAllPCGComponents(const FPCGMoveResourceParams& InParams) const;
+
 	/** Notification that a Selection key needs refreshing */
 	PCG_API void NotifySelectionKeyChanged(const FPCGSelectionKey& InSelectionKey, const UObject* InOriginatingChangeObject, TArrayView<const FBox> InBounds);
 
@@ -343,6 +355,10 @@ public:
 	void NotifyLandscapeEditModeExited() {}
 
 	FPCGOnPCGComponentUnregistered OnPCGComponentUnregistered;
+
+	FPCGOnAllComponentsGenerated OnAllComponentsGenerated;
+	FPCGOnAllComponentsCleanedup OnAllComponentsCleanedup;
+	FPCGOnAllComponentsClearedLink OnAllComponentsClearedLink;
 	
 	UE_DEPRECATED(5.7, "Deprecated in favor of OnPCGSourceGenerationDone, will not be notified anymore")
 	FPCGOnPCGComponentGenerationDone OnPCGComponentGenerationDone;
diff --git a/Engine/Plugins/PCG/Source/PCGEditor/Private/Nodes/PCGEditorGraphNodeGetUserParameter.cpp b/Engine/Plugins/PCG/Source/PCGEditor/Private/Nodes/PCGEditorGraphNodeGetUserParameter.cpp
--- a/Engine/Plugins/PCG/Source/PCGEditor/Private/Nodes/PCGEditorGraphNodeGetUserParameter.cpp
+++ b/Engine/Plugins/PCG/Source/PCGEditor/Private/Nodes/PCGEditorGraphNodeGetUserParameter.cpp
@@ -69,7 +69,7 @@ void UPCGEditorGraphGetUserParameter::OnRenameNode(const FString& NewName)
 				bConfirmClicked = true;
 			};
 
-			PCGEditorUtils::OpenConfirmationDialog(DialogMessage, OnConfirm, /*OnCancel=*/[] {}, DisableNotificationWidget);
+			PCGEditorUtils::OpenConfirmationDialog(MoveTemp(DialogMessage), OnConfirm, /*OnCancel=*/[] {}, DisableNotificationWidget);
 
 			if (!bConfirmClicked)
 			{
diff --git a/Engine/Plugins/PCG/Source/PCGEditor/Private/PCGEditorModule.cpp b/Engine/Plugins/PCG/Source/PCGEditor/Private/PCGEditorModule.cpp
--- a/Engine/Plugins/PCG/Source/PCGEditor/Private/PCGEditorModule.cpp
+++ b/Engine/Plugins/PCG/Source/PCGEditor/Private/PCGEditorModule.cpp
@@ -811,6 +811,25 @@ void FPCGEditorModule::PopulateMenuActions(FMenuBuilder& MenuBuilder)
 			}
 		})),
 	NAME_None);
+
+	MenuBuilder.AddMenuEntry(
+	LOCTEXT("ClearPCGLinkForAllComponents", "Clear link for all PCG components"),
+	LOCTEXT("ClearPCGLinkForAllComponents_Tooltip", "Clear PCG Link for all loaded PCG Components. Note that all components will be cleaned up after that."),
+	FSlateIcon(),
+	FUIAction(
+		FExecuteAction::CreateLambda([]()
+		{
+			PCGEditorUtils::OpenConfirmationDialog(LOCTEXT("ClearPCGLinkForAllComponents_Confirmation", "This will clear the link for ALL loaded PCG components. It's a destructive operation. Are you sure?"),
+			[]()
+			{
+				if (UPCGSubsystem* PCGSubsystem = UPCGSubsystem::GetSubsystemForCurrentWorld())
+				{
+					PCGSubsystem->ClearLinkForAllPCGComponents(FPCGMoveResourceParams{.TemplateTargetClass=AActor::StaticClass()});
+				}
+			},
+			[](){});
+		})),
+	NAME_None);
 }
 
 void FPCGEditorModule::RegisterSettings()
diff --git a/Engine/Plugins/PCG/Source/PCGEditor/Private/PCGEditorUtils.cpp b/Engine/Plugins/PCG/Source/PCGEditor/Private/PCGEditorUtils.cpp
--- a/Engine/Plugins/PCG/Source/PCGEditor/Private/PCGEditorUtils.cpp
+++ b/Engine/Plugins/PCG/Source/PCGEditor/Private/PCGEditorUtils.cpp
@@ -631,7 +631,7 @@ private:
 	FSimpleDelegate OnCancelEvent;
 };
 
-void PCGEditorUtils::OpenConfirmationDialog(FText& InMessage, TFunctionRef<void()> OnOk, TFunctionRef<void()> OnCancel, TSharedPtr<SWidget> InAdditionalContent)
+void PCGEditorUtils::OpenConfirmationDialog(FText InMessage, TFunction<void()> OnOk, TFunction<void()> OnCancel, TSharedPtr<SWidget> InAdditionalContent)
 {
 	TSharedRef<SWindow> Window = SNew(SWindow)
 		.CreateTitleBar(false)
diff --git a/Engine/Plugins/PCG/Source/PCGEditor/Private/PCGEditorUtils.h b/Engine/Plugins/PCG/Source/PCGEditor/Private/PCGEditorUtils.h
--- a/Engine/Plugins/PCG/Source/PCGEditor/Private/PCGEditorUtils.h
+++ b/Engine/Plugins/PCG/Source/PCGEditor/Private/PCGEditorUtils.h
@@ -40,5 +40,5 @@ namespace PCGEditorUtils
 	/** Open an asset or jumps to an actor/component and select it. Will load it, and won't load a new level*/
 	void OpenAssetOrMoveToActorOrComponent(const FSoftObjectPath& InPath);
 
-	void OpenConfirmationDialog(FText& InMessage, TFunctionRef<void()> OnOk, TFunctionRef<void()> OnCancel = [] {}, TSharedPtr<SWidget> InAdditionalContent = SNullWidget::NullWidget);
+	void OpenConfirmationDialog(FText InMessage, TFunction<void()> OnOk, TFunction<void()> OnCancel = [] {}, TSharedPtr<SWidget> InAdditionalContent = SNullWidget::NullWidget);
 } // namespace PCGEditorUtils
