[PCG] Attribute selector details multi select support

#rnx
#Synced-CL 49040481

diff --git a/Engine/Plugins/PCG/Source/PCGEditor/Private/Details/PCGAttributePropertySelectorDetails.cpp b/Engine/Plugins/PCG/Source/PCGEditor/Private/Details/PCGAttributePropertySelectorDetails.cpp
--- a/Engine/Plugins/PCG/Source/PCGEditor/Private/Details/PCGAttributePropertySelectorDetails.cpp
+++ b/Engine/Plugins/PCG/Source/PCGEditor/Private/Details/PCGAttributePropertySelectorDetails.cpp
@@ -150,20 +150,29 @@ namespace PCGAttributePropertySelectorDetails
 
 	// Func signature: bool(FPCGAttributePropertySelector*, T)
 	template <typename T, typename Func>
-	void SetValue(FPCGAttributePropertySelector* InSelector, TSharedPtr<IPropertyHandle>& InPropertyHandle, T&& InValue, Func InCallback)
+	void SetValue(const FPCGAttributePropertySelector* InSelector, TSharedPtr<IPropertyHandle>& InPropertyHandle, T&& InValue, Func InCallback)
 	{
-		if (InSelector)
+		if (!InSelector)
 		{
-			FScopedTransaction Transaction(SetAttributePropertyNameTransaction);
+			return;
+		}
+
+		FPCGAttributePropertySelector TempSelector = *InSelector;
 
-			InPropertyHandle->NotifyPreChange();
-			if (!InCallback(InSelector, std::forward<T>(InValue)))
+		if (InCallback(&TempSelector, std::forward<T>(InValue)))
+		{
+			// Need to export text and import text to go through the normal flow (like update Blueprint instances)
+			FString SelectorAsText{};
+			FPCGAttributePropertySelector::StaticStruct()->ExportText(SelectorAsText, &TempSelector, /*Defaults=*/nullptr, nullptr, PPF_None, /*ExportRootScope=*/nullptr);
+			if (!SelectorAsText.IsEmpty())
 			{
-				Transaction.Cancel();
-				return;
-			}
+				FScopedTransaction Transaction(SetAttributePropertyNameTransaction);
 
-			InPropertyHandle->NotifyPostChange(EPropertyChangeType::ValueSet);
+				for (uint32 i = 0; i < InPropertyHandle->GetNumOuterObjects(); ++i)
+				{
+					InPropertyHandle->SetPerObjectValue(i, SelectorAsText);
+				}
+			}
 		}
 	}
 }
@@ -178,8 +187,8 @@ void FPCGAttributePropertySelectorDetails::CustomizeHeader(TSharedRef<IPropertyH
 
 	auto Validation = [this]() -> FSlateColor
 	{
-		FPCGAttributePropertySelector* Selector = GetStruct();
-		return (Selector && !Selector->IsValid()) ? FStyleColors::AccentRed : FSlateColor::UseForeground();
+		const TTuple<FPropertyAccess::Result, const FPCGAttributePropertySelector*> Result = GetStruct();
+		return (Result.Get<0>() == FPropertyAccess::Success && Result.Get<1>() && !Result.Get<1>()->IsValid()) ? FStyleColors::AccentRed : FSlateColor::UseForeground();
 	};
 
 	TSharedRef<SWidget> PropertyNameWidget = PropertyHandle->CreatePropertyNameWidget();
@@ -267,16 +276,20 @@ TSharedRef<SWidget> FPCGAttributePropertySelectorDetails::GenerateExtraMenu()
 	const bool bIsInput = StructProperty->Struct->IsChildOf<FPCGAttributePropertyInputSelector>();
 
 	// @todo_pcg: This submenu could be context friendly to the property type an only show when relevant.
-	MenuBuilder.AddSubMenu(
-		LOCTEXT("AddAnExtractor", "Add an extractor"),
-		LOCTEXT("AddAnExtractorTooltip", "List of all the extractor to add to the current attribute."),
-		FNewMenuDelegate::CreateLambda([this, bIsInput](FMenuBuilder& InSubMenuBuilder)
-		{
-			PCGAttributePropertySelectorDetails::BuildSubMenuSectionVector(InSubMenuBuilder, this, bIsInput);
-			PCGAttributePropertySelectorDetails::BuildSubMenuSectionRotator(InSubMenuBuilder, this, bIsInput);
-			PCGAttributePropertySelectorDetails::BuildSubMenuSectionTransform(InSubMenuBuilder, this);
-			PCGAttributePropertySelectorDetails::BuildSubMenuSectionString(InSubMenuBuilder, this);
-		}));
+	// Can't add an extractor if we have different values
+	if (GetStruct().Get<0>() == FPropertyAccess::Success)
+	{
+		MenuBuilder.AddSubMenu(
+			LOCTEXT("AddAnExtractor", "Add an extractor"),
+			LOCTEXT("AddAnExtractorTooltip", "List of all the extractor to add to the current attribute."),
+			FNewMenuDelegate::CreateLambda([this, bIsInput](FMenuBuilder& InSubMenuBuilder)
+			{
+				PCGAttributePropertySelectorDetails::BuildSubMenuSectionVector(InSubMenuBuilder, this, bIsInput);
+				PCGAttributePropertySelectorDetails::BuildSubMenuSectionRotator(InSubMenuBuilder, this, bIsInput);
+				PCGAttributePropertySelectorDetails::BuildSubMenuSectionTransform(InSubMenuBuilder, this);
+				PCGAttributePropertySelectorDetails::BuildSubMenuSectionString(InSubMenuBuilder, this);
+			}));
+	}
 
 	MenuBuilder.BeginSection("Attributes", LOCTEXT("AttributesHeader", "Attributes"));
 	{
@@ -370,27 +383,47 @@ TSharedRef<SWidget> FPCGAttributePropertySelectorDetails::GenerateExtraMenu()
 	return MenuBuilder.MakeWidget();
 }
 
-FPCGAttributePropertySelector* FPCGAttributePropertySelectorDetails::GetStruct()
+TTuple<FPropertyAccess::Result, const FPCGAttributePropertySelector*> FPCGAttributePropertySelectorDetails::GetStruct() const
 {
 	void* Data = nullptr;
-	FPropertyAccess::Result Result = PropertyHandle->GetValueData(Data);
+	const FPropertyAccess::Result Result = PropertyHandle->GetValueData(Data);
+	const FPCGAttributePropertySelector* OutStruct = nullptr;
 
-	return (Result == FPropertyAccess::Success) ? reinterpret_cast<FPCGAttributePropertySelector*>(Data) : nullptr;
-}
-
-const FPCGAttributePropertySelector* FPCGAttributePropertySelectorDetails::GetStruct() const
-{
-	void* Data = nullptr;
-	FPropertyAccess::Result Result = PropertyHandle->GetValueData(Data);
+	switch (Result)
+	{
+	case FPropertyAccess::MultipleValues:
+		{
+			TArray<void*> RawData;
+			PropertyHandle->AccessRawData(RawData);
+			OutStruct = !RawData.IsEmpty() ? static_cast<const FPCGAttributePropertySelector*>(RawData[0]) : nullptr;
+		}
+		break;
+	case FPropertyAccess::Success:
+		OutStruct = static_cast<const FPCGAttributePropertySelector*>(Data);
+		break;
+	case FPropertyAccess::Fail: // fall-through
+	default:
+		break;
+	}
 
-	return (Result == FPropertyAccess::Success) ? reinterpret_cast<const FPCGAttributePropertySelector*>(Data) : nullptr;
+	return {Result, OutStruct};
 }
 
 FText FPCGAttributePropertySelectorDetails::GetText() const
 {
-	const FPCGAttributePropertySelector* Selector = GetStruct();
-
-	return Selector ? Selector->GetDisplayText() : FText();
+	TTuple<FPropertyAccess::Result, const FPCGAttributePropertySelector*> Result = GetStruct();
+	
+	switch (Result.Get<0>())
+	{
+	case FPropertyAccess::MultipleValues:
+		return LOCTEXT("MultipleValues", "Multiple Values");
+	case FPropertyAccess::Success:
+		check(Result.Get<1>());
+		return Result.Get<1>()->GetDisplayText();
+	case FPropertyAccess::Fail: // fall-through
+	default:
+		return FText{};
+	}
 }
 
 void FPCGAttributePropertySelectorDetails::SetText(const FText& NewText, ETextCommit::Type CommitInfo)
@@ -400,17 +433,17 @@ void FPCGAttributePropertySelectorDetails::SetText(const FText& NewText, ETextCo
 		return;
 	}
 
-	PCGAttributePropertySelectorDetails::SetValue(GetStruct(), PropertyHandle, NewText.ToString(), [](FPCGAttributePropertySelector* InStruct, FString InValue) -> bool { return InStruct && InStruct->Update(std::move(InValue)); });
+	PCGAttributePropertySelectorDetails::SetValue(GetStruct().Get<1>(), PropertyHandle, NewText.ToString(), [](FPCGAttributePropertySelector* InStruct, FString InValue) -> bool { return InStruct && InStruct->Update(std::move(InValue)); });
 }
 
 void FPCGAttributePropertySelectorDetails::SetPointProperty(EPCGPointProperties EnumValue)
 {
-	PCGAttributePropertySelectorDetails::SetValue(GetStruct(), PropertyHandle, EnumValue, [](FPCGAttributePropertySelector* InStruct, EPCGPointProperties InValue) -> bool { return InStruct && InStruct->SetPointProperty(InValue); });
+	PCGAttributePropertySelectorDetails::SetValue(GetStruct().Get<1>(), PropertyHandle, EnumValue, [](FPCGAttributePropertySelector* InStruct, EPCGPointProperties InValue) -> bool { return InStruct && InStruct->SetPointProperty(InValue); });
 }
 
 void FPCGAttributePropertySelectorDetails::SetAttributeName(FName NewName)
 {
-	PCGAttributePropertySelectorDetails::SetValue(GetStruct(), PropertyHandle, NewName, [](FPCGAttributePropertySelector* InStruct, FName InValue) -> bool { return InStruct && InStruct->SetAttributeName(InValue); });
+	PCGAttributePropertySelectorDetails::SetValue(GetStruct().Get<1>(), PropertyHandle, NewName, [](FPCGAttributePropertySelector* InStruct, FName InValue) -> bool { return InStruct && InStruct->SetAttributeName(InValue); });
 }
 
 void FPCGAttributePropertySelectorDetails::SetExtraProperty(EPCGExtraProperties EnumValue)
@@ -419,7 +452,7 @@ void FPCGAttributePropertySelectorDetails::SetExtraProperty(EPCGExtraProperties
 	const int32 EnumIndex = EnumPtr ? EnumPtr->GetIndexByValue(static_cast<int64>(EnumValue)) : INDEX_NONE;
 	const FName DomainName = EnumPtr ? FName(EnumPtr->GetMetaData(*PCGObjectMetadata::EnumMetadataDomain.ToString(), EnumIndex)) : FName(NAME_None);
 	
-	PCGAttributePropertySelectorDetails::SetValue(GetStruct(), PropertyHandle, EnumValue, [DomainName](FPCGAttributePropertySelector* InStruct, EPCGExtraProperties InValue) -> bool
+	PCGAttributePropertySelectorDetails::SetValue(GetStruct().Get<1>(), PropertyHandle, EnumValue, [DomainName](FPCGAttributePropertySelector* InStruct, EPCGExtraProperties InValue) -> bool
 		{
 			if (!InStruct)
 			{
@@ -434,7 +467,7 @@ void FPCGAttributePropertySelectorDetails::SetExtraProperty(EPCGExtraProperties
 
 void FPCGAttributePropertySelectorDetails::SetSelector(FPCGAttributePropertySelector InSelector)
 {
-	PCGAttributePropertySelectorDetails::SetValue(GetStruct(), PropertyHandle, InSelector, [](FPCGAttributePropertySelector* InStruct, FPCGAttributePropertySelector InValue) -> bool
+	PCGAttributePropertySelectorDetails::SetValue(GetStruct().Get<1>(), PropertyHandle, InSelector, [](FPCGAttributePropertySelector* InStruct, FPCGAttributePropertySelector InValue) -> bool
 	{
 		if (!InStruct)
 		{
@@ -453,7 +486,7 @@ void FPCGAttributePropertySelectorDetails::SetSelector(FPCGAttributePropertySele
 
 void FPCGAttributePropertySelectorDetails::AddExtractor(FName InExtractor)
 {
-	PCGAttributePropertySelectorDetails::SetValue(GetStruct(), PropertyHandle, InExtractor, [](FPCGAttributePropertySelector* InStruct, FName InValue) -> bool 
+	PCGAttributePropertySelectorDetails::SetValue(GetStruct().Get<1>(), PropertyHandle, InExtractor, [](FPCGAttributePropertySelector* InStruct, FName InValue) -> bool 
 	{
 		if (!InStruct)
 		{
diff --git a/Engine/Plugins/PCG/Source/PCGEditor/Public/Details/PCGAttributePropertySelectorDetails.h b/Engine/Plugins/PCG/Source/PCGEditor/Public/Details/PCGAttributePropertySelectorDetails.h
--- a/Engine/Plugins/PCG/Source/PCGEditor/Public/Details/PCGAttributePropertySelectorDetails.h
+++ b/Engine/Plugins/PCG/Source/PCGEditor/Public/Details/PCGAttributePropertySelectorDetails.h
@@ -3,6 +3,7 @@
 #pragma once
 
 #include "IPropertyTypeCustomization.h"
+#include "PropertyEditorModule.h"
 #include "Layout/Visibility.h"
 #include "Metadata/PCGAttributePropertySelector.h"
 
@@ -29,8 +30,8 @@ public:
 	void AddExtractor(FName InExtractor);
 
 protected:
-	FPCGAttributePropertySelector* GetStruct();
-	const FPCGAttributePropertySelector* GetStruct() const;
+	/** Return the result of GetValueData on the PropertyHandle. The output struct will be the first valid selector. */
+	TTuple<FPropertyAccess::Result, const FPCGAttributePropertySelector*> GetStruct() const;
 
 	TSharedRef<SWidget> GenerateExtraMenu();
 	EVisibility ExtraMenuVisibility() const;
