[PCG] Made DeleteTags tag:value aware. Added option to return full tag string in GetTags node.

#Synced-CL 49004414

diff --git a/Engine/Plugins/PCG/Source/PCG/Private/Elements/PCGDeleteTags.cpp b/Engine/Plugins/PCG/Source/PCG/Private/Elements/PCGDeleteTags.cpp
--- a/Engine/Plugins/PCG/Source/PCG/Private/Elements/PCGDeleteTags.cpp
+++ b/Engine/Plugins/PCG/Source/PCG/Private/Elements/PCGDeleteTags.cpp
@@ -6,6 +6,7 @@
 #include "PCGCustomVersion.h"
 #include "PCGData.h"
 #include "Helpers/PCGHelpers.h"
+#include "Helpers/PCGTagHelpers.h"
 
 #include "Algo/AnyOf.h"
 
@@ -91,13 +92,40 @@ bool FPCGDeleteTagsElement::ExecuteInternal(FPCGContext* Context) const
 		: PCGHelpers::GetStringArrayFromCommaSeparatedList(Settings->SelectedTags);
 	PRAGMA_ENABLE_DEPRECATION_WARNINGS
 
-	TSet<FString> TagsToFilter(FilterTags);
+	TSet<FString> ValueTagsToFilter;
+	TSet<FString> StringTagsToFilter;
+
+	// When the user specifies tags with values inside, we'll do a straight string comparison,
+	// otherwise, we'll assume that the input tag can be a tag value.
+	for (const FString& FilterTag : FilterTags)
+	{
+		PCG::Private::FParseTagResult ParsedFilterTag(FilterTag);
+
+		if (ParsedFilterTag.IsValid())
+		{
+			if (ParsedFilterTag.HasValue())
+			{
+				StringTagsToFilter.Add(FilterTag);
+			}
+			else
+			{
+				ValueTagsToFilter.Add(FilterTag);
+			}
+		}
+	}
 
 	const bool bKeepInFilter = (Settings->Operation == EPCGTagFilterOperation::KeepOnlySelectedTags);
 
 	TArray<FPCGTaggedData> Inputs = Context->InputData.GetInputsByPin(PCGPinConstants::DefaultInputLabel);
 	TArray<FPCGTaggedData>& Outputs = Context->OutputData.TaggedData;
 
+	auto TestTag = [Settings, bKeepInFilter](const TSet<FString>& TagsToFilter, const FString Tag)
+	{
+		return ((Settings->Operator == EPCGStringMatchingOperator::Equal && TagsToFilter.Contains(Tag) == bKeepInFilter) ||
+			(Settings->Operator == EPCGStringMatchingOperator::Substring && Algo::AnyOf(TagsToFilter, [&Tag](const FString& TagToFilter) { return Tag.Contains(TagToFilter); }) == bKeepInFilter) ||
+			(Settings->Operator == EPCGStringMatchingOperator::Matches && Algo::AnyOf(TagsToFilter, [&Tag](const FString& TagToFilter) { return Tag.MatchesWildcard(TagToFilter); }) == bKeepInFilter));
+	};
+
 	for (const FPCGTaggedData& Input : Inputs)
 	{
 		FPCGTaggedData& Output = Outputs.Add_GetRef(Input);
@@ -107,9 +135,11 @@ bool FPCGDeleteTagsElement::ExecuteInternal(FPCGContext* Context) const
 
 		for (const FString& Tag : Input.Tags)
 		{
-			if ((Settings->Operator == EPCGStringMatchingOperator::Equal && TagsToFilter.Contains(Tag) == bKeepInFilter) ||
-				(Settings->Operator == EPCGStringMatchingOperator::Substring && Algo::AnyOf(TagsToFilter, [&Tag](const FString& TagToFilter) { return Tag.Contains(TagToFilter); }) == bKeepInFilter) ||
-				(Settings->Operator == EPCGStringMatchingOperator::Matches && Algo::AnyOf(TagsToFilter, [&Tag](const FString& TagToFilter) { return Tag.MatchesWildcard(TagToFilter); }) == bKeepInFilter))
+			const bool TestAgainstStringTags = TestTag(StringTagsToFilter, Tag);
+			const bool TestAgainstValueTags = TestTag(ValueTagsToFilter, PCG::Private::FParseTagResult(Tag).GetOriginalAttribute());
+
+			if ((bKeepInFilter && (TestAgainstStringTags || TestAgainstValueTags)) ||
+				(!bKeepInFilter && (TestAgainstStringTags && TestAgainstValueTags)))
 			{
 				Output.Tags.Add(Tag);
 			}
diff --git a/Engine/Plugins/PCG/Source/PCG/Private/Elements/PCGGetDataInfo.cpp b/Engine/Plugins/PCG/Source/PCG/Private/Elements/PCGGetDataInfo.cpp
--- a/Engine/Plugins/PCG/Source/PCG/Private/Elements/PCGGetDataInfo.cpp
+++ b/Engine/Plugins/PCG/Source/PCG/Private/Elements/PCGGetDataInfo.cpp
@@ -49,11 +49,18 @@ bool FPCGGetTagsElement::ExecuteInternal(FPCGContext* Context) const
 
 		FPCGMetadataAttribute<FString>* Attribute = ParamData->Metadata->FindOrCreateAttribute<FString>(FName("Attribute"), FString(), /*bAllowsInterpolation=*/false, /*bOverrideParent=*/false, /*bOverwriteIfTypeMismatch=*/true);
 		FPCGMetadataAttribute<FString>* Values = nullptr;
+		FPCGMetadataAttribute<FString>* Tags = nullptr;
+
 		if (Settings->bExtractTagValues)
 		{
 			Values = ParamData->Metadata->FindOrCreateAttribute<FString>(FName("Values"), FString(), false, false, true);
 		}
 
+		if (Settings->bExtractFullTagStrings)
+		{
+			Tags = ParamData->Metadata->FindOrCreateAttribute<FString>(FName("Tags"), FString(), false, false, true);
+		}
+
 		for (const FString& Tag : Input.Tags)
 		{
 			PCG::Private::FParseTagResult TagData = PCG::Private::ParseTag(Tag);
@@ -73,6 +80,11 @@ bool FPCGGetTagsElement::ExecuteInternal(FPCGContext* Context) const
 						Values->SetValue(EntryKey, FString("true"));
 					}
 				}
+
+				if (Tags)
+				{
+					Tags->SetValue(EntryKey, Tag);
+				}
 			}
 		}
 	}
diff --git a/Engine/Plugins/PCG/Source/PCG/Public/Elements/PCGGetDataInfo.h b/Engine/Plugins/PCG/Source/PCG/Public/Elements/PCGGetDataInfo.h
--- a/Engine/Plugins/PCG/Source/PCG/Public/Elements/PCGGetDataInfo.h
+++ b/Engine/Plugins/PCG/Source/PCG/Public/Elements/PCGGetDataInfo.h
@@ -30,6 +30,10 @@ public:
 	/** Creates a 'Values' attribute and stores the values of the valued tags (e.g. 'Tag:Value' tags) as a string. */
 	UPROPERTY(BlueprintReadWrite, EditAnywhere, Category = Settings)
 	bool bExtractTagValues = false;
+
+	/** Creates a 'Tags' attribute that stores the full tag string (e.g. "Tag:Value") as a string. */
+	UPROPERTY(BlueprintReadWrite, EditAnywhere, Category = Settings)
+	bool bExtractFullTagStrings = false;
 };
 
 class FPCGGetTagsElement : public IPCGElement
