[PCG] Exposed options in the isolate mode

#rnx
#Synced-CL 49004414

diff --git a/Engine/Plugins/PCG/Source/PCG/Private/Helpers/PCGActorHelpers.cpp b/Engine/Plugins/PCG/Source/PCG/Private/Helpers/PCGActorHelpers.cpp
--- a/Engine/Plugins/PCG/Source/PCG/Private/Helpers/PCGActorHelpers.cpp
+++ b/Engine/Plugins/PCG/Source/PCG/Private/Helpers/PCGActorHelpers.cpp
@@ -897,6 +897,16 @@ TArray<APCGIsolatedActor*> UPCGActorHelpers::GetIsolatedActorsOriginatingFrom(co
 }
 
 AActor* UPCGActorHelpers::IsolateFromActor(AActor* Actor, const TArray<FName>& RequiredTags)
+{
+	FPCGMoveResourceParams Params
+	{
+		.RequiredTags = RequiredTags
+	};
+
+	return IsolateFromActor(Actor, Params);
+}
+
+AActor* UPCGActorHelpers::IsolateFromActor(AActor* Actor, const FPCGMoveResourceParams& InParams)
 {
 	UPCGComponent* PCGComponent = Actor ? Actor->GetComponentByClass<UPCGComponent>() : nullptr;
 
@@ -906,42 +916,78 @@ AActor* UPCGActorHelpers::IsolateFromActor(AActor* Actor, const TArray<FName>& R
 	}
 
 	// Will store data temporarily (output from the ClearPCGLink) until we move it to the created actor
+	FPCGMoveResourceParams Params = InParams;
 	FPCGDataCollection TempCollection;
+	bool bMoveCollectionToIsolatedActor = !Params.TargetDataCollection;
 
-	FPCGMoveResourceParams Params
+	if (bMoveCollectionToIsolatedActor)
 	{
-		.TemplateTargetClass = APCGIsolatedActor::StaticClass(),
-		.RequiredTags = RequiredTags,
-		.TargetDataCollection = &TempCollection,
-	};
+		Params.TargetDataCollection = &TempCollection;
+	}
 
 	AActor* IsolatedActor = PCGComponent->ClearPCGLink(Params);
 
-	if (IsolatedActor)
+	// Setup isolated actor data if we're created one.
+	if (APCGIsolatedActor* IsolatedPCGActor = Cast<APCGIsolatedActor>(IsolatedActor))
+	{
+		IsolatedPCGActor->Origin = Actor;
+
+		// Move collection to isolated actor
+		if (bMoveCollectionToIsolatedActor)
+		{
+			IsolatedPCGActor->Data = MoveTemp(*Params.TargetDataCollection);
+		}
+
+		// Re-outer data
+		for (FPCGTaggedData& TaggedData : IsolatedPCGActor->Data.TaggedData)
+		{
+			TaggedData.Data->VisitDataNetwork([IsolatedPCGActor](const UPCGData* InData)
+			{
+				const_cast<UPCGData*>(InData)->Rename(nullptr, IsolatedPCGActor, REN_DontCreateRedirectors);
+			});
+		}
+	}
+
+	// @todo_pcg: Push tags onto newly created actor?
+
+	// Remove PCG Generated tags, so they don't get ignored elsewhere.
+	// This is a bit more complex, as we need to tackle the "new actor" case and the "moved in place" case.
+	if(IsolatedActor && IsolatedActor != Params.GetTarget<AActor>())
 	{
-		// @todo_pcg: Push tags onto actor?
-		// Remove all PCG Generated tags, so we don't ignore those elsewhere
 		IsolatedActor->ForEachComponent(/*bIncludeFromChildActors=*/true, [](UActorComponent* Component)
 		{
 			Component->ComponentTags.Remove(PCGHelpers::DefaultPCGTag);
 		});
+	}
+	// In the case of a general target, we need to make sure that we untag just the proper components,
+	// And not any that would have been tracked by a PCG component - especially in the "move to self" case.
+	// Note that this isn't foolproof though, because we could have components that are owned by a PCG component on another actor
+	// but we don't know this unfortunately.
+	else if (AActor* TargetActor = Params.GetTarget<AActor>())
+	{
+		TArray<UPCGComponent*> PCGComponents;
+		TargetActor->GetComponents<UPCGComponent>(PCGComponents);
 
-		if (APCGIsolatedActor* IsolatedPCGActor = Cast<APCGIsolatedActor>(IsolatedActor))
+		TargetActor->ForEachComponent(/*bIncludeFromChildActors=*/true, [&PCGComponents](UActorComponent* Component)
 		{
-			IsolatedPCGActor->Origin = Actor;
+			if (Component->ComponentTags.Contains(PCGHelpers::DefaultPCGTag))
+			{
+				bool bRemoveTag = true;
+				const UObject* PtrOnComponent = Component;
 
-			// Move collection to isolated actor
-			IsolatedPCGActor->Data = MoveTemp(TempCollection);
+				TArrayView<const UObject*> ViewOnComponent = MakeArrayView(&PtrOnComponent, 1);
 
-			// Re-outer data
-			for (FPCGTaggedData& TaggedData : IsolatedPCGActor->Data.TaggedData)
-			{
-				TaggedData.Data->VisitDataNetwork([IsolatedPCGActor](const UPCGData* InData)
+				for (UPCGComponent* PCGComponent : PCGComponents)
+				{
+					bRemoveTag &= !PCGComponent->IsAnyObjectManagedByResource(ViewOnComponent);
+				}
+
+				if (bRemoveTag)
 				{
-					const_cast<UPCGData*>(InData)->Rename(nullptr, IsolatedPCGActor, REN_DontCreateRedirectors);
-				});
+					Component->ComponentTags.Remove(PCGHelpers::DefaultPCGTag);
+				}
 			}
-		}
+		});
 	}
 
 	return IsolatedActor;
diff --git a/Engine/Plugins/PCG/Source/PCG/Public/Helpers/PCGActorHelpers.h b/Engine/Plugins/PCG/Source/PCG/Public/Helpers/PCGActorHelpers.h
--- a/Engine/Plugins/PCG/Source/PCG/Public/Helpers/PCGActorHelpers.h
+++ b/Engine/Plugins/PCG/Source/PCG/Public/Helpers/PCGActorHelpers.h
@@ -20,8 +20,6 @@
 class AActor;
 class APCGIsolatedActor;
 class IPCGGraphExecutionSource;
-struct FPCGContext;
-struct FPCGGetExecutionSourcesFromSelectionKeyParams;
 class UActorComponent;
 class UDataLayerInstance;
 class UHLODLayer;
@@ -37,6 +35,9 @@ class UPCGManagedSkinnedMeshComponent;
 class UPCGManagedSplineMeshComponent;
 class UStaticMesh;
 class UWorld;
+struct FPCGContext;
+struct FPCGGetExecutionSourcesFromSelectionKeyParams;
+struct FPCGMoveResourceParams;
 
 struct UE_DEPRECATED(5.5, "Use FPCGISMComponentBuilderParams instead.") FPCGISMCBuilderParameters
 {
@@ -262,9 +263,12 @@ PRAGMA_ENABLE_DEPRECATION_WARNINGS
 	/** Returns all isolated actors that originate from the given actor. */
 	static TArray<APCGIsolatedActor*> GetIsolatedActorsOriginatingFrom(const AActor* InActor);
 	
-	/** Extracts (like Clear PCG Link) content from the target actor, assuming it has a PCG component, into a PCG Isolated Actor. */
+	/** Extracts (like Clear PCG Link) content from the actor, assuming it has a PCG component, into a PCG Isolated Actor. */
 	static PCG_API AActor* IsolateFromActor(AActor* Actor, const TArray<FName>& RequiredTags);
 
+	/** Extracts (like Clear PCG Link) content from the actor, assuming it has a PCG component, into a provided target actor or a PCG Isolated Actor created by this call. */
+	static PCG_API AActor* IsolateFromActor(AActor* Actor, const FPCGMoveResourceParams& Params);
+
 	/** Finds all execution sources using the passed in param as context for search. */
 	static PCG_API TArray<IPCGGraphExecutionSource*> GetExecutionSourcesFromSelectionKey(const FPCGGetExecutionSourcesFromSelectionKeyParams& InParams);
 };
diff --git a/Engine/Plugins/PCG/Source/PCGEditor/Private/EditorMode/PCGEdMode.cpp b/Engine/Plugins/PCG/Source/PCGEditor/Private/EditorMode/PCGEdMode.cpp
--- a/Engine/Plugins/PCG/Source/PCGEditor/Private/EditorMode/PCGEdMode.cpp
+++ b/Engine/Plugins/PCG/Source/PCGEditor/Private/EditorMode/PCGEdMode.cpp
@@ -126,8 +126,13 @@ namespace PCGEditorMode
 			{
 				Tags.Add(FName(Args[ArgIndex]));
 			}
+
+			FPCGMoveResourceParams Params
+			{
+				.RequiredTags = Tags
+			};
 			
-			UPCGActorHelpers::IsolateFromActor(Actor, Tags);
+			UPCGActorHelpers::IsolateFromActor(Actor, Params);
 		}));
 }
 
diff --git a/Engine/Plugins/PCG/Source/PCGEditor/Private/EditorMode/Tools/Query/PCGQueryTool.cpp b/Engine/Plugins/PCG/Source/PCGEditor/Private/EditorMode/Tools/Query/PCGQueryTool.cpp
--- a/Engine/Plugins/PCG/Source/PCGEditor/Private/EditorMode/Tools/Query/PCGQueryTool.cpp
+++ b/Engine/Plugins/PCG/Source/PCGEditor/Private/EditorMode/Tools/Query/PCGQueryTool.cpp
@@ -252,7 +252,17 @@ void UPCGInteractiveToolSettings_Isolate::Apply(UInteractiveTool* OwningTool)
 {
 	if (AActor* Actor = SelectedActor.Get())
 	{
-		UPCGActorHelpers::IsolateFromActor(Actor, SelectedTags);
+		FPCGMoveResourceParams Params
+		{
+			.TemplateTargetClass = APCGIsolatedActor::StaticClass(),
+			.Target = (bExtractToSelf ? Actor : nullptr),
+			.RequiredTags = SelectedTags,
+			.ExcludedTags = ExcludedTags,
+			.bAttachToParent = bAttachToActor,
+			.DefaultName = IsolatedActorLabel
+		};
+
+		UPCGActorHelpers::IsolateFromActor(Actor, Params);
 	}
 }
 
diff --git a/Engine/Plugins/PCG/Source/PCGEditor/Private/EditorMode/Tools/Query/PCGQueryTool.h b/Engine/Plugins/PCG/Source/PCGEditor/Private/EditorMode/Tools/Query/PCGQueryTool.h
--- a/Engine/Plugins/PCG/Source/PCGEditor/Private/EditorMode/Tools/Query/PCGQueryTool.h
+++ b/Engine/Plugins/PCG/Source/PCGEditor/Private/EditorMode/Tools/Query/PCGQueryTool.h
@@ -124,6 +124,18 @@ public:
 
 	UPROPERTY(EditAnywhere, Category = PCG)
 	TArray<FName> SelectedTags;
+
+	UPROPERTY(EditAnywhere, Category = PCG)
+	TArray<FName> ExcludedTags;
+
+	UPROPERTY(EditAnywhere, Category = PCG)
+	bool bExtractToSelf = false;
+
+	UPROPERTY(EditAnywhere, Category = PCG, meta = (EditCondition = "!bExtractToSelf"))
+	bool bAttachToActor = false;
+
+	UPROPERTY(EditAnywhere, Category = PCG, meta = (EditCondition = "!bExtractToSelf"))
+	FName IsolatedActorLabel = NAME_None;
 };
 
 /**
