[PCG] Fix another offset bug in generate landscape height map / grass maps.

#rnx
#Synced-CL 49065038

diff --git a/Engine/Plugins/PCG/Source/PCG/Private/Elements/PCGGenerateLandscapeTextures.cpp b/Engine/Plugins/PCG/Source/PCG/Private/Elements/PCGGenerateLandscapeTextures.cpp
--- a/Engine/Plugins/PCG/Source/PCG/Private/Elements/PCGGenerateLandscapeTextures.cpp
+++ b/Engine/Plugins/PCG/Source/PCG/Private/Elements/PCGGenerateLandscapeTextures.cpp
@@ -615,8 +615,8 @@ bool FPCGGenerateLandscapeTexturesElement::ExecuteInternal(FPCGContext* InContex
 		FBox OutputBounds;
 		OutputBounds.Min.X = Context->TotalLandscapeComponentBounds.Min.X + WindowStart.X * QuadSizeWorld;
 		OutputBounds.Min.Y = Context->TotalLandscapeComponentBounds.Min.Y + WindowStart.Y * QuadSizeWorld;
-		OutputBounds.Max.X = OutputBounds.Min.X + OutputResolution.X * QuadSizeWorld;
-		OutputBounds.Max.Y = OutputBounds.Min.Y + OutputResolution.Y * QuadSizeWorld;
+		OutputBounds.Max.X = OutputBounds.Min.X + (OutputResolution.X - 1) * QuadSizeWorld; // -1 because resolution is verts, not quads (so discard last col/row)
+		OutputBounds.Max.Y = OutputBounds.Min.Y + (OutputResolution.Y - 1) * QuadSizeWorld;
 		OutputBounds.Min.Z = Context->TotalLandscapeComponentBounds.Min.Z;
 		OutputBounds.Max.Z = Context->TotalLandscapeComponentBounds.Max.Z;
 
@@ -635,7 +635,7 @@ bool FPCGGenerateLandscapeTexturesElement::ExecuteInternal(FPCGContext* InContex
 		
 		TWeakPtr<FPCGContextHandle> ContextHandle = Context->GetOrCreateHandle();
 
-		ENQUEUE_RENDER_COMMAND(GenerateGrassMaps)([ContextHandle, LandscapeComponentExtent=Context->LandscapeComponentExtent, LandscapeTileCoords=MoveTemp(LandscapeTileCoords), LandscapeComponentResolution, LandscapeGridScale, LandscapeZ, bGenerateHeightMap=Settings->bGenerateHeightMap, NumTilesX, NumTilesY, OutputResolution, WindowStart, WindowEnd, IndexMapping=MoveTemp(Context->IndexMapping)](FRHICommandListImmediate& RHICmdList) // todo ok to move?
+		ENQUEUE_RENDER_COMMAND(GenerateGrassMaps)([ContextHandle, LandscapeComponentExtent=Context->LandscapeComponentExtent, LandscapeTileCoords=MoveTemp(LandscapeTileCoords), LandscapeComponentResolution, LandscapeGridScale, LandscapeZ, bGenerateHeightMap=Settings->bGenerateHeightMap, NumTilesX, NumTilesY, OutputResolution, WindowStart, WindowEnd, IndexMapping=MoveTemp(Context->IndexMapping)](FRHICommandListImmediate& RHICmdList)
 		{
 			LLM_SCOPE_BYTAG(PCG);
 
