[PCG] Fix offset sample positions when converting textures to point data.

#rnx
#Synced-CL 49065038

diff --git a/Engine/Plugins/PCG/Source/PCG/Private/Data/PCGTextureData.cpp b/Engine/Plugins/PCG/Source/PCG/Private/Data/PCGTextureData.cpp
--- a/Engine/Plugins/PCG/Source/PCG/Private/Data/PCGTextureData.cpp
+++ b/Engine/Plugins/PCG/Source/PCG/Private/Data/PCGTextureData.cpp
@@ -293,9 +293,6 @@ const UPCGPointArrayData* UPCGBaseTextureData::CreatePointArrayData(FPCGContext*
 const UPCGBasePointData* UPCGBaseTextureData::CreateBasePointData(FPCGContext* Context, TSubclassOf<UPCGBasePointData> PointDataClass) const
 {
 	TRACE_CPUPROFILER_EVENT_SCOPE(UPCGBaseTextureData::CreatePointData);
-	// TODO: this is a trivial implementation
-	// A better sampler would allow to sample a fixed number of points in either direction
-	// or based on a given texel size
 	FBox2D LocalSurfaceBounds(FVector2D(-1.0f, -1.0f), FVector2D(1.0f, 1.0f));
 
 	UPCGBasePointData* Data = FPCGContext::NewObject_AnyThread<UPCGBasePointData>(Context, GetTransientPackage(), PointDataClass);
@@ -351,8 +348,8 @@ const UPCGBasePointData* UPCGBaseTextureData::CreateBasePointData(FPCGContext* C
 
 			const int32 WriteIndex = StartWriteIndex + NumWritten;
 
-			// TODO: we should have a 0.5 bias here
-			FVector2D LocalCoordinate((2.0 * X + 0.5) / XCount - 1.0, (2.0 * Y + 0.5) / YCount - 1.0);
+			// Full expression: -1.0 + 2.0 * (X + /*TexelCenter*/0.5) / XCount
+			FVector2D LocalCoordinate((2.0 * X + 1.0) / XCount - 1.0, (2.0 * Y + 1.0) / YCount - 1.0);
 			FLinearColor Color = FLinearColor(EForceInit::ForceInit);
 
 			if (PCGTextureSamplingHelpers::Sample<FLinearColor>(LocalCoordinate, Surface, this, Width, Height, Color, [this](int32 Index) { return ColorData[Index]; }))
