[PCG] Cached inputs on operations that are more likely to get them multiple time as they are time-sliced.

#rnx
#Synced-CL 49049835

diff --git a/Engine/Plugins/PCG/Source/PCG/Private/Elements/PCGAttributeNoise.cpp b/Engine/Plugins/PCG/Source/PCG/Private/Elements/PCGAttributeNoise.cpp
--- a/Engine/Plugins/PCG/Source/PCG/Private/Elements/PCGAttributeNoise.cpp
+++ b/Engine/Plugins/PCG/Source/PCG/Private/Elements/PCGAttributeNoise.cpp
@@ -217,12 +217,18 @@ bool FPCGAttributeNoiseElement::ExecuteInternal(FPCGContext* InContext) const
 	const UPCGAttributeNoiseSettings* Settings = Context->GetInputSettings<UPCGAttributeNoiseSettings>();
 	check(Settings);
 
-	TArray<FPCGTaggedData> Inputs = Context->InputData.GetInputsByPin(PCGPinConstants::DefaultInputLabel);
+	const TArray<FPCGTaggedData>& Inputs = Context->CachedInputs;
 	TArray<FPCGTaggedData>& Outputs = Context->OutputData.TaggedData;
 
 	// Precompute a seed based on the settings one and the component one
 	const int Seed = Context->GetSeed();
 
+	if (!Context->bInputsGathered)
+	{
+		Context->CachedInputs = Context->InputData.GetInputsByPin(PCGPinConstants::DefaultInputLabel);
+		Context->bInputsGathered = true;
+	}
+
 	while (Context->CurrentInput < Inputs.Num())
 	{
 		int32 CurrentInput = Context->CurrentInput;
diff --git a/Engine/Plugins/PCG/Source/PCG/Private/Elements/PCGBoundsFromMesh.cpp b/Engine/Plugins/PCG/Source/PCG/Private/Elements/PCGBoundsFromMesh.cpp
--- a/Engine/Plugins/PCG/Source/PCG/Private/Elements/PCGBoundsFromMesh.cpp
+++ b/Engine/Plugins/PCG/Source/PCG/Private/Elements/PCGBoundsFromMesh.cpp
@@ -82,6 +82,12 @@ bool FPCGBoundsFromMeshElement::PrepareDataInternal(FPCGContext* InContext) cons
 	const UPCGBoundsFromMeshSettings* Settings = Context->GetInputSettings<UPCGBoundsFromMeshSettings>();
 	check(Settings);
 
+	if (!Context->bInputsGathered)
+	{
+		Context->CachedInputs = Context->InputData.GetInputsByPin(PCGPinConstants::DefaultInputLabel);
+		Context->bInputsGathered = true;
+	}
+
 	if (!Context->bPrepareDone)
 	{
 		// Some background info:
@@ -131,7 +137,7 @@ bool FPCGBoundsFromMeshElement::PrepareDataInternal(FPCGContext* InContext) cons
 		else
 		{
 			// In this case, we'll visit all points, gather the unique meshes and assign to each point a given mesh index
-			TArray<FPCGTaggedData> Inputs = Context->InputData.GetInputsByPin(PCGPinConstants::DefaultInputLabel);
+			const TArray<FPCGTaggedData>& Inputs = Context->CachedInputs;
 
 			while (Context->CurrentPrepareIndex < Inputs.Num())
 			{
@@ -222,7 +228,7 @@ bool FPCGBoundsFromMeshElement::ExecuteInternal(FPCGContext* InContext) const
 	const UPCGBoundsFromMeshSettings* Settings = Context->GetInputSettings<UPCGBoundsFromMeshSettings>();
 	check(Settings);
 
-	TArray<FPCGTaggedData> Inputs = Context->InputData.GetInputsByPin(PCGPinConstants::DefaultInputLabel);
+	const TArray<FPCGTaggedData>& Inputs = Context->CachedInputs;
 	TArray<FPCGTaggedData>& Outputs = Context->OutputData.TaggedData;
 
 	auto BoundsFromMesh = [](const UObject* MeshObject, FBox& OutBounds) -> bool
diff --git a/Engine/Plugins/PCG/Source/PCG/Public/Elements/PCGAttributeNoise.h b/Engine/Plugins/PCG/Source/PCG/Public/Elements/PCGAttributeNoise.h
--- a/Engine/Plugins/PCG/Source/PCG/Public/Elements/PCGAttributeNoise.h
+++ b/Engine/Plugins/PCG/Source/PCG/Public/Elements/PCGAttributeNoise.h
@@ -120,6 +120,8 @@ struct FPCGAttributeNoiseContext : public FPCGContext
 {
 	int32 CurrentInput = 0;
 	bool bDataPreparedForCurrentInput = false;
+	bool bInputsGathered = false; 
+	TArray<FPCGTaggedData> CachedInputs;
 	FPCGAttributePropertyInputSelector InputSource;
 	FPCGAttributePropertyInputSelector CustomSeedSource;
 	FPCGAttributePropertyOutputSelector OutputTarget;
diff --git a/Engine/Plugins/PCG/Source/PCG/Public/Elements/PCGBoundsFromMesh.h b/Engine/Plugins/PCG/Source/PCG/Public/Elements/PCGBoundsFromMesh.h
--- a/Engine/Plugins/PCG/Source/PCG/Public/Elements/PCGBoundsFromMesh.h
+++ b/Engine/Plugins/PCG/Source/PCG/Public/Elements/PCGBoundsFromMesh.h
@@ -57,12 +57,14 @@ struct FPCGBoundsFromMeshContext : public FPCGContext, public IPCGAsyncLoadingCo
 		TArray<int32> MeshValueIndex;
 	};
 
+	TArray<FPCGTaggedData> CachedInputs;
 	TArray<InputMeshesData> PerInputData;
 
 	FSoftObjectPath SingleMesh;
 	TMap<int, FBox> MeshToBoundsMap;
 
 	TArray<FSoftObjectPath> MeshesToLoad;
+	bool bInputsGathered = false;
 	bool bPrepareDone = false;
 	bool bBoundsQueried = false;
 	int CurrentPrepareIndex = 0;
