FSpline Refactor: Prepare FSpline for non-tangent based splines.

#Synced-CL 49044627

diff --git a/Engine/Plugins/PCG/Source/PCG/Private/Data/PCGSplineStruct.cpp b/Engine/Plugins/PCG/Source/PCG/Private/Data/PCGSplineStruct.cpp
--- a/Engine/Plugins/PCG/Source/PCG/Private/Data/PCGSplineStruct.cpp
+++ b/Engine/Plugins/PCG/Source/PCG/Private/Data/PCGSplineStruct.cpp
@@ -2,6 +2,7 @@
 
 #include "Data/PCGSplineStruct.h"
 
+#include "Curves/Splines/TangentSpline.h"
 #include "Metadata/PCGMetadataCommon.h"
 
 #include UE_INLINE_GENERATED_CPP_BY_NAME(PCGSplineStruct)
@@ -96,8 +97,17 @@ namespace PCGSplineStruct
 //
 //		return FBoxSphereBounds(BoundingBox.TransformBy(LocalToWorld));
 //#else
-		const int32 NumPoints = Spline.GetNumControlPoints();
+		const FTangentSpline* TangentSpline = Spline.Get<UE::Spline::ESplineType::Tangent>();
+
+		if (!TangentSpline)
+		{
+			ensureAlwaysMsgf(false, TEXT("ShouldUseSplineCurves() returned false but Spline failed to yield an FTangentSpline!"));
+			return FBoxSphereBounds();
+		}
+
+		const int32 NumPoints = TangentSpline->GetSpline().GetNumPoints();
 		const int32 NumSegments = bClosedLoop ? NumPoints : NumPoints - 1;
+		const FInterpCurveVector& SplinePoints = TangentSpline->GetSplinePointsPosition();
 
 		FVector Min(WORLD_MAX);
 		FVector Max(-WORLD_MAX);
@@ -107,11 +117,11 @@ namespace PCGSplineStruct
 			{
 				const bool bLoopSegment = (Index == NumPoints - 1);
 				const int32 NextIndex = bLoopSegment ? 0 : Index + 1;
-				const FInterpCurvePoint<FVector>& ThisInterpPoint = Spline.GetSplinePointsPosition().Points[Index];
-				FInterpCurvePoint<FVector> NextInterpPoint = Spline.GetSplinePointsPosition().Points[NextIndex];
+				const FInterpCurvePoint<FVector>& ThisInterpPoint = SplinePoints.Points[Index];
+				FInterpCurvePoint<FVector> NextInterpPoint = SplinePoints.Points[NextIndex];
 				if (bLoopSegment)
 				{
-					NextInterpPoint.InVal = ThisInterpPoint.InVal + Spline.GetSplinePointsPosition().LoopKeyOffset;
+					NextInterpPoint.InVal = ThisInterpPoint.InVal + SplinePoints.LoopKeyOffset;
 				}
 
 				CurveVectorFindIntervalBounds(ThisInterpPoint, NextInterpPoint, Min, Max);
@@ -119,7 +129,7 @@ namespace PCGSplineStruct
 		}
 		else if (NumPoints == 1)
 		{
-			Min = Max = Spline.GetSplinePointsPosition().Points[0].OutVal;
+			Min = Max = SplinePoints.Points[0].OutVal;
 		}
 		else
 		{
@@ -152,6 +162,15 @@ void FPCGSplineStruct::Initialize(const USplineComponent* InSplineComponent)
 
 void FPCGSplineStruct::Initialize(const TArray<FSplinePoint>& InSplinePoints, bool bIsClosedLoop, const FTransform& InTransform, TArray<PCGMetadataEntryKey> InOptionalEntryKeys)
 {
+	if (PCGSplineStruct::GUseSplineCurves)
+	{
+		Spline = FSpline(UE::Spline::ESplineType::Unimplemented);
+	}
+	else
+	{
+		Spline = FSpline(UE::Spline::ESplineType::Tangent);
+	}
+
 	Transform = InTransform;
 	DefaultUpVector = FVector::ZAxisVector;
 	ReparamStepsPerSegment = 10; // default value in USplineComponent
@@ -187,7 +206,6 @@ void FPCGSplineStruct::ApplyTo(USplineComponent* InSplineComponent) const
 	InSplineComponent->DefaultUpVector = DefaultUpVector;
 	InSplineComponent->ReparamStepsPerSegment = ReparamStepsPerSegment;
 
-	// Not using a ternary so we are not doing any implicit conversions (SetSpline is overloaded)
 	if (ShouldUseSplineCurves())
 	{
 		InSplineComponent->SetSpline(WarninglessSplineCurves());
@@ -252,7 +270,10 @@ FPCGSplineStruct& FPCGSplineStruct::operator=(FPCGSplineStruct&& Other)
 
 void FPCGSplineStruct::AddPoint(const FSplinePoint& InSplinePoint, bool bUpdateSpline, const int64* InOptionalControlPointEntryKey)
 {
-	Spline.AddPoint(InSplinePoint);
+	if (FTangentSpline* TangentSpline = Spline.Get<UE::Spline::ESplineType::Tangent>())
+	{
+		TangentSpline->AddPoint(InSplinePoint);
+	}
  
 	const int32 Index = PCGSplineStruct::UpperBound(WarninglessSplineCurves().Position.Points, InSplinePoint.InputKey);
 	
@@ -324,79 +345,181 @@ void FPCGSplineStruct::AddPoints(const TArray<FSplinePoint>& InSplinePoints, boo
 void FPCGSplineStruct::SetLocation(int32 Index, const FVector& InLocation)
 {
 	WarninglessSplineCurves().Position.Points[Index].OutVal = InLocation;
-	Spline.SetLocation(Index, InLocation);
+
+	if (FTangentSpline* TangentSpline = Spline.Get<UE::Spline::ESplineType::Tangent>())
+	{
+		TangentSpline->SetLocation(Index, InLocation);
+	}
 }
 
 FVector FPCGSplineStruct::GetLocation(const int32 Index) const
 {
-	return ShouldUseSplineCurves()
-		? WarninglessSplineCurves().Position.Points[Index].OutVal
-		: Spline.GetLocation(Index);
+	FVector Location = FVector::ZeroVector;
+
+	if (ShouldUseSplineCurves())
+	{
+		Location = WarninglessSplineCurves().Position.Points[Index].OutVal;
+	}
+	else if (const FTangentSpline* TangentSpline = Spline.Get<UE::Spline::ESplineType::Tangent>())
+	{
+		Location = TangentSpline->GetLocation(Index);
+	}
+	else
+	{
+		ensureAlwaysMsgf(false, TEXT("ShouldUseSplineCurves() returned false but Spline failed to yield an FTangentSpline!"));
+	}
+
+	return Location;
 }
 	
 void FPCGSplineStruct::SetArriveTangent(const int32 Index, const FVector& InArriveTangent)
 {
 	WarninglessSplineCurves().Position.Points[Index].ArriveTangent = InArriveTangent;
-	Spline.SetInTangent(Index, InArriveTangent);
+
+	if (FTangentSpline* TangentSpline = Spline.Get<UE::Spline::ESplineType::Tangent>())
+	{
+		TangentSpline->SetInTangent(Index, InArriveTangent);
+	}
 }
 
 FVector FPCGSplineStruct::GetArriveTangent(const int32 Index) const
 {
-	return ShouldUseSplineCurves()
-		? WarninglessSplineCurves().Position.Points[Index].ArriveTangent
-		: Spline.GetInTangent(Index);
+	FVector Tangent = FVector::ZeroVector;
+
+	if (ShouldUseSplineCurves())
+	{
+		Tangent = WarninglessSplineCurves().Position.Points[Index].ArriveTangent;
+	}
+	else if (const FTangentSpline* TangentSpline = Spline.Get<UE::Spline::ESplineType::Tangent>())
+	{
+		Tangent = TangentSpline->GetInTangent(Index);
+	}
+	else
+	{
+		ensureAlwaysMsgf(false, TEXT("ShouldUseSplineCurves() returned false but Spline failed to yield an FTangentSpline!"));
+	}
+
+	return Tangent;
 }
 	
 void FPCGSplineStruct::SetLeaveTangent(const int32 Index, const FVector& InLeaveTangent)
 {
 	WarninglessSplineCurves().Position.Points[Index].LeaveTangent = InLeaveTangent;
-	Spline.SetOutTangent(Index, InLeaveTangent);
+
+	if (FTangentSpline* TangentSpline = Spline.Get<UE::Spline::ESplineType::Tangent>())
+	{
+		TangentSpline->SetOutTangent(Index, InLeaveTangent);
+	}
 }
 
 FVector FPCGSplineStruct::GetLeaveTangent(const int32 Index) const
 {
-	return ShouldUseSplineCurves()
-		? WarninglessSplineCurves().Position.Points[Index].LeaveTangent
-		: Spline.GetOutTangent(Index);
+	FVector Tangent = FVector::ZeroVector;
+
+	if (ShouldUseSplineCurves())
+	{
+		Tangent = WarninglessSplineCurves().Position.Points[Index].LeaveTangent;
+	}
+	else if (const FTangentSpline* TangentSpline = Spline.Get<UE::Spline::ESplineType::Tangent>())
+	{
+		Tangent = TangentSpline->GetOutTangent(Index);
+	}
+	else
+	{
+		ensureAlwaysMsgf(false, TEXT("ShouldUseSplineCurves() returned false but Spline failed to yield an FTangentSpline!"));
+	}
+
+	return Tangent;
 }
 
 void FPCGSplineStruct::SetRotation(int32 Index, const FQuat& InRotation)
 {
 	WarninglessSplineCurves().Rotation.Points[Index].OutVal = InRotation;
-	Spline.SetRotation(Index, InRotation);
+
+	if (FTangentSpline* TangentSpline = Spline.Get<UE::Spline::ESplineType::Tangent>())
+	{
+		TangentSpline->SetRotation(Index, InRotation);
+	}
 }
 
 FQuat FPCGSplineStruct::GetRotation(const int32 Index) const
 {
-	return ShouldUseSplineCurves()
-		? WarninglessSplineCurves().Rotation.Points[Index].OutVal
-		: Spline.GetRotation(Index);
+	FQuat Quat = FQuat::Identity;
+
+	if (ShouldUseSplineCurves())
+	{
+		Quat = WarninglessSplineCurves().Rotation.Points[Index].OutVal;
+	}
+	else if (const FTangentSpline* TangentSpline = Spline.Get<UE::Spline::ESplineType::Tangent>())
+	{
+		Quat = TangentSpline->GetRotation(Index);
+	}
+	else
+	{
+		ensureAlwaysMsgf(false, TEXT("ShouldUseSplineCurves() returned false but Spline failed to yield an FTangentSpline!"));
+	}
+
+	return Quat;
 }
 	
 void FPCGSplineStruct::SetScale(int32 Index, const FVector& InScale)
 {
 	WarninglessSplineCurves().Scale.Points[Index].OutVal = InScale;
-	Spline.SetScale(Index, InScale);
+
+	if (FTangentSpline* TangentSpline = Spline.Get<UE::Spline::ESplineType::Tangent>())
+	{
+		TangentSpline->SetScale(Index, InScale);
+	}
 }
 
 FVector FPCGSplineStruct::GetScale(const int32 Index) const
 {
-	return ShouldUseSplineCurves()
-		? WarninglessSplineCurves().Scale.Points[Index].OutVal
-		: Spline.GetScale(Index);
+	FVector Scale = FVector::OneVector;
+
+	if (ShouldUseSplineCurves())
+	{
+		Scale = WarninglessSplineCurves().Scale.Points[Index].OutVal;
+	}
+	else if (const FTangentSpline* TangentSpline = Spline.Get<UE::Spline::ESplineType::Tangent>())
+	{
+		Scale = TangentSpline->GetScale(Index);
+	}
+	else
+	{
+		ensureAlwaysMsgf(false, TEXT("ShouldUseSplineCurves() returned false but Spline failed to yield an FTangentSpline!"));
+	}
+
+	return Scale;
 }
 	
 void FPCGSplineStruct::SetSplinePointType(int32 Index, EInterpCurveMode Type)
 {
 	WarninglessSplineCurves().Position.Points[Index].InterpMode = Type;
-	Spline.SetSplinePointType(Index, Type);
+
+	if (FTangentSpline* TangentSpline = Spline.Get<UE::Spline::ESplineType::Tangent>())
+	{
+		TangentSpline->SetSplinePointType(Index, Type);
+	}
 }
 
 EInterpCurveMode FPCGSplineStruct::GetSplinePointType(int32 Index) const
 {
-	return ShouldUseSplineCurves()
-		? WarninglessSplineCurves().Position.Points[Index].InterpMode.GetValue()
-		: Spline.GetSplinePointType(Index);
+	EInterpCurveMode Mode = CIM_Unknown;
+
+	if (ShouldUseSplineCurves())
+	{
+		Mode = WarninglessSplineCurves().Position.Points[Index].InterpMode.GetValue();
+	}
+	else if (const FTangentSpline* TangentSpline = Spline.Get<UE::Spline::ESplineType::Tangent>())
+	{
+		Mode = TangentSpline->GetSplinePointType(Index);
+	}
+	else
+	{
+		ensureAlwaysMsgf(false, TEXT("ShouldUseSplineCurves() returned false but Spline failed to yield an FTangentSpline!"));
+	}
+
+	return Mode;
 }
 
 const TArray<FName>& FPCGSplineStruct::GetAttributeChannelNames() const
@@ -408,14 +531,17 @@ const TArray<FName>& FPCGSplineStruct::GetAttributeChannelNames() const
 			CachedSortedAttributeChannelNames.Empty();
 		}
 	}
-	else if (TArray<FName> Channels = Spline.GetFloatPropertyChannels(); Channels.Num() != CachedSortedAttributeChannelNames.Num())
+	else if (const FTangentSpline* TangentSpline = Spline.Get<UE::Spline::ESplineType::Tangent>())
 	{
-		Channels.Sort([](const FName& A, const FName& B)
-			{ return A.ToString() < B.ToString(); });
-		
-		if (PCG::TScopeLock ScopeLock(Lock); Channels.Num() != CachedSortedAttributeChannelNames.Num())
+		if (TArray<FName> Channels = TangentSpline->GetAttributeChannelNamesByValueType<float>(); Channels.Num() != CachedSortedAttributeChannelNames.Num())
 		{
-			CachedSortedAttributeChannelNames = MoveTemp(Channels);
+			Channels.Sort([](const FName& A, const FName& B)
+				{ return A.ToString() < B.ToString(); });
+
+			if (PCG::TScopeLock ScopeLock(Lock); Channels.Num() != CachedSortedAttributeChannelNames.Num())
+			{
+				CachedSortedAttributeChannelNames = MoveTemp(Channels);
+			}
 		}
 	}
 
@@ -424,21 +550,41 @@ const TArray<FName>& FPCGSplineStruct::GetAttributeChannelNames() const
 
 float FPCGSplineStruct::GetAttributeValue(float InputKey, FName ChannelName) const
 {
-	return ShouldUseSplineCurves()
-		? 0.f
-		: Spline.EvaluateAttribute<float>(InputKey, ChannelName);
+	if (ShouldUseSplineCurves())
+	{
+		return 0.f;
+	}
+	else if (const FTangentSpline* TangentSpline = Spline.Get<UE::Spline::ESplineType::Tangent>())
+	{
+		return TangentSpline->EvaluateAttribute<float>(ChannelName, InputKey);
+	}
+	else
+	{
+		ensureAlwaysMsgf(false, TEXT("ShouldUseSplineCurves() returned false but Spline failed to yield an FTangentSpline!"));
+		return 0.f;
+	}
 }
 
 void FPCGSplineStruct::UpdateSpline()
 {
+	// If SplineComponent.UseSplineCurves is false and FSpline is not implemented, initialize it from SplineCurves.
+	if (PCGSplineStruct::GUseSplineCurves == false && Spline.GetCurrentImplementation() != UE::Spline::ESplineType::Tangent)
+	{
+		Spline = WarninglessSplineCurves();
+	}
+
 	const bool bLoopPositionOverride = false;
 	const bool bStationaryEndpoints = false;
 	const float LoopPosition = 0.0f;
 
-	FSpline::FUpdateSplineParams Params{bClosedLoop, bStationaryEndpoints, ReparamStepsPerSegment, bLoopPositionOverride, LoopPosition};
+	FTangentSpline::FUpdateSplineParams Params{bClosedLoop, bStationaryEndpoints, ReparamStepsPerSegment, bLoopPositionOverride, LoopPosition};
 
 	WarninglessSplineCurves().UpdateSpline(Params.bClosedLoop, Params.bStationaryEndpoints, Params.ReparamStepsPerSegment, Params.bLoopPositionOverride, Params.LoopPosition, Params.Scale3D);
-	Spline.UpdateSpline(Params);
+
+	if (FTangentSpline* TangentSpline = Spline.Get<UE::Spline::ESplineType::Tangent>())
+	{
+		TangentSpline->UpdateSpline(Params);
+	}
 }
 
 int FPCGSplineStruct::GetNumberOfSplineSegments() const
@@ -449,16 +595,36 @@ int FPCGSplineStruct::GetNumberOfSplineSegments() const
 
 int FPCGSplineStruct::GetNumberOfPoints() const
 {
-	return ShouldUseSplineCurves()
-		? WarninglessSplineCurves().Position.Points.Num()
-		: Spline.GetNumControlPoints();
+	if (ShouldUseSplineCurves())
+	{
+		return WarninglessSplineCurves().Position.Points.Num();
+	}
+	else if (const FTangentSpline* TangentSpline = Spline.Get<UE::Spline::ESplineType::Tangent>())
+	{
+		return TangentSpline->GetSpline().GetNumPoints();
+	}
+	else
+	{
+		ensureAlwaysMsgf(false, TEXT("ShouldUseSplineCurves() returned false but Spline failed to yield an FTangentSpline!"));
+		return 0;
+	}
 }
 
 float FPCGSplineStruct::GetLoopKeyOffset() const
 {
-	return ShouldUseSplineCurves()
-		? WarninglessSplineCurves().Position.LoopKeyOffset
-		: Spline.GetSplinePointsPosition().LoopKeyOffset;
+	if (ShouldUseSplineCurves())
+	{
+		return WarninglessSplineCurves().Position.LoopKeyOffset;
+	}
+	else if (const FTangentSpline* TangentSpline = Spline.Get<UE::Spline::ESplineType::Tangent>())
+	{
+		return TangentSpline->GetSplinePointsPosition().LoopKeyOffset;
+	}
+	else
+	{
+		ensureAlwaysMsgf(false, TEXT("ShouldUseSplineCurves() returned false but Spline failed to yield an FTangentSpline!"));
+		return 0.f;
+	}
 }
 
 void FPCGSplineStruct::SetLoopKeyOffset(float LoopKeyOffset)
@@ -470,17 +636,34 @@ void FPCGSplineStruct::SetLoopKeyOffset(float LoopKeyOffset)
 		const bool bLoopPositionOverride = true;
 		const float LoopPosition = LoopKeyOffset + (WarninglessSplineCurves().Position.Points.IsEmpty() ? 0.0f : WarninglessSplineCurves().Position.Points.Last().InVal);
 
-		FSpline::FUpdateSplineParams Params{ bClosedLoop, bStationaryEndpoints, ReparamStepsPerSegment, bLoopPositionOverride, LoopPosition };
+		FTangentSpline::FUpdateSplineParams Params{ bClosedLoop, bStationaryEndpoints, ReparamStepsPerSegment, bLoopPositionOverride, LoopPosition };
+
 		WarninglessSplineCurves().UpdateSpline(Params.bClosedLoop, Params.bStationaryEndpoints, Params.ReparamStepsPerSegment, Params.bLoopPositionOverride, Params.LoopPosition, Params.Scale3D);
-		Spline.UpdateSpline(Params);
+
+		if (FTangentSpline* TangentSpline = Spline.Get<UE::Spline::ESplineType::Tangent>())
+		{
+			TangentSpline->UpdateSpline(Params);
+
+		}
 	}
 }
 
 FVector::FReal FPCGSplineStruct::GetSplineLength() const
 {
-	return ShouldUseSplineCurves()
-		? WarninglessSplineCurves().GetSplineLength()
-		: Spline.GetSplineLength();
+	if (ShouldUseSplineCurves())
+	{
+		return WarninglessSplineCurves().GetSplineLength();
+	}
+	else if (const FTangentSpline* TangentSpline = Spline.Get<UE::Spline::ESplineType::Tangent>())
+	{
+		const float MaxParameter = TangentSpline->GetParameterSpace().Max;
+		return TangentSpline->GetDistanceAtParameter(MaxParameter);
+	}
+	else
+	{
+		ensureAlwaysMsgf(false, TEXT("ShouldUseSplineCurves() returned false but Spline failed to yield an FTangentSpline!"));
+		return 0.f;
+	}
 }
 
 FBox FPCGSplineStruct::GetBounds() const
@@ -506,11 +689,9 @@ FBox FPCGSplineStruct::GetSegmentBounds(int32 SegmentIndex) const
 		return FBox(EForceInit::ForceInit);
 	}
 
-	const FInterpCurveVector& SplinePointsPosition = ShouldUseSplineCurves()
-		? WarninglessSplineCurves().Position
-		: Spline.GetSplinePointsPosition();
+	const FInterpCurveVector& SplinePointsPosition = GetSplinePointsPosition();
 
-	const int32 NumPoints = GetNumberOfPoints();
+	const int32 NumPoints = SplinePointsPosition.Points.Num();
 	const bool bLoopSegment = (SegmentIndex == NumPoints - 1);
 	const int32 NextIndex = bLoopSegment ? 0 : SegmentIndex + 1;
 	const FInterpCurvePoint<FVector>& ThisInterpPoint = SplinePointsPosition.Points[SegmentIndex];
@@ -530,23 +711,71 @@ FBox FPCGSplineStruct::GetSegmentBounds(int32 SegmentIndex) const
 
 const FInterpCurveVector& FPCGSplineStruct::GetSplinePointsScale() const
 {
-	return ShouldUseSplineCurves()
-		? WarninglessSplineCurves().Scale
-		: Spline.GetSplinePointsScale();
+	const FTangentSpline* TangentSpline = Spline.Get<UE::Spline::ESplineType::Tangent>();
+
+	// We normally branch directly on ShouldUseSplineCurves(), but because this function returns a const-ref
+	// we need a safe fallback if ShouldUseSplineCurves() is false and TangentSpline is unexpectedly nullptr.
+	// Cache the value and assert the contract up front, then fall back to the spline curves data if needed.
+	const bool bShouldUseSplineCurves = ShouldUseSplineCurves();
+	if (!bShouldUseSplineCurves)
+	{
+		ensureAlwaysMsgf(TangentSpline, TEXT("ShouldUseSplineCurves() returned false but Spline failed to yield an FTangentSpline!"));
+	}
+
+	if (bShouldUseSplineCurves || !TangentSpline)
+	{
+		return WarninglessSplineCurves().Scale;
+	}
+	else
+	{
+		return TangentSpline->GetSplinePointsScale();
+	}
 }
 
 const FInterpCurveVector& FPCGSplineStruct::GetSplinePointsPosition() const
 {
-	return ShouldUseSplineCurves()
-		? WarninglessSplineCurves().Position
-		: Spline.GetSplinePointsPosition();
+	const FTangentSpline* TangentSpline = Spline.Get<UE::Spline::ESplineType::Tangent>();
+
+	// We normally branch directly on ShouldUseSplineCurves(), but because this function returns a const-ref
+	// we need a safe fallback if ShouldUseSplineCurves() is false and TangentSpline is unexpectedly nullptr.
+	// Cache the value and assert the contract up front, then fall back to the spline curves data if needed.
+	const bool bShouldUseSplineCurves = ShouldUseSplineCurves();
+	if (!bShouldUseSplineCurves)
+	{
+		ensureAlwaysMsgf(TangentSpline, TEXT("ShouldUseSplineCurves() returned false but Spline failed to yield an FTangentSpline!"));
+	}
+
+	if (bShouldUseSplineCurves || !TangentSpline)
+	{
+		return WarninglessSplineCurves().Position;
+	}
+	else
+	{
+		return TangentSpline->GetSplinePointsPosition();
+	}
 }
 
 const FInterpCurveQuat& FPCGSplineStruct::GetSplinePointsRotation() const
 {
-	return ShouldUseSplineCurves()
-		? WarninglessSplineCurves().Rotation
-		: Spline.GetSplinePointsRotation();
+	const FTangentSpline* TangentSpline = Spline.Get<UE::Spline::ESplineType::Tangent>();
+
+	// We normally branch directly on ShouldUseSplineCurves(), but because this function returns a const-ref
+	// we need a safe fallback if ShouldUseSplineCurves() is false and TangentSpline is unexpectedly nullptr.
+	// Cache the value and assert the contract up front, then fall back to the spline curves data if needed.
+	const bool bShouldUseSplineCurves = ShouldUseSplineCurves();
+	if (!bShouldUseSplineCurves)
+	{
+		ensureAlwaysMsgf(TangentSpline, TEXT("ShouldUseSplineCurves() returned false but Spline failed to yield an FTangentSpline!"));
+	}
+
+	if (bShouldUseSplineCurves || !TangentSpline)
+	{
+		return WarninglessSplineCurves().Rotation;
+	}
+	else
+	{
+		return TangentSpline->GetSplinePointsRotation();
+	}
 }
 
 FVector::FReal FPCGSplineStruct::GetDistanceAlongSplineAtSplinePoint(int32 PointIndex) const
@@ -570,10 +799,14 @@ FVector::FReal FPCGSplineStruct::GetDistanceAlongSplineAtSplinePoint(int32 Point
 			return WarninglessSplineCurves().ReparamTable.Points[PointIndex * ReparamStepsPerSegment].InVal;
 		}
 	}
+	else if (const FTangentSpline* TangentSpline = Spline.Get<UE::Spline::ESplineType::Tangent>())
+	{
+		const float ParameterAtIndex = TangentSpline->GetParameterAtIndex(PointIndex);
+		return TangentSpline->GetDistanceAtParameter(ParameterAtIndex);
+	}
 	else
 	{
-		const float ParameterAtIndex = Spline.GetParameterAtIndex(PointIndex);
-		return Spline.GetDistanceAtParameter(ParameterAtIndex);
+		ensureAlwaysMsgf(false, TEXT("ShouldUseSplineCurves() returned false but Spline failed to yield an FTangentSpline!"));
 	}
 
 	return 0.0f;
@@ -581,19 +814,13 @@ FVector::FReal FPCGSplineStruct::GetDistanceAlongSplineAtSplinePoint(int32 Point
 
 FVector FPCGSplineStruct::GetLocationAtDistanceAlongSpline(FVector::FReal Distance, ESplineCoordinateSpace::Type CoordinateSpace) const
 {
-	const float Param = ShouldUseSplineCurves()
-		? WarninglessSplineCurves().ReparamTable.Eval(Distance, 0.0f)
-		: Spline.GetParameterAtDistance(Distance);
-	
+	const float Param = GetInputKeyAtDistanceAlongSpline(Distance);
 	return GetLocationAtSplineInputKey(Param, CoordinateSpace);
 }
 
 FTransform FPCGSplineStruct::GetTransformAtDistanceAlongSpline(FVector::FReal Distance, ESplineCoordinateSpace::Type CoordinateSpace, bool bUseScale) const
 {
-	const float Param = ShouldUseSplineCurves()
-		? WarninglessSplineCurves().ReparamTable.Eval(Distance, 0.0f)
-		: Spline.GetParameterAtDistance(Distance);
-	
+	const float Param = GetInputKeyAtDistanceAlongSpline(Distance);
 	return GetTransformAtSplineInputKey(Param, CoordinateSpace, bUseScale);
 }
 
@@ -605,10 +832,20 @@ float FPCGSplineStruct::GetInputKeyAtDistanceAlongSpline(FVector::FReal Distance
 	{
 		return 0.0f;
 	}
-	
-	return ShouldUseSplineCurves()
-		? WarninglessSplineCurves().ReparamTable.Eval(Distance, 0.0f)
-		: Spline.GetParameterAtDistance(Distance);
+
+	if (ShouldUseSplineCurves())
+	{
+		return WarninglessSplineCurves().ReparamTable.Eval(Distance, 0.0f);
+	}
+	else if (const FTangentSpline* TangentSpline = Spline.Get<UE::Spline::ESplineType::Tangent>())
+	{
+		return TangentSpline->GetParameterAtDistance(Distance);
+	}
+	else
+	{
+		ensureAlwaysMsgf(false, TEXT("ShouldUseSplineCurves() returned false but Spline failed to yield an FTangentSpline!"));
+		return 0.f;
+	}
 }
 
 FVector FPCGSplineStruct::GetRightVectorAtSplineInputKey(float InKey, ESplineCoordinateSpace::Type CoordinateSpace) const
@@ -644,33 +881,51 @@ float FPCGSplineStruct::FindInputKeyClosestToWorldLocation(const FVector& WorldL
 {
 	const FVector LocalLocation = Transform.InverseTransformPosition(WorldLocation);
 	float Dummy;
-	return ShouldUseSplineCurves()
-		? WarninglessSplineCurves().Position.FindNearest(LocalLocation, Dummy)
-		: Spline.FindNearest(LocalLocation, Dummy);
+
+	if (ShouldUseSplineCurves())
+	{
+		return WarninglessSplineCurves().Position.FindNearest(LocalLocation, Dummy);
+	}
+	else if (const FTangentSpline* TangentSpline = Spline.Get<UE::Spline::ESplineType::Tangent>())
+	{
+		return TangentSpline->FindNearest(LocalLocation, Dummy);
+	}
+	else
+	{
+		ensureAlwaysMsgf(false, TEXT("ShouldUseSplineCurves() returned false but Spline failed to yield an FTangentSpline!"));
+		return 0.f;
+	}
 }
 
 float FPCGSplineStruct::FindInputKeyOnSegmentClosestToWorldLocation(const FVector& WorldLocation, int32 Index) const
 {
 	const FVector LocalLocation = Transform.InverseTransformPosition(WorldLocation);
 	float Dummy;
-	return ShouldUseSplineCurves()
-		? WarninglessSplineCurves().Position.FindNearestOnSegment(LocalLocation, Index, Dummy)
-		: Spline.FindNearestOnSegment(LocalLocation, Index, Dummy);
+
+	if (ShouldUseSplineCurves())
+	{
+		return WarninglessSplineCurves().Position.FindNearestOnSegment(LocalLocation, Index, Dummy);
+	}
+	else if (const FTangentSpline* TangentSpline = Spline.Get<UE::Spline::ESplineType::Tangent>())
+	{
+		return TangentSpline->FindNearestOnSegment(LocalLocation, Index, Dummy);
+	}
+	else
+	{
+		ensureAlwaysMsgf(false, TEXT("ShouldUseSplineCurves() returned false but Spline failed to yield an FTangentSpline!"));
+		return 0.f;
+	}
 }
 
 TTuple<int, float> FPCGSplineStruct::GetSegmentStartIndexAndKeyAtInputKey(float InKey) const
 {
-	const int32 Index = ShouldUseSplineCurves()
-		? WarninglessSplineCurves().Position.GetPointIndexForInputValue(InKey)
-		: Spline.GetSplinePointsPosition().GetPointIndexForInputValue(InKey);
+	const int32 Index = GetSplinePointsPosition().GetPointIndexForInputValue(InKey);
 	return {Index, GetInputKeyAtSegmentStart(Index)};
 }
 
 float FPCGSplineStruct::GetInputKeyAtSegmentStart(int InSegmentIndex) const
 {
-	const auto& SplinePointsPosition = ShouldUseSplineCurves()
-		? WarninglessSplineCurves().Position
-		: Spline.GetSplinePointsPosition();
+	const auto& SplinePointsPosition = GetSplinePointsPosition();
 
 	if (IsClosedLoop() && !SplinePointsPosition.Points.IsEmpty() && InSegmentIndex >= SplinePointsPosition.Points.Num())
 	{
@@ -719,9 +974,20 @@ FVector FPCGSplineStruct::GetLocationAtSplineInputKey(float InKey, ESplineCoordi
 
 FQuat FPCGSplineStruct::GetQuaternionAtSplineInputKey(float InKey, ESplineCoordinateSpace::Type CoordinateSpace) const
 {
-	FQuat Quat = ShouldUseSplineCurves()
-		? WarninglessSplineCurves().Rotation.Eval(InKey, FQuat::Identity)
-		: Spline.EvaluateRotation(InKey);
+	FQuat Quat = FQuat::Identity;
+	if (ShouldUseSplineCurves())
+	{
+		Quat = WarninglessSplineCurves().Rotation.Eval(InKey, FQuat::Identity);
+	}
+	else if (const FTangentSpline* TangentSpline = Spline.Get<UE::Spline::ESplineType::Tangent>())
+	{
+		Quat = TangentSpline->EvaluateRotation(InKey);
+	}
+	else
+	{
+		ensureAlwaysMsgf(false, TEXT("ShouldUseSplineCurves() returned false but Spline failed to yield an FTangentSpline!"));
+		return Quat;
+	}
 	Quat.Normalize();
 
 	FVector Direction = ShouldUseSplineCurves()
@@ -742,9 +1008,22 @@ FQuat FPCGSplineStruct::GetQuaternionAtSplineInputKey(float InKey, ESplineCoordi
 
 FVector FPCGSplineStruct::GetScaleAtSplineInputKey(float InKey) const
 {
-	return ShouldUseSplineCurves()
-		? WarninglessSplineCurves().Scale.Eval(InKey, FVector(1.0f))
-		: Spline.GetSplinePointsScale().Eval(InKey, FVector(1.0f));
+	FVector Scale = FVector::OneVector;
+
+	if (ShouldUseSplineCurves())
+	{
+		Scale = WarninglessSplineCurves().Scale.Eval(InKey, FVector::OneVector);
+	}
+	else if (const FTangentSpline* TangentSpline = Spline.Get<UE::Spline::ESplineType::Tangent>())
+	{
+		Scale = TangentSpline->EvaluateScale(InKey);
+	}
+	else
+	{
+		ensureAlwaysMsgf(false, TEXT("ShouldUseSplineCurves() returned false but Spline failed to yield an FTangentSpline!"));
+	}
+
+	return Scale;
 }
 
 FVector FPCGSplineStruct::GetTangentAtSplineInputKey(float InKey, ESplineCoordinateSpace::Type CoordinateSpace) const
@@ -891,9 +1170,9 @@ void FPCGSplineStruct::PostSerialize(const FArchive& Ar)
 #if WITH_EDITOR
 	if (Ar.IsLoading() && Ar.IsPersistent() && !Ar.HasAnyPortFlags(PPF_Duplicate | PPF_DuplicateForPIE))
 	{
-		// This is kind of a hack. It will be true for every FPCGSplineStruct which did not have an FSpline before but does not.
-		// It is also going to be true if we already had a Spline but did not save it while it was configured to actually do anything.
-		if (!Spline.WasEnabled())
+		// If PCGSplineStruct.UseSplineCurves is false and Spline was not previously implemented
+		// (implying it is not consistent with SplineCurves), initialize Spline from SplineCurves.
+		if (Spline.GetPreviousImplementation() == UE::Spline::ESplineType::Unimplemented && !PCGSplineStruct::GUseSplineCurves)
 		{
 			Spline = WarninglessSplineCurves();
 		}
@@ -903,5 +1182,5 @@ void FPCGSplineStruct::PostSerialize(const FArchive& Ar)
 
 bool FPCGSplineStruct::ShouldUseSplineCurves() const
 {
-	return PCGSplineStruct::GUseSplineCurves || Spline.IsEnabled() == false;
+	return PCGSplineStruct::GUseSplineCurves || Spline.GetCurrentImplementation() != UE::Spline::ESplineType::Tangent;
 }
\ No newline at end of file
diff --git a/Engine/Source/Editor/ComponentVisualizers/Private/SplineComponentVisualizer.cpp b/Engine/Source/Editor/ComponentVisualizers/Private/SplineComponentVisualizer.cpp
--- a/Engine/Source/Editor/ComponentVisualizers/Private/SplineComponentVisualizer.cpp
+++ b/Engine/Source/Editor/ComponentVisualizers/Private/SplineComponentVisualizer.cpp
@@ -3576,13 +3576,16 @@ void FSplineComponentVisualizer::OnAddAttributeKey()
 	
 	const float NearestParam = SplineComp->FindInputKeyClosestToWorldLocation(SelectionState->GetSelectedSplinePosition());
 	const float CurrentValue = SplineComp->GetFloatPropertyAtSplineInputKey(NearestParam, SelectionState->GetSelectedAttributeName());
-	const int32 NewAttributeIndex = SplineComp->SetFloatPropertyAtSplineInputKey(NearestParam, CurrentValue, SelectionState->GetSelectedAttributeName());
+	const int32 NewAttributeIndex = SplineComp->AddFloatPropertyAtSplineInputKey(NearestParam, CurrentValue, SelectionState->GetSelectedAttributeName());
 
 	// Update selection
-	SelectionState->SetSelectedSegmentIndex(INDEX_NONE);
-	SelectionState->SetSelectedAttribute(NewAttributeIndex, SelectionState->GetSelectedAttributeName());
-	SelectionState->SetSelectedSplinePosition(FVector::ZeroVector);
-	SelectionState->SetCachedRotation(SplineComp->GetQuaternionAtSplineInputKey(NearestParam, ESplineCoordinateSpace::World));
+	if (NewAttributeIndex != INDEX_NONE)
+	{
+		SelectionState->SetSelectedSegmentIndex(INDEX_NONE);
+		SelectionState->SetSelectedAttribute(NewAttributeIndex, SelectionState->GetSelectedAttributeName());
+		SelectionState->SetSelectedSplinePosition(FVector::ZeroVector);
+		SelectionState->SetCachedRotation(SplineComp->GetQuaternionAtSplineInputKey(NearestParam, ESplineCoordinateSpace::World));
+	}
 
 	SplineComp->bSplineHasBeenEdited = true;
 	NotifyPropertiesModified(SplineComp, SplineProperties, EPropertyChangeType::ArrayAdd);
diff --git a/Engine/Source/Editor/DetailCustomizations/Private/SplineComponentDetails.cpp b/Engine/Source/Editor/DetailCustomizations/Private/SplineComponentDetails.cpp
--- a/Engine/Source/Editor/DetailCustomizations/Private/SplineComponentDetails.cpp
+++ b/Engine/Source/Editor/DetailCustomizations/Private/SplineComponentDetails.cpp
@@ -1005,7 +1005,7 @@ bool FSplinePointDetails::CanSetInputKey() const
 	}
 
 PRAGMA_DISABLE_DEPRECATION_WARNINGS
-	bool bUsingSplineCurves = SplineComp->GetSplinePropertyName() == GET_MEMBER_NAME_CHECKED(USplineComponent, SplineCurves);
+	bool bUsingSplineCurves = SplineComp->ShouldUseSplineCurves();
 PRAGMA_ENABLE_DEPRECATION_WARNINGS
 
 	return IsOnePointSelected() && bUsingSplineCurves;
diff --git a/Engine/Source/Runtime/Engine/Classes/Components/SplineComponent.h b/Engine/Source/Runtime/Engine/Classes/Components/SplineComponent.h
--- a/Engine/Source/Runtime/Engine/Classes/Components/SplineComponent.h
+++ b/Engine/Source/Runtime/Engine/Classes/Components/SplineComponent.h
@@ -76,15 +76,7 @@ struct FSplineCurves
 {
 	GENERATED_BODY()
 
-	static FSplineCurves FromSplineInterface(const FSpline& InProxy)
-	{
-		FSplineCurves NewCurve;
-		NewCurve.Position = InProxy.GetSplinePointsPosition();
-		NewCurve.Rotation = InProxy.GetSplinePointsRotation();
-		NewCurve.Scale = InProxy.GetSplinePointsScale();
-		NewCurve.UpdateSpline(NewCurve.Position.bIsLooped, false);
-		return NewCurve;
-	}
+	static FSplineCurves FromSplineInterface(const FSpline& InProxy);
 	
 	/** Spline built from position data. */
 	UPROPERTY()
@@ -230,15 +222,13 @@ protected:
 	
 public:
 
-	UFUNCTION(CallInEditor, Category="Spline")
-	void PopulateFromLegacy();
-
 	/** Return the name of the authoritative spline property. */
 	ENGINE_API static FName GetSplinePropertyName();
 
 	/** Return the set of names of the backing spline properties. */
 	ENGINE_API static TSet<FName> GetSplinePropertyNames();
 	
+	UE_DEPRECATED(5.8, TEXT("SplineCurves will be made protected."))
 	UPROPERTY(EditAnywhere, Replicated, Category=Points)
 	FSplineCurves SplineCurves;
 
@@ -416,9 +406,6 @@ public:
 
 	UE_EXPERIMENTAL(5.7, "Spline display setting change notifications are experimental while the API is under active development.")
 	FSimpleMulticastDelegate& GetOnSplineDisplayChanged() { return OnSplineDisplayChanged; }
-	
-	/** Reset all spline data to match the current state of SplineCurves */
-	void SynchronizeSplines() { SetSpline(WarninglessSplineCurves()); }
 
 	ENGINE_API int32 GetVersion() const;
 	ENGINE_API const FInterpCurveVector& GetSplinePointsPosition() const;
@@ -966,8 +953,8 @@ public:
 	 * @return The index of the new attribute.
 	 */
 	UFUNCTION(BlueprintCallable, Category = "Spline|Metadata")
-	ENGINE_API int32 SetFloatPropertyAtSplineInputKey(float InKey, float Value, FName PropertyName) const;
-	
+	ENGINE_API int32 AddFloatPropertyAtSplineInputKey(float InKey, float Value, FName PropertyName);
+
 	/**
 	 * Get a float attribute parameter by index.
 	 * @param Index The index of the specified attribute.
@@ -1018,6 +1005,18 @@ public:
 	FInterpCurveQuat& GetSplinePointsRotation() { return WarninglessSplineCurves().Rotation; }
 	FInterpCurveVector& GetSplinePointsScale() { return WarninglessSplineCurves().Scale; }
 
+public: // Deprecated Interface
+
+	UE_DEPRECATED(5.8, TEXT("SynchronizeSplines is deprecated."))
+	void SynchronizeSplines() {}
+
+	UE_DEPRECATED(5.8, TEXT("PopulateFromLegacy is deprecated."))
+	void PopulateFromLegacy() {}
+
+	UE_DEPRECATED(5.8, TEXT("SetFloatPropertyAtSplineInputKey is deprecated, please use AddFloatPropertyAtSplineInputKey"))
+	UFUNCTION(BlueprintCallable, Category = "Spline|Metadata", meta = (DeprecatedFunction, DeprecationMessage = "Please use AddFloatPropertyAtSplineInputKey"))
+	ENGINE_API int32 SetFloatPropertyAtSplineInputKey(float InKey, float Value, FName PropertyName) const { return INDEX_NONE; }
+
 protected:
 
 	/** Returns the length of the specified spline segment up to the parametric value given */
@@ -1034,6 +1033,9 @@ protected:
 	
 private:
 	
+	/** Whether or not FSplineCurves is the authoritative data. If false, we assume Spline is implemented as FTangentSpline. */
+	ENGINE_API bool ShouldUseSplineCurves() const;
+
 	/** Returns an interval representing the full range of valid input keys, or an empty interval if the spline is empty. */
 	UE::Geometry::FInterval1f GetInputKeyRange() const;
 	
@@ -1052,9 +1054,9 @@ private:
 	}
 	
 	/** The dummy value used for queries when there are no point in a spline */
-	static ENGINE_API const FInterpCurvePointVector DummyPointPosition;
-	static ENGINE_API const FInterpCurvePointQuat DummyPointRotation;
-	static ENGINE_API const FInterpCurvePointVector DummyPointScale;
+	static ENGINE_API const FInterpCurveVector DummyPointPosition;
+	static ENGINE_API const FInterpCurveQuat DummyPointRotation;
+	static ENGINE_API const FInterpCurveVector DummyPointScale;
 
 	// Remove when appending _DEPRECATED to SplineCurves.
 	FSplineCurves& WarninglessSplineCurves()
@@ -1092,6 +1094,9 @@ PRAGMA_ENABLE_DEPRECATION_WARNINGS
 	// FSplineComponentVisualizer will access some private members when attempting to call NotifyPropertiesModified
 	friend class FSplineComponentVisualizer;
 
+	// FSplinePointDetails will access ShouldUseSplineCurves for enabling/disabling some editing functionality
+	friend class FSplinePointDetails;
+
 	// AFunctionalSplineTestBase accesses private methods for functional testing.
 	friend class AFunctionalSplineTestBase;
 };
diff --git a/Engine/Source/Runtime/Engine/Classes/Curves/Spline.h b/Engine/Source/Runtime/Engine/Classes/Curves/Spline.h
--- a/Engine/Source/Runtime/Engine/Classes/Curves/Spline.h
+++ b/Engine/Source/Runtime/Engine/Classes/Curves/Spline.h
@@ -2,19 +2,38 @@
 
 #pragma once
 
-#include "CoreMinimal.h"
-#include "Math/InterpCurve.h"
 #include "Splines/SplineInterfaces.h"
+#include "Curves/Splines/TangentSpline.h"	// only for FLegacyTangentSpline, move that to its own header.
 
 #include "Spline.generated.h"
 
-struct FSplinePoint;
-struct FSplineCurves;
-class FLegacySpline;
-class FNewSpline;
-
 struct UE_EXPERIMENTAL(5.7, "New spline APIs are experimental.") FSpline;
 
+namespace UE::Spline
+{
+enum class ESplineType : uint8
+{
+	Unimplemented = 0,
+	LegacyTangent = 1,
+	Tangent = 2,
+};
+
+template<ESplineType Impl>
+struct TSplineTypeMap;
+
+template<>
+struct TSplineTypeMap<ESplineType::LegacyTangent>
+{
+	using Type = FLegacyTangentSpline;
+};
+
+template<>
+struct TSplineTypeMap<ESplineType::Tangent>
+{
+	using Type = FTangentSpline;
+};
+} // UE::Spline
+
 /**
  * A general purpose, reflected spline.
  * The implementation can be configured at runtime.
@@ -25,186 +44,125 @@ struct FSpline
 	GENERATED_BODY()
 
 	ENGINE_API FSpline();
+	ENGINE_API explicit FSpline(UE::Spline::ESplineType InType);
+
 	ENGINE_API FSpline(const FSpline& Other);
 	ENGINE_API FSpline& operator=(const FSpline& Other);
-	ENGINE_API FSpline& operator=(const FSplineCurves& Other);
-	
-	/* Control Point Index Interface */
-	
-	/** Adds point by parameter. The new point is inserted after all points with parameter values less than OR EQUAL to the new point's parameter.
-	 * Nick would like to change this behavior in the future. See UE-250236. */
-	ENGINE_API void AddPoint(const FSplinePoint& InPoint);
-	/** Adds point by index. Provided parameter is ignored. The new point is inserted before all points with parameter values greater than the new point's parameter. */
-	ENGINE_API void InsertPoint(const FSplinePoint& InPoint, int32 Index);
-	ENGINE_API FSplinePoint GetPoint(const int32 Index) const;
-	ENGINE_API void RemovePoint(const int32 Index);
-	
-	ENGINE_API void SetLocation(int32 Index, const FVector& InLocation);
-	ENGINE_API FVector GetLocation(const int32 Index) const;
-	
-	ENGINE_API void SetInTangent(const int32 Index, const FVector& InTangent);
-	ENGINE_API FVector GetInTangent(const int32 Index) const;
-	
-	ENGINE_API void SetOutTangent(const int32 Index, const FVector& OutTangent);
-	ENGINE_API FVector GetOutTangent(const int32 Index) const;
-
-	ENGINE_API void SetRotation(int32 Index, const FQuat& InRotation);
-	ENGINE_API FQuat GetRotation(const int32 Index) const;
-	
-	ENGINE_API void SetScale(int32 Index, const FVector& InScale);
-	ENGINE_API FVector GetScale(const int32 Index) const;
-	
-	ENGINE_API void SetSplinePointType(int32 Index, EInterpCurveMode Type);
-	ENGINE_API EInterpCurveMode GetSplinePointType(int32 Index) const;
 
-	ENGINE_API float GetParameterAtIndex(int32 Index) const;
-	ENGINE_API float GetParameterAtDistance(float Distance) const;
-	ENGINE_API float GetDistanceAtParameter(float Parameter) const;
+	ENGINE_API FSpline(FSpline&& Other) = default;
+	ENGINE_API FSpline& operator=(FSpline&& Other) = default;
 
-	ENGINE_API FQuat GetOrientation(int32 Index) const;
-	ENGINE_API void SetOrientation(int32 Index, const FQuat& InOrientation);
-	
-	/* Parameter Interface */
+	ENGINE_API FSpline& operator=(const FSplineCurves& Other);
 
 	ENGINE_API FVector Evaluate(float Param) const;
+
 	ENGINE_API FVector EvaluateDerivative(float Param) const;
-	ENGINE_API FQuat EvaluateRotation(float Param) const;
-	ENGINE_API FVector EvaluateScale(float Param) const;
 
-	ENGINE_API FQuat GetOrientation(float Param) const;
-	
-	/* Attribute Interface */
-
-	// Non-Templated functions
-	bool SupportsAttributes() const { return IsNew(); }
-	ENGINE_API bool HasAttributeChannel(FName AttributeName) const;
-	ENGINE_API bool RemoveAttributeChannel(FName AttributeName) const;
-	ENGINE_API TArray<FName> GetFloatPropertyChannels() const;
-	ENGINE_API TArray<FName> GetVectorPropertyChannels() const;
-
-	template <typename AttrType> ENGINE_API float GetAttributeParameter(int32 Index, const FName& Name) const;
-	template <typename AttrType> ENGINE_API int32 SetAttributeParameter(int32 Index, float Parameter, const FName& Name);
-	template <typename AttrType> ENGINE_API int32 NumAttributeValues(FName AttributeName) const;
-	template <typename AttrType> ENGINE_API AttrType GetAttributeValue(int32 Index, const FName& Name) const;
-	template <typename AttrType> ENGINE_API void SetAttributeValue(int32 Index, const AttrType& Value, const FName& Name);
-	template <typename AttrType> ENGINE_API bool CreateAttributeChannel(FName AttributeName) const;
-	template <typename AttrType> ENGINE_API int32 AddAttributeValue(float Param, const AttrType& Value, FName AttributeName) const;
-	template <typename AttrType> ENGINE_API void RemoveAttributeValue(int32 Index, FName AttributeName);
-	template <typename AttrType> ENGINE_API AttrType EvaluateAttribute(float Param, FName AttributeName) const;
-
-	ENGINE_API float FindNearest(const FVector& InLocation, float& OutSquaredDist) const;
-	ENGINE_API float FindNearestOnSegment(const FVector& InLocation, int32 SegmentIndex, float& OutSquaredDist) const;
-	
-	/* Misc Interface */
-	
+	ENGINE_API int32 GetNumberOfSegments() const;
+
+	/** Reset the spline to an empty spline. Preserves the current implementation. */
+	ENGINE_API void Reset();
+
+	/**
+	 * Fetch the underlying concrete spline implementation.
+	 *
+	 * Pointers returned by Get<Impl>() are invalidated by:
+	 * - Serialize(FArchive&) when Ar.IsLoading() is true
+	 * - ImportTextItem(...)
+	 * - operator=(...)
+	 * - SetCurrentImplementation(...)
+	 * - destruction of this FSpline
+	 *
+	 * @return A valid pointer if Impl == GetCurrentImplementation(), otherwise nullptr.
+	 */
+	template<UE::Spline::ESplineType Impl>
+	typename UE::Spline::TSplineTypeMap<Impl>::Type* Get();
+
+	/** A const version of Get<Impl> */
+	template<UE::Spline::ESplineType Impl>
+	const typename UE::Spline::TSplineTypeMap<Impl>::Type* Get() const;
+		
 	ENGINE_API bool operator==(const FSpline& Other) const;
 	ENGINE_API bool operator!=(const FSpline& Other) const;
-
-	ENGINE_API bool IsNearlyEqual(const FSpline& Other, const float RelativeTolerance, const float AbsoluteTolerance) const;
 	
 	friend FArchive& operator<<(FArchive& Ar, FSpline& Spline)
-	{
+	{ 
 		Spline.Serialize(Ar);
 		return Ar;
 	}
+
 	ENGINE_API bool Serialize(FArchive& Ar);
-	ENGINE_API void SerializeLoad(FArchive& Ar);
-	ENGINE_API void SerializeSave(FArchive& Ar) const;
 	ENGINE_API bool ExportTextItem(FString& ValueStr, FSpline const& DefaultValue, class UObject* Parent, int32 PortFlags, class UObject* ExportRootScope) const;
 	ENGINE_API bool ImportTextItem(const TCHAR*& Buffer, int32 PortFlags, UObject* Parent, FOutputDevice* ErrorText);
 
-	uint32 GetVersion() const { return Version; }
-	
-	ENGINE_API const FInterpCurveVector& GetSplinePointsPosition() const;
-	ENGINE_API const FInterpCurveQuat& GetSplinePointsRotation() const;
-	ENGINE_API const FInterpCurveVector& GetSplinePointsScale() const;
+	ENGINE_API UE::Spline::ESplineType GetCurrentImplementation() const;
 
-	/** Returns the length of the specified spline segment up to the parametric value given */
-	ENGINE_API float GetSegmentLength(const int32 Index, const float Param, const FVector& Scale3D = FVector(1.0f)) const;
+#if WITH_EDITOR
+	/** Returns the implementation of this spline when when most recently loaded. */
+	ENGINE_API UE::Spline::ESplineType GetPreviousImplementation() const;
+#endif
 
-	/** Returns total length along this spline */
-	ENGINE_API float GetSplineLength() const;
+	/** Returns the default implementation for new spline. */
+	ENGINE_API static UE::Spline::ESplineType GetDefaultImplementation();
 
-	/** Returns the total number of segments on this spline. */
-	ENGINE_API int32 GetNumSegments() const;
-	
-	/** Returns the total number of control points on this spline. */
-	ENGINE_API int32 GetNumControlPoints() const;
-	
-	/** Reset the spline to an empty spline. */
-	ENGINE_API void Reset();
-	
-	/** Reset the rotation attribute channel to default values. */
-	ENGINE_API void ResetRotation();
+private:
 	
-	/** Reset the scale attribute channel to default values. */
-	ENGINE_API void ResetScale();
+	/** Sets the current implementation and performs any necessary allocation. */
+	void SetCurrentImplementation(UE::Spline::ESplineType Implementation);
 
-	ENGINE_API void SetClosedLoop(bool bClosed);
-	ENGINE_API bool IsClosedLoop() const;
-	
-	struct FUpdateSplineParams
-	{
-		bool bClosedLoop = false;
-		bool bStationaryEndpoints = false;
-		int32 ReparamStepsPerSegment = 10;
-		bool bLoopPositionOverride = false;
-		float LoopPosition = 0.0f;
-		FVector Scale3D = FVector(1.0f);
-	};
-	
-	/** Update the spline's internal data according to the passed-in params. */
-	ENGINE_API void UpdateSpline(const FUpdateSplineParams& InParams);
+	void SerializeLoad(FArchive& Ar);
+	void SerializeSave(FArchive& Ar) const;
 
-	/** Update the spline's internal data according to the most recently used update params (or default params if never updated). */
-	ENGINE_API void UpdateSpline();
+private:
 
-	ENGINE_API TSharedPtr<UE::Geometry::Spline::TSplineInterface<FVector>> GetSplineInterface() const;
+	UE::Spline::ESplineType CurrentImplementation;
 
-private:
-	
-	static ENGINE_API const FInterpCurveVector PositionCurve;
-	static ENGINE_API const FInterpCurveQuat RotationCurve;
-	static ENGINE_API const FInterpCurveVector ScaleCurve;
-	
-	// Used for upgrade logic in spline component.
-	// Not ideal, but allows us to automatically populate the proxy
-	// at serialize time when we might otherwise not be able to.
-	friend class USplineComponent;
-	friend struct FPCGSplineStruct;
-	
 #if WITH_EDITOR
-	uint8 PreviousImplementation;
+	UE::Spline::ESplineType PreviousImplementation;
 #endif
-	uint8 CurrentImplementation;
-	
-	uint32 Version;
 
-	FUpdateSplineParams CachedUpdateSplineParams;
-	
-	// probably better implemented as a TSharedPtr<ISplineInterface> or something
-	struct {													// Invalid when CurrentImplementation is 0 or 3.
-		TSharedPtr<FLegacySpline> LegacyData;					// Valid when CurrentImplementation is 1.
-		TSharedPtr<FNewSpline> NewData;							// Valid when CurrentImplementation is 2.
-	};
-	
-	bool IsEnabled() const { return CurrentImplementation != 0; }
-	bool IsLegacy() const { return CurrentImplementation == 1; }
-	bool IsNew() const { return CurrentImplementation == 2; }
+	// Valid when CurrentImplementation is ESplineType::LegacyTangent
+	TUniquePtr<FLegacyTangentSpline> LegacySplineImpl;
 
-#if WITH_EDITOR
-	bool WasEnabled() const { return PreviousImplementation != 0; }
-	bool WasLegacy() const { return PreviousImplementation == 1; }
-	bool WasNew() const { return PreviousImplementation == 2; }
+	// Valid when CurrentImplementation is not ESplineType::Unimplemented and not ESplineType::LegacyTangent;
+	TUniquePtr<UE::Geometry::Spline::TSplineInterface<FVector>> SplineImpl;
+};
 
-	void OnSplineImplementationChanged();
-#endif
+template<UE::Spline::ESplineType Impl>
+typename UE::Spline::TSplineTypeMap<Impl>::Type* FSpline::Get()
+{
+	if (Impl != CurrentImplementation)
+	{
+		return nullptr;
+	}
 
-public:
+	if constexpr (Impl == UE::Spline::ESplineType::LegacyTangent)
+	{
+		return LegacySplineImpl.Get();
+	}
+	else
+	{
+		return static_cast<typename UE::Spline::TSplineTypeMap<Impl>::Type*>(SplineImpl.Get());
+	}
+}
 
-	// True if a given FSpline will actually be implemented when instantiated.
-	static bool IsEnabledStatic();
-};
+template<UE::Spline::ESplineType Impl>
+const typename UE::Spline::TSplineTypeMap<Impl>::Type* FSpline::Get() const
+{
+	if (Impl != CurrentImplementation)
+	{
+		return nullptr;
+	}
+
+	if constexpr (Impl == UE::Spline::ESplineType::LegacyTangent)
+	{
+		return LegacySplineImpl.Get();
+	}
+	else
+	{
+		return static_cast<const typename UE::Spline::TSplineTypeMap<Impl>::Type*>(SplineImpl.Get());
+	}
+}
 
 template<>
 struct TStructOpsTypeTraits<FSpline> : public TStructOpsTypeTraitsBase2<FSpline>
