[PCG] Remove counter productive ensure that fires when you duplicate output pin settings on a custom HLSL node.

#rnx
#Synced-CL 49105085

diff --git a/Engine/Plugins/PCG/Source/PCG/Private/Graph/PCGGraphCompilerGPU.cpp b/Engine/Plugins/PCG/Source/PCG/Private/Graph/PCGGraphCompilerGPU.cpp
--- a/Engine/Plugins/PCG/Source/PCG/Private/Graph/PCGGraphCompilerGPU.cpp
+++ b/Engine/Plugins/PCG/Source/PCG/Private/Graph/PCGGraphCompilerGPU.cpp
@@ -1035,7 +1035,6 @@ void FPCGGraphCompilerGPU::CreateDataInterfaces(
 
 			if (InOutPinToDataInterface.Contains({ TaskId, OutputPinProperties.Label }))
 			{
-				ensure(false);
 				continue;
 			}
 
