[PCG] Fix failing GPU SM Spawner when using Primitive Data spawning and low instance counts.

#rnx
#Synced-CL 49105085

diff --git a/Engine/Plugins/PCG/Source/PCG/Private/Compute/BuiltInKernels/PCGCountUniqueAttributeValuesKernel.cpp b/Engine/Plugins/PCG/Source/PCG/Private/Compute/BuiltInKernels/PCGCountUniqueAttributeValuesKernel.cpp
--- a/Engine/Plugins/PCG/Source/PCG/Private/Compute/BuiltInKernels/PCGCountUniqueAttributeValuesKernel.cpp
+++ b/Engine/Plugins/PCG/Source/PCG/Private/Compute/BuiltInKernels/PCGCountUniqueAttributeValuesKernel.cpp
@@ -208,8 +208,15 @@ int UPCGCountUniqueAttributeValuesKernel::ComputeThreadCount(const UPCGDataBindi
 	{
 		const int ElementCount = InputPinDesc->ComputeTotalElementCount();
 
-		const int32 PartitionAttributeId = InBinding->GetAttributeId(AttributeName, EPCGKernelAttributeType::StringKey);
-		const int ValueCount = (PartitionAttributeId != INDEX_NONE) ? InputPinDesc->GetNumStringKeyValues(PartitionAttributeId) : 0;
+		int ValueCount = 0;
+		if (const int32 PartitionAttributeId = InBinding->GetAttributeId(AttributeName, EPCGKernelAttributeType::StringKey); PartitionAttributeId != INDEX_NONE)
+		{
+			ValueCount = InputPinDesc->GetNumStringKeyValues(PartitionAttributeId);
+		}
+		else
+		{
+			ValueCount = FMath::Max(ValueCount, MaxIntValue + 1);
+		}
 
 		// One thread per element, but also make sure we have enough threads to write each unique attribute value (rather than leaving uninitialized).
 		return FMath::Max(ElementCount, ValueCount);
