[PCG] Allow to use the Template Actor/Component when spawning by attribute

#rnx
#Synced-CL 49088187

diff --git a/Engine/Plugins/PCG/Source/PCG/Private/Elements/PCGAddComponent.cpp b/Engine/Plugins/PCG/Source/PCG/Private/Elements/PCGAddComponent.cpp
--- a/Engine/Plugins/PCG/Source/PCG/Private/Elements/PCGAddComponent.cpp
+++ b/Engine/Plugins/PCG/Source/PCG/Private/Elements/PCGAddComponent.cpp
@@ -13,6 +13,7 @@
 #include "Engine/Blueprint.h"
 #include "Engine/BlueprintGeneratedClass.h"
 #include "Engine/Engine.h"
+#include "UObject/FortniteMainBranchObjectVersion.h"
 #include "UObject/Package.h"
 
 #include UE_INLINE_GENERATED_CPP_BY_NAME(PCGAddComponent)
@@ -30,6 +31,12 @@ void UPCGAddComponentSettings::PostLoad()
 	Super::PostLoad();
 
 #if WITH_EDITOR
+	// Note: Before we were not allowing template editing if we were spawning as attribute, so force it at false in that case.
+	if (GetLinkerCustomVersion(FFortniteMainBranchObjectVersion::GUID) < FFortniteMainBranchObjectVersion::PCGAllowTemplateWhenSpawningFromAttribute && bUseClassAttribute)
+	{
+		bAllowTemplateComponentEditing = false;
+	}
+
 	SetupOnBlueprintChangedEvent();
 
 	if (TemplateComponentClass && TemplateComponent)
@@ -374,6 +381,7 @@ bool FPCGAddComponentElement::ExecuteInternal(FPCGContext* InContext) const
 
 		// Now create and add the components
 		AddComponents(Context,
+			Settings,
 			TransformAccessor.Get(),
 			TransformKeys.Get(),
 			ClassOrTemplateAccessor.Get(),
@@ -394,6 +402,7 @@ bool FPCGAddComponentElement::ExecuteInternal(FPCGContext* InContext) const
 }
 
 void FPCGAddComponentElement::AddComponents(FPCGAddComponentContext* Context,
+	const UPCGAddComponentSettings* Settings,
 	const IPCGAttributeAccessor* TransformAccessor,
 	const IPCGAttributeAccessorKeys* TransformKeys,
 	const IPCGAttributeAccessor* ClassOrTemplateAccessor,
@@ -472,12 +481,29 @@ void FPCGAddComponentElement::AddComponents(FPCGAddComponentContext* Context,
 			TSoftClassPtr<UActorComponent> ComponentClassPtr(ClassOrTemplatePath);
 			ComponentClass = ComponentClassPtr.Get();
 
+			bool bCopyTemplateComponent = false;
+
 			if (!ComponentClass)
 			{
 				TSoftObjectPtr<UActorComponent> ComponentObjectPtr(ClassOrTemplatePath);
 				ComponentTemplate = ComponentObjectPtr.Get();
 				ComponentClass = ComponentTemplate ? ComponentTemplate->GetClass() : nullptr;
 			}
+			else if (Settings->bAllowTemplateComponentEditing && Settings->TemplateComponent)
+			{
+				// Use the Template from the settings if it exists. Only use it if it is a parent class of the class we want to spawn
+				if (!ComponentClass->IsChildOf(Settings->TemplateComponent->GetClass()))
+				{
+					PCGLog::LogErrorOnGraph(FText::Format(LOCTEXT("InvalidTemplateClass", "For input {0}, point index {1}, tried to spawn a component of class {2}, while providing a template component of type {3}. This is invalid. The template component must be of a parent class of the spawned component class."),
+						Context->CurrentInputIndex,
+						j,
+						FText::FromName(ComponentClass->GetFName()),
+						FText::FromName(Settings->TemplateComponent->GetClass()->GetFName())));
+					continue;
+				}
+
+				bCopyTemplateComponent = true;
+			}
 			
 			if (TargetActor && ComponentClass)
 			{
@@ -488,6 +514,11 @@ void FPCGAddComponentElement::AddComponents(FPCGAddComponentContext* Context,
 					continue;
 				}
 
+				if (bCopyTemplateComponent)
+				{
+					UEngine::CopyPropertiesForUnrelatedObjects(Settings->TemplateComponent, Component);
+				}
+
 				Component->RegisterComponent();
 				TargetActor->AddInstanceComponent(Component);
 
diff --git a/Engine/Plugins/PCG/Source/PCG/Private/Elements/PCGSpawnActor.cpp b/Engine/Plugins/PCG/Source/PCG/Private/Elements/PCGSpawnActor.cpp
--- a/Engine/Plugins/PCG/Source/PCG/Private/Elements/PCGSpawnActor.cpp
+++ b/Engine/Plugins/PCG/Source/PCG/Private/Elements/PCGSpawnActor.cpp
@@ -193,6 +193,11 @@ void UPCGSpawnActorSettings::SetAllowTemplateActorEditing(bool bInAllowTemplateA
 #endif // WITH_EDITOR
 }
 
+bool UPCGSpawnActorSettings::GetAllowTemplateActorEditing() const
+{
+	return Option != EPCGSpawnActorOption::CollapseActors && bAllowTemplateActorEditing;
+}
+
 FPCGElementPtr UPCGSpawnActorSettings::CreateElement() const
 {
 	return MakeShared<FPCGSpawnActorElement>();
@@ -259,25 +264,13 @@ void UPCGSpawnActorSettings::PostLoad()
 	Super::PostLoad();
 
 #if WITH_EDITOR
-	// Apply data deprecation - deprecated in UE 5.3
-	if (bGenerationTrigger_DEPRECATED != EPCGSpawnActorGenerationTrigger::Default)
+	// Since the template actor editing is set to false by default, this needs to be corrected on post-load for proper deprecation
+	// Note: Before we were not allowing template editing if we were spawning as attribute, so force it at false in that case.
+	if (GetLinkerCustomVersion(FFortniteMainBranchObjectVersion::GUID) < FFortniteMainBranchObjectVersion::PCGAllowTemplateWhenSpawningFromAttribute && bSpawnByAttribute)
 	{
-		GenerationTrigger = bGenerationTrigger_DEPRECATED;
-		bGenerationTrigger_DEPRECATED = EPCGSpawnActorGenerationTrigger::Default;
+		bAllowTemplateActorEditing = false;
 	}
-
-	if (!ActorOverrides_DEPRECATED.IsEmpty())
-	{
-		for (const FPCGActorPropertyOverride& Override : ActorOverrides_DEPRECATED)
-		{
-			SpawnedActorPropertyOverrideDescriptions.Emplace(Override.InputSource, Override.PropertyTarget);
-		}
-
-		ActorOverrides_DEPRECATED.Empty();
-	}
-
-	// Since the template actor editing is set to false by default, this needs to be corrected on post-load for proper deprecation
-	if (TemplateActor)
+	else if (TemplateActor && Option != EPCGSpawnActorOption::CollapseActors)
 	{
 		bAllowTemplateActorEditing = true;
 	}
@@ -391,7 +384,7 @@ void UPCGSpawnActorSettings::OnObjectsReplaced(const TMap<UObject*, UObject*>& I
 void UPCGSpawnActorSettings::RefreshTemplateActor()
 {
 	// Implementation note: this is similar to the child actor component implementation
-	if (TemplateActorClass && bAllowTemplateActorEditing)
+	if (TemplateActorClass && GetAllowTemplateActorEditing())
 	{
 		const bool bCreateNewTemplateActor = (!TemplateActor || TemplateActor->GetClass() != TemplateActorClass);
 
@@ -619,7 +612,7 @@ bool FPCGSpawnActorElement::SpawnAndPrepareSubgraphs(FPCGSubgraphContext* Contex
 					FPCGTaggedData PartialInput = Input;
 					PartialInput.Data = Element.Value.PartitionData;
 
-					SpawnOrCollapse(Element.Key, nullptr, PartialInput, Cast<UPCGBasePointData>(Element.Value.PartitionData), OutPointData);
+					SpawnOrCollapse(Element.Key, Settings->TemplateActor, PartialInput, Cast<UPCGBasePointData>(Element.Value.PartitionData), OutPointData);
 
 					// Exception case here: if we've spawned actors but are merging the PCG inputs,
 					// normally this node is taken as a subgraph node (e.g. no need to do anything more than forwarding the inputs)
@@ -794,8 +787,17 @@ void FPCGSpawnActorElement::SpawnActors(FPCGSubgraphContext* Context, AActor* Ta
 	check(Context && TargetActor && PointData);
 	check(IsInGameThread());
 
-	if (PointData->GetNumPoints() == 0)
+	if (PointData->GetNumPoints() == 0 || !InTemplateActorClass)
+	{
+		return;
+	}
+
+	// Verify first if the template actor and template class matches. If they are not, the spawn will fail.
+	if (InTemplateActor && !InTemplateActorClass->IsChildOf(InTemplateActor->GetClass()))
 	{
+		PCGLog::LogErrorOnGraph(FText::Format(LOCTEXT("InvalidTemplateClass", "Tried to spawn an actor of class {0}, while providing a template actor of type {1}. This is invalid. The template actor must be of a parent class of the spawned actor class."),
+			FText::FromName(InTemplateActorClass->GetFName()),
+			FText::FromName(InTemplateActor->GetClass()->GetFName())));
 		return;
 	}
 
@@ -816,7 +818,10 @@ void FPCGSpawnActorElement::SpawnActors(FPCGSubgraphContext* Context, AActor* Ta
 	const bool bForceDisableActorParsing = (Settings->bForceDisableActorParsing);
 
 	AActor* TemplateActor = nullptr;
-	if (InTemplateActor)
+	// We can use the template actor as-is if the classes matches perfectly.
+	const bool bClassMismatch = InTemplateActor && InTemplateActorClass != InTemplateActor->GetClass();
+
+	if (InTemplateActor && !bClassMismatch)
 	{
 		if (Settings->SpawnedActorPropertyOverrideDescriptions.IsEmpty())
 		{
@@ -829,13 +834,19 @@ void FPCGSpawnActorElement::SpawnActors(FPCGSubgraphContext* Context, AActor* Ta
 	}
 	else
 	{
-		if (Settings->SpawnedActorPropertyOverrideDescriptions.IsEmpty())
+		if (Settings->SpawnedActorPropertyOverrideDescriptions.IsEmpty() && !bClassMismatch)
 		{
 			TemplateActor = Cast<AActor>(InTemplateActorClass->GetDefaultObject());
 		}
 		else
 		{
 			TemplateActor = NewObject<AActor>(GetTransientPackage(), InTemplateActorClass, NAME_None, RF_ArchetypeObject);
+
+			// If we have a class mismatch, copy the properties from the template actor.
+			if (bClassMismatch)
+			{
+				UEngine::CopyPropertiesForUnrelatedObjects(InTemplateActor, TemplateActor);
+			}
 		}
 	}
 
diff --git a/Engine/Plugins/PCG/Source/PCG/Public/Elements/PCGAddComponent.h b/Engine/Plugins/PCG/Source/PCG/Public/Elements/PCGAddComponent.h
--- a/Engine/Plugins/PCG/Source/PCG/Public/Elements/PCGAddComponent.h
+++ b/Engine/Plugins/PCG/Source/PCG/Public/Elements/PCGAddComponent.h
@@ -122,6 +122,7 @@ protected:
 
 private:
 	void AddComponents(FPCGAddComponentContext* Context,
+		const UPCGAddComponentSettings* Settings,
 		const IPCGAttributeAccessor* TransformAccessor,
 		const IPCGAttributeAccessorKeys* TransformKeys,
 		const IPCGAttributeAccessor* ClassAccessor,
diff --git a/Engine/Plugins/PCG/Source/PCG/Public/Elements/PCGSpawnActor.h b/Engine/Plugins/PCG/Source/PCG/Public/Elements/PCGSpawnActor.h
--- a/Engine/Plugins/PCG/Source/PCG/Public/Elements/PCGSpawnActor.h
+++ b/Engine/Plugins/PCG/Source/PCG/Public/Elements/PCGSpawnActor.h
@@ -72,7 +72,7 @@ public:
 	UPROPERTY(BlueprintReadWrite, EditAnywhere, Category = Settings, meta = (EditCondition = "Option!=EPCGSpawnActorOption::CollapseActors", EditConditionHides))
 	TArray<FName> TagsToAddOnActors;
 
-	UPROPERTY(EditAnywhere, BlueprintReadWrite, Instanced, Category = Settings, meta = (ShowInnerProperties, EditCondition = "bAllowTemplateActorEditing && Option != EPCGSpawnActorOption::CollapseActors", EditConditionHides))
+	UPROPERTY(EditAnywhere, BlueprintReadWrite, Instanced, Category = Settings, meta = (ShowInnerProperties, EditCondition = "GetAllowTemplateActorEditing()", EditConditionHides))
 	TObjectPtr<AActor> TemplateActor;
 
 	UPROPERTY(EditAnywhere, BlueprintReadWrite, Category = Settings, meta = (EditCondition = "Option != EPCGSpawnActorOption::CollapseActors", EditConditionHides))
@@ -91,7 +91,7 @@ public:
 	UPROPERTY(BlueprintReadWrite, EditAnywhere, Category = Settings, meta = (InlineEditConditionToggle))
 	bool bSpawnByAttribute = false;
 
-	UPROPERTY(BlueprintReadWrite, EditAnywhere, Category = Settings, meta = (EditCondition = "bSpawnByAttribute"))
+	UPROPERTY(BlueprintReadWrite, EditAnywhere, Category = Settings, meta = (EditCondition = "bSpawnByAttribute", PCG_Overridable))
 	FName SpawnAttribute = NAME_None;
 
 	/** Adds a warning to the node on repeated spawning with identical conditions (ie. same actor at same spawn location, etc). */
@@ -111,17 +111,10 @@ protected:
 	UPROPERTY(BlueprintReadOnly, EditAnywhere, Category = Settings, meta = (OnlyPlaceable, DisallowCreateNew))
 	TSubclassOf<AActor> TemplateActorClass;
 
-	UPROPERTY(BlueprintReadWrite, EditAnywhere, Category = Settings, meta = (EditCondition = "Option != EPCGSpawnActorOption::CollapseActors"))
+	/** Can modify the template actor for the expected class. When spawning by attribute, it is mandatory for that template actor to be a parent class of ALL actors that will be spawned. */
+	UPROPERTY(BlueprintReadWrite, EditAnywhere, Category = Settings, meta = (EditConditionHides, EditCondition = "Option != EPCGSpawnActorOption::CollapseActors"))
 	bool bAllowTemplateActorEditing = false;
 
-#if WITH_EDITORONLY_DATA
-	UPROPERTY()
-	EPCGSpawnActorGenerationTrigger bGenerationTrigger_DEPRECATED = EPCGSpawnActorGenerationTrigger::Default;
-
-	UPROPERTY()
-	TArray<FPCGActorPropertyOverride> ActorOverrides_DEPRECATED;
-#endif
-
 public:
 	virtual void PostLoad() override;
 	virtual void BeginDestroy() override;
@@ -139,7 +132,9 @@ public:
 	PCG_API void SetTemplateActorClass(const TSubclassOf<AActor>& InTemplateActorClass);
 	PCG_API void SetAllowTemplateActorEditing(bool bInAllowTemplateActorEditing);
 	const TSubclassOf<AActor>& GetTemplateActorClass() const { return TemplateActorClass; }
-	bool GetAllowTemplateActorEditing() const { return bAllowTemplateActorEditing; }
+
+	UFUNCTION()
+	PCG_API bool GetAllowTemplateActorEditing() const;
 
 #if WITH_EDITOR
 	virtual FName GetDefaultNodeName() const override { return FName(TEXT("SpawnActor")); }
diff --git a/Engine/Source/Runtime/Core/Public/UObject/FortniteMainBranchObjectVersions.inl b/Engine/Source/Runtime/Core/Public/UObject/FortniteMainBranchObjectVersions.inl
--- a/Engine/Source/Runtime/Core/Public/UObject/FortniteMainBranchObjectVersions.inl
+++ b/Engine/Source/Runtime/Core/Public/UObject/FortniteMainBranchObjectVersions.inl
@@ -727,5 +727,8 @@ FORTNITE_MAIN_VERSION(AddMeshDeformerUsageFlag, 237)
 // Add landscape texture patch alpha override for height/weight patches
 FORTNITE_MAIN_VERSION(LandscapeTexturePatchAlphaOverride, 238)
 
+// For SpawnActor/AddComponent, allow template to be used for the spawned actor.
+FORTNITE_MAIN_VERSION(PCGAllowTemplateWhenSpawningFromAttribute, 239)
+
 // -----<new versions can be added above this line>-------------------------------------------------
 #undef FORTNITE_MAIN_VERSION
