[PCG] Fixed default isolate call helper function.

#rnx
#Synced-CL 49172122

diff --git a/Engine/Plugins/PCG/Source/PCG/Private/Helpers/PCGActorHelpers.cpp b/Engine/Plugins/PCG/Source/PCG/Private/Helpers/PCGActorHelpers.cpp
--- a/Engine/Plugins/PCG/Source/PCG/Private/Helpers/PCGActorHelpers.cpp
+++ b/Engine/Plugins/PCG/Source/PCG/Private/Helpers/PCGActorHelpers.cpp
@@ -900,6 +900,7 @@ AActor* UPCGActorHelpers::IsolateFromActor(AActor* Actor, const TArray<FName>& R
 {
 	FPCGMoveResourceParams Params
 	{
+		.TemplateTargetClass = APCGIsolatedActor::StaticClass(),
 		.RequiredTags = RequiredTags
 	};
 
@@ -922,6 +923,7 @@ AActor* UPCGActorHelpers::IsolateFromActor(AActor* Actor, const FPCGMoveResource
 
 	if (bMoveCollectionToIsolatedActor)
 	{
+		Params.TemplateTargetClass = APCGIsolatedActor::StaticClass();
 		Params.TargetDataCollection = &TempCollection;
 	}
 
