[PCG] Remove legacy functionality that managed when transient data was released.

The ref counting functionality in UPCGData supercedes it and has been working well.

#rnx
#Synced-CL 49105085

diff --git a/Engine/Plugins/PCG/Source/PCG/Private/Graph/PCGGraphExecutor.cpp b/Engine/Plugins/PCG/Source/PCG/Private/Graph/PCGGraphExecutor.cpp
--- a/Engine/Plugins/PCG/Source/PCG/Private/Graph/PCGGraphExecutor.cpp
+++ b/Engine/Plugins/PCG/Source/PCG/Private/Graph/PCGGraphExecutor.cpp
@@ -1209,22 +1209,7 @@ void FPCGGraphExecutor::PostTaskExecute(TSharedPtr<FPCGGraphActiveTask> ActiveTa
 #endif // WITH_EDITOR
 
 		// If the task is a post execute, then we can safely clear the data after getting it from the results.
-		bool bNeedsManualClear = (ActiveTask.Element == GraphCompiler.GetSharedTrivialPostGraphElement());
-
-		// Data holding transient resources needed originally to be released late (at the end of the generation)
-		// because we had no way to assess input usage. With the cvar here (if it's true), it's not the case anymore.
-		if (!bNeedsManualClear && !PCGSystemSwitches::CVarReleaseTransientResourcesEarly.GetValueOnAnyThread())
-		{
-			for (const FPCGTaggedData& TaggedData : ActiveTask.Context->OutputData.TaggedData)
-			{
-				// Don't release any data that holds onto transient resources mid-execution. Clear these at the end.
-				if (TaggedData.Data && TaggedData.Data->HoldsTransientResources())
-				{
-					bNeedsManualClear = true;
-					break;
-				}
-			}
-		}
+		const bool bNeedsManualClear = (ActiveTask.Element == GraphCompiler.GetSharedTrivialPostGraphElement());
 
 		// Store output in data map.
 		FTaskResult NewResult{ ActiveTask.NodeId, &ActiveTask.Context->OutputData, bNeedsManualClear };
diff --git a/Engine/Plugins/PCG/Source/PCG/Private/PCGCommon.cpp b/Engine/Plugins/PCG/Source/PCG/Private/PCGCommon.cpp
--- a/Engine/Plugins/PCG/Source/PCG/Private/PCGCommon.cpp
+++ b/Engine/Plugins/PCG/Source/PCG/Private/PCGCommon.cpp
@@ -107,11 +107,6 @@ namespace PCGSystemSwitches
 		TEXT("Forces all subgraph executions to be dynamic. Performance warning, also requires to flush all compiled graphs."));
 #endif
 
-	TAutoConsoleVariable<bool> CVarReleaseTransientResourcesEarly(
-		TEXT("pcg.ReleaseTransientResourcesEarly"),
-		true,
-		TEXT("Aggressively release transient (esp. GPU) data when possible."));
-
 #if UE_ENABLE_DEBUG_DRAWING
 	TAutoConsoleVariable<bool> CVarPCGDebugDrawGeneratedCells(
 		TEXT("pcg.GraphExecution.DebugDrawGeneratedCells"),
@@ -130,6 +125,15 @@ namespace PCGSystemSwitches
 		TEXT("pcg.Graph.GPU.PassGPUDataThroughGridLinks"),
 		true,
 		TEXT("Whether proxies for GPU data are cached in per pin output data and passed through grid links. If false data is read back to CPU."));
+
+	// Deprecated section.
+	PRAGMA_DISABLE_DEPRECATION_WARNINGS
+	TAutoConsoleVariable<bool> CVarReleaseTransientResourcesEarly(
+		TEXT("pcg.ReleaseTransientResourcesEarly"),
+		true,
+		TEXT("DEPRECATED: This feature is always enabled now."),
+		ECVF_ReadOnly);
+	PRAGMA_ENABLE_DEPRECATION_WARNINGS
 }
 
 namespace PCGHiGenGrid
diff --git a/Engine/Plugins/PCG/Source/PCG/Private/PCGData.cpp b/Engine/Plugins/PCG/Source/PCG/Private/PCGData.cpp
--- a/Engine/Plugins/PCG/Source/PCG/Private/PCGData.cpp
+++ b/Engine/Plugins/PCG/Source/PCG/Private/PCGData.cpp
@@ -88,8 +88,7 @@ void UPCGData::ClearUsage(EPCGDataUsage InUsage) const
 
 	Usage &= ~static_cast<__underlying_type(EPCGDataUsage)>(InUsage);
 
-	if (Usage == static_cast<__underlying_type(EPCGDataUsage)>(EPCGDataUsage::None) && 
-		(!PCGSystemSwitches::CVarReleaseTransientResourcesEarly.GetValueOnAnyThread() || CollectionRefCount == 0))
+	if (Usage == static_cast<__underlying_type(EPCGDataUsage)>(EPCGDataUsage::None) && CollectionRefCount == 0)
 	{
 #ifdef PCG_DATA_USAGE_LOGGING
 		UE_LOG(LogPCG, Warning, TEXT("%s: !! ReleaseTransientResources after usage unmarked !!"), *GetName());
@@ -119,11 +118,6 @@ void UPCGData::DecCollectionRefCount() const
 {
 	--CollectionRefCount;
 
-	if (!PCGSystemSwitches::CVarReleaseTransientResourcesEarly.GetValueOnAnyThread())
-	{
-		return;
-	}
-
 	ensure(CollectionRefCount >= 0);
 
 	if (Usage == static_cast<__underlying_type(EPCGDataUsage)>(EPCGDataUsage::None) && CollectionRefCount == 0)
diff --git a/Engine/Plugins/PCG/Source/PCG/Public/Compute/Data/PCGProxyForGPUData.h b/Engine/Plugins/PCG/Source/PCG/Public/Compute/Data/PCGProxyForGPUData.h
--- a/Engine/Plugins/PCG/Source/PCG/Public/Compute/Data/PCGProxyForGPUData.h
+++ b/Engine/Plugins/PCG/Source/PCG/Public/Compute/Data/PCGProxyForGPUData.h
@@ -42,7 +42,6 @@ public:
 
 	PCG_API virtual void AddToCrc(FArchiveCrc32& Ar, bool bFullDataCrc) const override;
 	virtual bool CanBeSerialized() const override { return false; }
-	virtual bool HoldsTransientResources() const override { return true; }
 	virtual bool IsCacheable() const override { return false; }
 	PCG_API virtual void ReleaseTransientResources(const TCHAR* InReason = nullptr) override;
 	PCG_API virtual EPCGDataType GetUnderlyingDataType() const override;
diff --git a/Engine/Plugins/PCG/Source/PCG/Public/Data/PCGRenderTargetData.h b/Engine/Plugins/PCG/Source/PCG/Public/Data/PCGRenderTargetData.h
--- a/Engine/Plugins/PCG/Source/PCG/Public/Data/PCGRenderTargetData.h
+++ b/Engine/Plugins/PCG/Source/PCG/Public/Data/PCGRenderTargetData.h
@@ -30,7 +30,6 @@ public:
 	// ~Begin UPCGData interface
 	PCG_ASSIGN_DEFAULT_TYPE_INFO(FPCGDataTypeInfoRenderTarget2D);
 	virtual void AddToCrc(FArchiveCrc32& Ar, bool bFullDataCrc) const override;
-	virtual bool HoldsTransientResources() const override { return bOwnsRenderTarget; }
 	virtual bool IsCacheable() const { return Super::IsCacheable() && !bOwnsRenderTarget; }
 	virtual void ReleaseTransientResources(const TCHAR* InReason = nullptr) override;
 	// ~End UPCGData interface
diff --git a/Engine/Plugins/PCG/Source/PCG/Public/Data/PCGTextureData.h b/Engine/Plugins/PCG/Source/PCG/Public/Data/PCGTextureData.h
--- a/Engine/Plugins/PCG/Source/PCG/Public/Data/PCGTextureData.h
+++ b/Engine/Plugins/PCG/Source/PCG/Public/Data/PCGTextureData.h
@@ -225,8 +225,6 @@ public:
 	//~Begin UPCGData interface
 	PCG_ASSIGN_DEFAULT_TYPE_INFO(FPCGDataTypeInfoTexture2D)
 	virtual void AddToCrc(FArchiveCrc32& Ar, bool bFullDataCrc) const override;
-	// Don't hold onto exported buffers currently as graphics memory usage (and lifetimes) may cause issues.
-	virtual bool HoldsTransientResources() const override { return ResourceType == EPCGTextureResourceType::ExportedTexture; }
 	virtual bool IsCacheable() const { return Super::IsCacheable() && ResourceType != EPCGTextureResourceType::ExportedTexture; }
 	virtual void ReleaseTransientResources(const TCHAR* InReason = nullptr) override { TextureHandle = nullptr; }
 	//~End UPCGData interface
diff --git a/Engine/Plugins/PCG/Source/PCG/Public/PCGCommon.h b/Engine/Plugins/PCG/Source/PCG/Public/PCGCommon.h
--- a/Engine/Plugins/PCG/Source/PCG/Public/PCGCommon.h
+++ b/Engine/Plugins/PCG/Source/PCG/Public/PCGCommon.h
@@ -455,13 +455,15 @@ namespace PCGSystemSwitches
 	extern TAutoConsoleVariable<bool> CVarForceDynamicGraphDispatch;
 #endif
 
-	extern PCG_API TAutoConsoleVariable<bool> CVarReleaseTransientResourcesEarly;
 	extern PCG_API TAutoConsoleVariable<bool> CVarPCGDebugDrawGeneratedCells;
 	extern TAutoConsoleVariable<bool> CVarPassGPUDataThroughGridLinks;
 
 #if !UE_BUILD_SHIPPING
 	extern PCG_API TAutoConsoleVariable<bool> CVarFuzzGPUMemory;
 #endif
+
+	UE_DEPRECATED(5.8, "This feature is always enabled now")
+	extern PCG_API TAutoConsoleVariable<bool> CVarReleaseTransientResourcesEarly;
 }
 
 /** Describes space referential for operations that create data */
diff --git a/Engine/Plugins/PCG/Source/PCG/Public/PCGData.h b/Engine/Plugins/PCG/Source/PCG/Public/PCGData.h
--- a/Engine/Plugins/PCG/Source/PCG/Public/PCGData.h
+++ b/Engine/Plugins/PCG/Source/PCG/Public/PCGData.h
@@ -74,9 +74,6 @@ public:
 	/** Whether the data can be placed in the graph cache. */
 	virtual bool IsCacheable() const { return true; }
 
-	/** Whether this data is holding onto one or more transient resources. This information can be used to pro-actively release the resource when the data is no longer in use. */
-	virtual bool HoldsTransientResources() const { return false; }
-
 	virtual void ReleaseTransientResources(const TCHAR* InReason = nullptr) {}
 
 	/** If this data is a proxy, returns the underlying data type, otherwise returns this data type. */
@@ -183,6 +180,11 @@ private:
 	/** Usage flags to record any current usages of this data. Used to help manage lifetimes of any internal transient resources. */
 	mutable std::atomic<__underlying_type(EPCGDataUsage)> Usage{ 0 };
 	mutable std::atomic<int32> CollectionRefCount{ 0 };
+
+	// Deprecated section
+public:
+	UE_DEPRECATED(5.8, "No longer required")
+	virtual bool HoldsTransientResources() const { return false; }
 };
 
 USTRUCT(BlueprintType, meta = (HasNativeBreak = "/Script/PCG.PCGDataFunctionLibrary.BreakTaggedData", HasNativeMake = "/Script/PCG.PCGDataFunctionLibrary.MakeTaggedData"))
