[PCG] Proactively release memory from raw buffer data instead of waiting for GC, which reduced total memory usage on a project signficantly.

#rnx
#Synced-CL 49105085

diff --git a/Engine/Plugins/PCG/Source/PCG/Private/Compute/Data/PCGRawBufferData.cpp b/Engine/Plugins/PCG/Source/PCG/Private/Compute/Data/PCGRawBufferData.cpp
--- a/Engine/Plugins/PCG/Source/PCG/Private/Compute/Data/PCGRawBufferData.cpp
+++ b/Engine/Plugins/PCG/Source/PCG/Private/Compute/Data/PCGRawBufferData.cpp
@@ -15,3 +15,10 @@ void UPCGRawBufferData::AddToCrc(FArchiveCrc32& Ar, bool bFullDataCrc) const
 	// This data does not have a bespoke CRC implementation so just use the unique object instance UID.
 	AddUIDToCrc(Ar);
 }
+
+void UPCGRawBufferData::ReleaseTransientResources(const TCHAR* InReason)
+{
+	Super::ReleaseTransientResources(InReason);
+
+	Data.Empty();
+}
diff --git a/Engine/Plugins/PCG/Source/PCG/Private/Compute/Data/PCGRawBufferData.h b/Engine/Plugins/PCG/Source/PCG/Private/Compute/Data/PCGRawBufferData.h
--- a/Engine/Plugins/PCG/Source/PCG/Private/Compute/Data/PCGRawBufferData.h
+++ b/Engine/Plugins/PCG/Source/PCG/Private/Compute/Data/PCGRawBufferData.h
@@ -29,6 +29,7 @@ public:
 	PCG_ASSIGN_TYPE_INFO(FPCGDataTypeInfoRawBuffer);
 
 	PCG_API virtual void AddToCrc(FArchiveCrc32& Ar, bool bFullDataCrc) const override;
+	PCG_API virtual void ReleaseTransientResources(const TCHAR* InRease = nullptr) override;
 	//~End UPCGData interface
 
 	void SetData(TArray<uint32>&& InData) { Data = MoveTemp(InData); }
diff --git a/Engine/Plugins/PCG/Source/PCG/Private/Compute/DataInterfaces/PCGDebugDataInterface.cpp b/Engine/Plugins/PCG/Source/PCG/Private/Compute/DataInterfaces/PCGDebugDataInterface.cpp
--- a/Engine/Plugins/PCG/Source/PCG/Private/Compute/DataInterfaces/PCGDebugDataInterface.cpp
+++ b/Engine/Plugins/PCG/Source/PCG/Private/Compute/DataInterfaces/PCGDebugDataInterface.cpp
@@ -92,6 +92,8 @@ FComputeDataProviderRenderProxy* UPCGDebugDataProvider::GetRenderProxy()
 			return;
 		}
 
+		LLM_SCOPE_BYTAG(PCG);
+
 		const float* DataAsFloat = static_cast<const float*>(InData);
 
 		UE_LOG(LogPCG, Warning, TEXT("%s - %s"), *SettingsName, *KernelName);
diff --git a/Engine/Plugins/PCG/Source/PCG/Private/Compute/DataInterfaces/PCGRawBufferDataInterface.cpp b/Engine/Plugins/PCG/Source/PCG/Private/Compute/DataInterfaces/PCGRawBufferDataInterface.cpp
--- a/Engine/Plugins/PCG/Source/PCG/Private/Compute/DataInterfaces/PCGRawBufferDataInterface.cpp
+++ b/Engine/Plugins/PCG/Source/PCG/Private/Compute/DataInterfaces/PCGRawBufferDataInterface.cpp
@@ -199,10 +199,14 @@ FComputeDataProviderRenderProxy* UPCGRawBufferDataProvider::GetRenderProxy()
 			return;
 		}
 
+		LLM_SCOPE_BYTAG(PCG);
+
 		ReadbackState->Data.Append(static_cast<const uint32*>(InData), InNumBytes / sizeof(uint32));
 
 		ExecuteOnGameThread(UE_SOURCE_LOCATION, [ReadbackState, DataProviderWeakPtr, GenerationCount, ExportMode, OutputPinLabel, OutputPinLabelAlias, BindingWeakPtr, PinDesc]()
 		{
+			LLM_SCOPE_BYTAG(PCG);
+
 			UPCGDataBinding* Binding = BindingWeakPtr.Get();
 			if (!Binding)
 			{
