[PCG] Fixed bad data validation in metadata ops

#rnx
#Synced-CL 49172122

diff --git a/Engine/Plugins/PCG/Source/PCG/Private/Elements/Metadata/PCGMetadataOpElementBase.cpp b/Engine/Plugins/PCG/Source/PCG/Private/Elements/Metadata/PCGMetadataOpElementBase.cpp
--- a/Engine/Plugins/PCG/Source/PCG/Private/Elements/Metadata/PCGMetadataOpElementBase.cpp
+++ b/Engine/Plugins/PCG/Source/PCG/Private/Elements/Metadata/PCGMetadataOpElementBase.cpp
@@ -677,15 +677,15 @@ bool FPCGMetadataElementBase::PrepareDataInternal(FPCGContext* Context) const
 
 			const int32 CurrentInputNum = CurrentPinInputData.Num();
 
-			OperandInputNumMax = FMath::Max(OperandInputNumMax, FMath::Max(1, CurrentInputNum));
-
 			// For the current input, no input (0) could be default value and we support N:1 and 1:N
-			if (CurrentInputNum > 1 && CurrentInputNum != OperandInputNumMax) 
+			if (CurrentInputNum > 1 && OperandInputNumMax > 1 && CurrentInputNum != OperandInputNumMax)
 			{
 				PCGE_LOG(Error, GraphAndLog, LOCTEXT("MismatchedOperandDataCount", "Number of data elements provided on inputs must be 1:N, N:1, or N:N."));
 				return EPCGTimeSliceInitResult::AbortExecution;
 			}
 
+			OperandInputNumMax = FMath::Max(OperandInputNumMax, FMath::Max(1, CurrentInputNum));
+
 			if (Settings->DefaultValuesAreEnabled() && !bIsInputConnected && CurrentPinInputData.IsEmpty() && Settings->IsPinDefaultValueActivated(CurrentPinLabel))
 			{
 				FPCGTaggedData& DefaultData = CurrentPinInputData.Emplace_GetRef();
