[PCG] Fixed bad access to array in split splines

#rnx
#Synced-CL 49172122

diff --git a/Engine/Plugins/PCG/Source/PCG/Private/Elements/PCGSplitSplines.cpp b/Engine/Plugins/PCG/Source/PCG/Private/Elements/PCGSplitSplines.cpp
--- a/Engine/Plugins/PCG/Source/PCG/Private/Elements/PCGSplitSplines.cpp
+++ b/Engine/Plugins/PCG/Source/PCG/Private/Elements/PCGSplitSplines.cpp
@@ -174,7 +174,7 @@ namespace PCGSplitSpline
 		}
 		else
 		{
-			TEnumAsByte<ESplinePointType::Type> SplitPointType = GetPointTypeForTouchedPoint(SplinePoints[LastControlPointIndexInSegment].Type);
+			TEnumAsByte<ESplinePointType::Type> SplitPointType = GetPointTypeForTouchedPoint(SplinePoints[LastControlPointIndexInSegment % SplinePoints.Num()].Type);
 			AddSplitPoint(EndKey, SplitPointType, LastControlKey, NextControlKey);
 			bNeedsEndTangentAdjustment = true;
 		}
