[PCG] Slightly better management of parameters in the isolate from actor calls.

#rnx
#Synced-CL 49172122

diff --git a/Engine/Plugins/PCG/Source/PCG/Private/Helpers/PCGActorHelpers.cpp b/Engine/Plugins/PCG/Source/PCG/Private/Helpers/PCGActorHelpers.cpp
--- a/Engine/Plugins/PCG/Source/PCG/Private/Helpers/PCGActorHelpers.cpp
+++ b/Engine/Plugins/PCG/Source/PCG/Private/Helpers/PCGActorHelpers.cpp
@@ -918,12 +918,19 @@ AActor* UPCGActorHelpers::IsolateFromActor(AActor* Actor, const FPCGMoveResource
 
 	// Will store data temporarily (output from the ClearPCGLink) until we move it to the created actor
 	FPCGMoveResourceParams Params = InParams;
+
+	if (!Params.TemplateTargetClass)
+	{
+		Params.TemplateTargetClass = APCGIsolatedActor::StaticClass();
+	}
+
+	check(Params.TemplateTargetClass);
+
 	FPCGDataCollection TempCollection;
-	bool bMoveCollectionToIsolatedActor = !Params.TargetDataCollection;
+	bool bMoveCollectionToIsolatedActor = !Params.TargetDataCollection && Params.TemplateTargetClass->IsChildOf(APCGIsolatedActor::StaticClass());
 
 	if (bMoveCollectionToIsolatedActor)
 	{
-		Params.TemplateTargetClass = APCGIsolatedActor::StaticClass();
 		Params.TargetDataCollection = &TempCollection;
 	}
 
