[PCG] ChangeTracking partition fix: refact regression CL 48622635

#rnx
#Synced-CL 49164840

diff --git a/Engine/Plugins/PCG/Source/PCG/Private/ChangeTracking/PCGChangeTracker.cpp b/Engine/Plugins/PCG/Source/PCG/Private/ChangeTracking/PCGChangeTracker.cpp
--- a/Engine/Plugins/PCG/Source/PCG/Private/ChangeTracking/PCGChangeTracker.cpp
+++ b/Engine/Plugins/PCG/Source/PCG/Private/ChangeTracking/PCGChangeTracker.cpp
@@ -8,6 +8,15 @@
 #include "PCGGraphExecutionStateInterface.h"
 #include "Subsystems/PCGSubsystem.h"
 
+namespace PCGChangeTracker
+{
+	/** Have a CVar for dynamic tracking for runtime until we have a better solution to also have "culling" for local components and dynamic tracking. */
+	static TAutoConsoleVariable<bool> CVarDisableDynamicTrackingForRuntimeGen(
+		TEXT("pcg.RuntimeGeneration.DisableDynamicTracking"),
+		false,
+		TEXT("In Editor and with runtime gen, a change with one tracked element will refresh all the local components. If it is too resource intensive, it can be disabled."));
+}
+
 void FPCGChangeTracker::ResetCurrentExecution()
 {
 	PCG::TScopeLock Lock(CurrentExecutionDynamicTrackingLock);
@@ -167,10 +176,10 @@ void FPCGChangeTracker::UpdateDynamicTracking(IPCGGraphExecutionSource* InExecut
 	const bool bIsLocal = InExecutionSource->GetExecutionState().IsLocalSource();
 	const bool bIsRuntime = InExecutionSource->GetExecutionState().IsManagedByRuntimeGenSystem();
 
-	/*if (bIsRuntime && PCGComponent::CVarDisableDynamicTrackingForRuntimeGen.GetValueOnAnyThread())
+	if (bIsRuntime && PCGChangeTracker::CVarDisableDynamicTrackingForRuntimeGen.GetValueOnAnyThread())
 	{
 		return;
-	}*/
+	}
 
 	// If the component is local, we defer the tracking to the original component.
 	// So move everything to the original (while making sure we are not duplicating keys/settings).
diff --git a/Engine/Plugins/PCG/Source/PCG/Private/PCGComponent.cpp b/Engine/Plugins/PCG/Source/PCG/Private/PCGComponent.cpp
--- a/Engine/Plugins/PCG/Source/PCG/Private/PCGComponent.cpp
+++ b/Engine/Plugins/PCG/Source/PCG/Private/PCGComponent.cpp
@@ -69,14 +69,6 @@ namespace PCGComponent
 		Delegate.Broadcast(PCGComponent, std::forward<Args>(InArgs)...);
 	}
 
-#if WITH_EDITOR
-	/** Have a CVar for dynamic tracking for runtime until we have a better solution to also have "culling" for local components and dynamic tracking. */
-	static TAutoConsoleVariable<bool> CVarDisableDynamicTrackingForRuntimeGen(
-		TEXT("pcg.RuntimeGeneration.DisableDynamicTracking"),
-		false,
-		TEXT("In Editor and with runtime gen, a change with one tracked element will refresh all the local components. If it is too resource intensive, it can be disabled."));
-#endif
-	
 	static TAutoConsoleVariable<bool> CVarDisableStealingForFlattenInPostProcessGraph(
 		TEXT("pcg.Advanced.DisableStealingForFlattenInPostProcessGraph"),
 		false,
diff --git a/Engine/Plugins/PCG/Source/PCG/Private/PCGComponentExecutionState.cpp b/Engine/Plugins/PCG/Source/PCG/Private/PCGComponentExecutionState.cpp
--- a/Engine/Plugins/PCG/Source/PCG/Private/PCGComponentExecutionState.cpp
+++ b/Engine/Plugins/PCG/Source/PCG/Private/PCGComponentExecutionState.cpp
@@ -324,6 +324,18 @@ void FPCGComponentExecutionState::ForEachSettingTrackingKey(const FPCGSelectionK
 	Component->ApplyToEachSettings(InKey, InCallback);
 }
 
+void FPCGComponentExecutionState::AddLocalSourceDynamicTrackingKeys(const FPCGSelectionKeyToSettingsMap& LocalSourceDynamicTrackingKeys)
+{
+	check(Component);
+	Component->ChangeTracker.AddLocalSourceDynamicTrackingKeys(LocalSourceDynamicTrackingKeys);
+}
+
+void FPCGComponentExecutionState::AddLocalSourceCurrentTrackingKeys(const FPCGSelectionKeyToSettingsMap& LocalSourceCurrentExecutionTrackingKeys, const TSet<const UPCGSettings*>& LocalSourceCurrentExecutionTrackingSettings)
+{
+	check(Component);
+	Component->ChangeTracker.AddLocalSourceCurrentTrackingKeys(LocalSourceCurrentExecutionTrackingKeys, LocalSourceCurrentExecutionTrackingSettings);
+}
+
 bool FPCGComponentExecutionState::IsRefreshInProgress() const
 {
 	check(Component);
diff --git a/Engine/Plugins/PCG/Source/PCG/Public/PCGComponentExecutionState.h b/Engine/Plugins/PCG/Source/PCG/Public/PCGComponentExecutionState.h
--- a/Engine/Plugins/PCG/Source/PCG/Public/PCGComponentExecutionState.h
+++ b/Engine/Plugins/PCG/Source/PCG/Public/PCGComponentExecutionState.h
@@ -65,7 +65,9 @@ public:
 	virtual TArray<FPCGSelectionKey> GatherTrackingKeys() const override;
 	virtual bool IsKeyTrackedAndCulled(const FPCGSelectionKey& Key, bool& bOutIsCulled) const override;
 	virtual void ForEachSettingTrackingKey(const FPCGSelectionKey& InKey, const TFunctionRef<void(const FPCGSelectionKey&, const FPCGSettingsAndCulling&)> InCallback) const override;
-	
+	virtual void AddLocalSourceDynamicTrackingKeys(const FPCGSelectionKeyToSettingsMap& LocalSourceDynamicTrackingKeys) override;
+	virtual void AddLocalSourceCurrentTrackingKeys(const FPCGSelectionKeyToSettingsMap& LocalSourceCurrentExecutionTrackingKeys, const TSet<const UPCGSettings*>& LocalSourceCurrentExecutionTrackingSettings) override;
+
 	virtual bool IsRefreshInProgress() const override;
 	virtual FPCGDynamicTrackingPriority GetDynamicTrackingPriority() const override;
 #endif // WITH_EDITOR
