[PCG] Improve Distance node perf

* Removing 3 virtuals that are called for all the pairs of points in Source/Target, which adds up a lot.

#rnx
#Synced-CL 49209048

diff --git a/Engine/Plugins/PCG/Source/PCG/Private/Elements/PCGDistance.cpp b/Engine/Plugins/PCG/Source/PCG/Private/Elements/PCGDistance.cpp
--- a/Engine/Plugins/PCG/Source/PCG/Private/Elements/PCGDistance.cpp
+++ b/Engine/Plugins/PCG/Source/PCG/Private/Elements/PCGDistance.cpp
@@ -320,17 +320,30 @@ bool FPCGDistanceElement::ExecuteInternal(FPCGContext* Context) const
 				SourcePointData->CopyPointsTo(OutputData, StartReadIndex, StartWriteIndex, Count);
 			}
 
-			const TConstPCGValueRange<FTransform> SourceTransformRange = SourcePointData->GetConstTransformValueRange();
-			const TConstPCGValueRange<FVector> SourceBoundsMinRange = SourcePointData->GetConstBoundsMinValueRange();
-			const TConstPCGValueRange<FVector> SourceBoundsMaxRange = SourcePointData->GetConstBoundsMaxValueRange();
+			struct FPointDataInfo
+			{
+				explicit FPointDataInfo(const UPCGBasePointData* InPointData)
+					: PointData(InPointData)
+					, TransformRange(PointData->GetConstTransformValueRange())
+					, BoundsMinRange(PointData->GetConstBoundsMinValueRange())
+					, BoundsMaxRange(PointData->GetConstBoundsMaxValueRange())
+				{}
+
+				const UPCGBasePointData* PointData;
+				const TConstPCGValueRange<FTransform> TransformRange;
+				const TConstPCGValueRange<FVector> BoundsMinRange;
+				const TConstPCGValueRange<FVector> BoundsMaxRange;
+			};
+
+			const FPointDataInfo SourceDataInfo{SourcePointData};
 
 			TPCGValueRange<float> OutputDensityRange = bSetDensity ? OutputData->GetDensityValueRange(/*bAllocate=*/false) : TPCGValueRange<float>();
 
 			for (int32 Index = StartReadIndex; Index < (StartReadIndex + Count); ++Index)
 			{
-				const FTransform& SourceTransform = SourceTransformRange[Index];
-				const FVector& SourceBoundsMin = SourceBoundsMinRange[Index];
-				const FVector& SourceBoundsMax = SourceBoundsMaxRange[Index];
+				const FTransform& SourceTransform = SourceDataInfo.TransformRange[Index];
+				const FVector& SourceBoundsMin = SourceDataInfo.BoundsMinRange[Index];
+				const FVector& SourceBoundsMax = SourceDataInfo.BoundsMaxRange[Index];
 
 				const FBoxSphereBounds SourceQueryBounds = FBoxSphereBounds(FBox(SourceBoundsMin - FVector(MaximumDistance), SourceBoundsMax + FVector(MaximumDistance))).TransformBy(SourceTransform);
 
@@ -340,10 +353,10 @@ bool FPCGDistanceElement::ExecuteInternal(FPCGContext* Context) const
 				FVector MinDistanceVector = FVector::ZeroVector;
 
 				// Signed distance field for calculating the closest point of source and target
-				auto CalculateSDF = [Index, &MinDistanceSquared, &MinDistanceVector, SourcePointData, &SourceTransform, &SourceBoundsMin, &SourceBoundsMax, SourceCenter, SourceShape, TargetShape](const UPCGBasePointData* TargetPointData, const int32 TargetPointIndex, const FBoxSphereBounds& Bounds)
+				auto CalculateSDF = [Index, &MinDistanceSquared, &MinDistanceVector, SourcePointData, &SourceTransform, &SourceBoundsMin, &SourceBoundsMax, SourceCenter, SourceShape, TargetShape](const FPointDataInfo& TargetDataInfo, const int32 TargetPointIndex, const FBoxSphereBounds& Bounds)
 				{
 					// If the source pointer and target pointer are the same, ignore distance to the exact same point
-					if (Index == TargetPointIndex && SourcePointData == TargetPointData)
+					if (Index == TargetPointIndex && SourcePointData == TargetDataInfo.PointData)
 					{
 						return;
 					}
@@ -351,7 +364,7 @@ bool FPCGDistanceElement::ExecuteInternal(FPCGContext* Context) const
 					const FVector& TargetCenter = Bounds.Origin;
 
 					const FVector SourceShapePos = PCGDistance::CalcPosition(SourceShape, SourceTransform, SourceBoundsMin, SourceBoundsMax, SourceCenter, TargetCenter);
-					const FVector TargetShapePos = PCGDistance::CalcPosition(TargetShape, TargetPointData->GetTransform(TargetPointIndex), TargetPointData->GetBoundsMin(TargetPointIndex), TargetPointData->GetBoundsMax(TargetPointIndex), TargetCenter, SourceCenter);
+					const FVector TargetShapePos = PCGDistance::CalcPosition(TargetShape, TargetDataInfo.TransformRange[TargetPointIndex], TargetDataInfo.BoundsMinRange[TargetPointIndex], TargetDataInfo.BoundsMaxRange[TargetPointIndex], TargetCenter, SourceCenter);
 
 					const FVector ToTargetShapeDir = TargetShapePos - SourceShapePos;
 					const FVector ToTargetCenterDir = TargetCenter - SourceCenter;
@@ -373,9 +386,9 @@ bool FPCGDistanceElement::ExecuteInternal(FPCGContext* Context) const
 					const PCGPointOctree::FPointOctree& Octree = TargetPointData->GetPointOctree();
 					
 					Octree.FindElementsWithBoundsTest(FBoxCenterAndExtent(SourceQueryBounds.Origin, SourceQueryBounds.BoxExtent),
-						[TargetPointData, &CalculateSDF](const PCGPointOctree::FPointRef& InPointRef)
+						[TargetDataInfo = FPointDataInfo{TargetPointData}, &CalculateSDF](const PCGPointOctree::FPointRef& InPointRef)
 						{
-							CalculateSDF(TargetPointData, InPointRef.Index, InPointRef.Bounds);
+							CalculateSDF(TargetDataInfo, InPointRef.Index, InPointRef.Bounds);
 						});
 				};
 
