Addapt path finding for base point data

#Synced-CL 49331429

diff --git a/Engine/Plugins/PCG/Source/PCG/Private/Elements/PCGPathfindingElement.cpp b/Engine/Plugins/PCG/Source/PCG/Private/Elements/PCGPathfindingElement.cpp
--- a/Engine/Plugins/PCG/Source/PCG/Private/Elements/PCGPathfindingElement.cpp
+++ b/Engine/Plugins/PCG/Source/PCG/Private/Elements/PCGPathfindingElement.cpp
@@ -5,12 +5,12 @@
 #include "PCGContext.h"
 #include "PCGComponent.h"
 #include "PCGParamData.h"
-#include "Data/PCGPointData.h"
+#include "Data/PCGBasePointData.h"
 #include "Data/PCGSplineData.h"
 #include "Helpers/PCGHelpers.h"
 #include "Metadata/Accessors/IPCGAttributeAccessor.h"
 #include "Metadata/Accessors/PCGAttributeAccessorHelpers.h"
-#include "Metadata/Accessors/PCGAttributeAccessorKeys.h"
+#include "Metadata/Accessors/PCGCustomAccessor.h"
 #include "SpatialAlgo/PCGAStar.h"
 
 #include "Components/SplineComponent.h"
@@ -33,8 +33,9 @@ namespace PCGPathfindingElement
 
 	namespace Helpers
 	{
-		TArray<FSplinePoint> ConvertPathToSplinePoints(const TArrayView<const FPCGPoint> Path, const EPCGPathfindingSplineMode SplineMode)
+		TArray<FSplinePoint> ConvertPathToSplinePoints(const PCGSpatialAlgo::AStar::FSearchState& SearchState, const TArrayView<PCGSpatialAlgo::AStar::FPointDescription> Path, const EPCGPathfindingSplineMode SplineMode)
 		{
+			using namespace PCGSpatialAlgo::AStar;
 			ESplinePointType::Type SplineCurveMode = ESplinePointType::Type::Constant;
 			switch (SplineMode)
 			{
@@ -51,11 +52,11 @@ namespace PCGPathfindingElement
 			SplinePoints.Reserve(Path.Num());
 
 			int Index = 0;
-			Algo::Transform(Path, SplinePoints, [&Index, SplineCurveMode](const FPCGPoint& Point)
+			Algo::Transform(Path, SplinePoints, [&SearchState, &Index, SplineCurveMode](const FPointDescription& Point)
 			{
 				return FSplinePoint(
 						Index++, // Spline points must be indexed in ascending order
-						Point.Transform.GetLocation(),
+						Point.GetLocation(),
 						FVector::ZeroVector,
 						FVector::ZeroVector,
 						FRotator::ZeroRotator,
@@ -165,7 +166,7 @@ bool FPCGPathfindingElement::PrepareDataInternal(FPCGContext* InContext) const
 		return true;
 	}
 
-	const EPCGTimeSliceInitResult ExecResult = Context->InitializePerExecutionState([&PointInputs, &FilterActorInputData, Settings](const ContextType* SlicedContext, ExecStateType& OutState)
+	const EPCGTimeSliceInitResult ExecResult = Context->InitializePerExecutionState([&Context = InContext, &PointInputs, &FilterActorInputData, Settings](ContextType* SlicedContext, ExecStateType& OutState)
 	{
 		TRACE_CPUPROFILER_EVENT_SCOPE(FPCGClusterElement::PrepareData::InitializePerExecutionState);
 
@@ -204,57 +205,97 @@ bool FPCGPathfindingElement::PrepareDataInternal(FPCGContext* InContext) const
 			return LocationsInputData[0].Data;
 		};
 
-		auto PopulateArray = [SlicedContext](const UPCGData* Data, const FPCGAttributePropertyInputSelector& InSelector, TArray<FPCGPoint>& OutPointArray)
+		bool bShouldCopyAllAttributesFromStartAndGoal = Settings->bCopyOriginatingPoints || Settings->CostFunctionMode==EPCGPathfindingCostFunctionMode::CostMultiplier || Settings->CostFunctionMode==EPCGPathfindingCostFunctionMode::FitnessScore;
+		auto PopulateArray = [&Context, SlicedContext, bShouldCopyAllAttributesFromStartAndGoal](const UPCGData* Data, const FPCGAttributePropertyInputSelector& InSelector) -> FPCGPointInputRange
 		{
 			const FPCGAttributePropertyInputSelector Selector = InSelector.CopyAndFixLast(Data);
+
+			bool AccessorIsPosition = (Selector.GetSelection() == EPCGAttributePropertySelection::Property && Selector.GetPointProperty() == EPCGPointProperties::Position);
+			const UPCGBasePointData* PointData = Cast<UPCGBasePointData>(Data);
+
+			// Simple case, reuse the original PointData.
+			if (PointData && AccessorIsPosition)
+			{
+				return FPCGPointInputRange{PointData, 0, PointData->GetNumPoints()};
+			}
+
 			const TUniquePtr<const IPCGAttributeAccessor> Accessor = PCGAttributeAccessorHelpers::CreateConstAccessor(Data, Selector);
 			const TUniquePtr<const IPCGAttributeAccessorKeys> Keys = PCGAttributeAccessorHelpers::CreateConstKeys(Data, Selector);
+		
 			if (!Accessor || !Keys)
 			{
 				PCGLog::Metadata::LogFailToCreateAccessorError(Selector, SlicedContext);
-				return;
+				return {};
 			}
 
-			if (const UPCGPointData* PointData = Cast<UPCGPointData>(Data))
+			UPCGBasePointData* OutPointData = nullptr; 
+			if (PointData)
 			{
-				OutPointArray = PointData->GetPointsCopy();
+				if (!PointData->IsEmpty())
+				{
+					OutPointData = FPCGContext::NewPointData_AnyThread(Context);
+					UPCGBasePointData::SetPoints(PointData, OutPointData, /* InDataIndices=*/{}, /* bCopyAll=*/true);
+				}
 			}
-			else
+			else if (Keys->GetNum() > 0)
 			{
-				OutPointArray.SetNum(Keys->GetNum());
+				OutPointData = FPCGContext::NewPointData_AnyThread(Context);
+				OutPointData->SetNumPoints(Keys->GetNum());
 			}
 
-			// For Points with $Position as the selection, nothing is needed. For all other conditions, set the position based on the selected attribute/property.
-			if (OutPointArray.IsEmpty() || (Selector.GetSelection() == EPCGAttributePropertySelection::Property && Selector.GetPointProperty() != EPCGPointProperties::Position))
+			if (OutPointData == nullptr)
 			{
-				return;
+				return {};
+			}
+				
+			SlicedContext->TrackObject(OutPointData);
+			if (bShouldCopyAllAttributesFromStartAndGoal && PointData)
+			{
+				// We will want to copy properties from Start/Goal into final result, so it should follow.
+				OutPointData->AllocateProperties(PointData->GetAllocatedProperties()|EPCGPointNativeProperties::Transform|EPCGPointNativeProperties::Seed);
+				FPCGInitializeFromDataParams InitParams(PointData);
+				InitParams.bInheritSpatialData = false;
+				OutPointData->InitializeFromDataWithParams(InitParams);
+			}
+			else
+			{
+				OutPointData->AllocateProperties(EPCGPointNativeProperties::Transform|EPCGPointNativeProperties::Seed);
 			}
 
 			TArray<FVector> Locations;
 			Locations.SetNumUninitialized(Keys->GetNum());
 			Accessor->GetRange(TArrayView<FVector>(Locations), 0, *Keys, EPCGAttributeAccessorFlags::AllowBroadcastAndConstructible);
 
-			check(Locations.Num() == OutPointArray.Num());
-			for (int i = 0; i < OutPointArray.Num(); ++i)
+			check(Locations.Num() == OutPointData->GetNumPoints());
+			TPCGValueRange<FTransform> TransformRange = OutPointData->GetTransformValueRange();
+			TPCGValueRange<int32> SeedRange = OutPointData->GetSeedValueRange();
+			for (int i = 0; i < OutPointData->GetNumPoints(); ++i)
 			{
-				OutPointArray[i].Transform.SetLocation(Locations[i]);
-				OutPointArray[i].Seed = PCGHelpers::ComputeSeedFromPosition(Locations[i]);
+				TransformRange[i] = FTransform{Locations[i]};
+				SeedRange[i] = PCGHelpers::ComputeSeedFromPosition(Locations[i]);
 			}
+			
+			return FPCGPointInputRange{OutPointData, 0, OutPointData->GetNumPoints()};
 		};
 
 		if (Settings->bStartLocationsAsInput)
 		{
 			if (const UPCGData* Data = GetLocationsFromInput(PCGPathfindingElement::Constants::StartLocationsInputPinLabel))
 			{
-				PopulateArray(Data, Settings->StartLocationAttribute, OutState.StartPoints);
+				OutState.StartPoints = PopulateArray(Data, Settings->StartLocationAttribute);
 			}
 		}
 		else
 		{
-			OutState.StartPoints = {FPCGPoint(FTransform(Settings->Start), /*InDensity=*/1, PCGHelpers::ComputeSeedFromPosition(Settings->Start))};
+			UPCGBasePointData* StartPoints = FPCGContext::NewPointData_AnyThread(Context);
+			SlicedContext->TrackObject(StartPoints);
+			StartPoints->SetNumPoints(1);
+			StartPoints->SetTransform(FTransform(Settings->Start));
+			OutState.StartPoints = FPCGPointInputRange{StartPoints, 0, StartPoints->GetNumPoints()};
+			
 		}
 
-		if (OutState.StartPoints.IsEmpty())
+		if (!OutState.StartPoints.PointData || OutState.StartPoints.RangeSize == 0)
 		{
 			return EPCGTimeSliceInitResult::NoOperation;
 		}
@@ -263,28 +304,32 @@ bool FPCGPathfindingElement::PrepareDataInternal(FPCGContext* InContext) const
 		{
 			if (const UPCGData* Data = GetLocationsFromInput(PCGPathfindingElement::Constants::GoalLocationsInputPinLabel))
 			{
-				PopulateArray(Data, Settings->GoalLocationAttribute, OutState.GoalPoints);
+				OutState.GoalPoints = PopulateArray(Data, Settings->GoalLocationAttribute);
 			}
 		}
 		else
 		{
-			OutState.GoalPoints = {FPCGPoint(FTransform(Settings->Goal), /*InDensity=*/1, PCGHelpers::ComputeSeedFromPosition(Settings->Goal))};
+			UPCGBasePointData* GoalPoints = FPCGContext::NewPointData_AnyThread(Context);
+			SlicedContext->TrackObject(GoalPoints);
+			GoalPoints->SetNumPoints(1);
+			GoalPoints->SetTransform(FTransform(Settings->Goal));
+			OutState.GoalPoints = FPCGPointInputRange{GoalPoints, 0, GoalPoints->GetNumPoints()};
 		}
 
-		if (OutState.GoalPoints.IsEmpty())
+		if (!OutState.GoalPoints.PointData || OutState.GoalPoints.RangeSize == 0)
 		{
 			return EPCGTimeSliceInitResult::NoOperation;
 		}
 
 		OutState.IterationCount = Settings->GoalMappingMode == EPCGPathfindingGoalMappingMode::EachStartToEachGoal
-			? OutState.StartPoints.Num() * OutState.GoalPoints.Num()
-			: OutState.StartPoints.Num();
+			? OutState.StartPoints.RangeSize * OutState.GoalPoints.RangeSize
+			: OutState.StartPoints.RangeSize;
 
 		// Validate that if in N:N -> Start:Goal mode, the inputs are the correct cardinality. All other cases were validated in PrepareData and are acceptable.
-		if (Settings->GoalMappingMode == EPCGPathfindingGoalMappingMode::EachStartToPairwiseGoal && OutState.StartPoints.Num() != OutState.GoalPoints.Num())
+		if (Settings->GoalMappingMode == EPCGPathfindingGoalMappingMode::EachStartToPairwiseGoal && OutState.StartPoints.RangeSize != OutState.GoalPoints.RangeSize)
 		{
 			PCGLog::LogErrorOnGraph(FText::Format(
-		LOCTEXT("EachStartToPairwiseGoalInvalidMapping", "For 'Each Start To Pairwise Goal' pathfinding mode, there must be a one-to-one mapping between start locations and goal locations. The current input was {0}:{1} Start->Goal locations."), OutState.StartPoints.Num(), OutState.GoalPoints.Num()), SlicedContext);
+				LOCTEXT("EachStartToPairwiseGoalInvalidMapping", "For 'Each Start To Pairwise Goal' pathfinding mode, there must be a one-to-one mapping between start locations and goal locations. The current input was {0}:{1} Start->Goal locations."), OutState.StartPoints.RangeSize, OutState.GoalPoints.RangeSize), SlicedContext);
 			return EPCGTimeSliceInitResult::AbortExecution;
 		}
 
@@ -300,7 +345,7 @@ bool FPCGPathfindingElement::PrepareDataInternal(FPCGContext* InContext) const
 	{
 		TRACE_CPUPROFILER_EVENT_SCOPE(FPCGClusterElement::PrepareData::InitializePerIterationStates);
 
-		const UPCGPointData* PointData = Cast<UPCGPointData>(PointInputs[IterationIndex].Data);
+		const UPCGBasePointData* PointData = Cast<UPCGBasePointData>(PointInputs[IterationIndex].Data);
 		if (!PointData || PointData->IsEmpty())
 		{
 			return EPCGTimeSliceInitResult::NoOperation;
@@ -313,10 +358,11 @@ bool FPCGPathfindingElement::PrepareDataInternal(FPCGContext* InContext) const
 
 		if (Settings->CostFunctionMode != EPCGPathfindingCostFunctionMode::Distance)
 		{
+			//@todo_pcg The accessor will only work for PointData (It will not work on goal or start point data), we need to fix this without affecting the perf (i.e.: constructing an accessor every time we check the cost)
 			const FPCGAttributePropertyInputSelector Selector = Settings->CostAttribute.CopyAndFixLast(PointData);
 			CostAccessor = MakeShareable(PCGAttributeAccessorHelpers::CreateConstAccessor(PointData, Selector).Release());
 
-			FPCGAttributeAccessorKeysPoints Keys(PointData->GetPoints());
+			FPCGPointTransform::ConstValueRange Keys = PointData->GetConstTransformValueRange();
 
 			if (!CostAccessor)
 			{
@@ -336,7 +382,7 @@ bool FPCGPathfindingElement::PrepareDataInternal(FPCGContext* InContext) const
 			const UPCGData* FilterActorsData = !FilterActorInputData.IsEmpty() ? FilterActorInputData[IterationIndex % FilterActorInputData.Num()].Data : nullptr;
 
 			// Accept only point or param data
-			if (FilterActorsData && !FilterActorsData->IsA<UPCGPointData>() && !FilterActorsData->IsA<UPCGParamData>())
+			if (FilterActorsData && !FilterActorsData->IsA<UPCGBasePointData>() && !FilterActorsData->IsA<UPCGParamData>())
 			{
 				PCGLog::InputOutput::LogTypedDataNotFoundWarning(EPCGDataType::PointOrParam, PCGWorldRayHitConstants::FilterActorPinLabel, Context);
 				return EPCGTimeSliceInitResult::NoOperation;
@@ -381,38 +427,38 @@ bool FPCGPathfindingElement::PrepareDataInternal(FPCGContext* InContext) const
 			{
 				const double MaxFitnessPenaltyFactor = FMath::Max(Settings->MaximumFitnessPenaltyFactor, 1.0);
 
-				OutState.SearchState.CostFunction = [FitnessAccessor = CostAccessor, MaxFitnessPenaltyFactor, PathTraceTest = MoveTemp(LineTraceTest)](const double PreviousNodeCost, const FPCGPoint* PreviousNodePoint, const double DistanceToCurrentSquared, const FPCGPoint* CurrentNodePoint)
+				OutState.SearchState.CostFunction = [&SearchState = OutState.SearchState, FitnessAccessor = CostAccessor, MaxFitnessPenaltyFactor, PathTraceTest = MoveTemp(LineTraceTest)](const PCGSpatialAlgo::AStar::FPointDescription& PreviousNodePoint, const double DistanceToCurrentSquared, const PCGSpatialAlgo::AStar::FPointDescription& CurrentNodePoint)
 				{
-					if (!PathTraceTest(PreviousNodePoint->Transform.GetLocation(), CurrentNodePoint->Transform.GetLocation()))
+					if (!PathTraceTest(PreviousNodePoint.GetLocation(), CurrentNodePoint.GetLocation()))
 					{
 						return std::numeric_limits<double>::max();
 					}
 
 					double FitnessScore = 1.0;
-					FPCGAttributeAccessorKeysPoints Key(*CurrentNodePoint);
+					FPCGAttributeAccessorKeysPointIndices Key(CurrentNodePoint.SrcPointRange.PointData);
 
-					FitnessAccessor->Get(FitnessScore, Key, EPCGAttributeAccessorFlags::AllowBroadcastAndConstructible);
+					FitnessAccessor->Get(FitnessScore, CurrentNodePoint.SrcPointRange.RangeStartIndex + CurrentNodePoint.PointIndexInRange, Key, EPCGAttributeAccessorFlags::AllowBroadcastAndConstructible);
 					FitnessScore = FMath::Clamp(FitnessScore, 0.0, 1.0);
 
-					return PreviousNodeCost + (1.0 - FitnessScore) * MaxFitnessPenaltyFactor * FMath::Sqrt(DistanceToCurrentSquared);
+					return (1.0 - FitnessScore) * MaxFitnessPenaltyFactor * FMath::Sqrt(DistanceToCurrentSquared);
 				};
 			}
 			else if (Settings->CostFunctionMode == EPCGPathfindingCostFunctionMode::CostMultiplier)
 			{
-				OutState.SearchState.CostFunction = [MultiplierAccessor = CostAccessor, PathTraceTest = MoveTemp(LineTraceTest)](const double PreviousNodeCost, const FPCGPoint* PreviousNodePoint, const double DistanceToCurrentSquared, const FPCGPoint* CurrentNodePoint)
+				OutState.SearchState.CostFunction = [&SearchState = OutState.SearchState, MultiplierAccessor = CostAccessor, PathTraceTest = MoveTemp(LineTraceTest)](const PCGSpatialAlgo::AStar::FPointDescription& PreviousNodePoint, const double DistanceToCurrentSquared, const PCGSpatialAlgo::AStar::FPointDescription& CurrentNodePoint)
 				{
-					if (!PathTraceTest(PreviousNodePoint->Transform.GetLocation(), CurrentNodePoint->Transform.GetLocation()))
+					if (!PathTraceTest(PreviousNodePoint.GetLocation(), CurrentNodePoint.GetLocation()))
 					{
 						return std::numeric_limits<double>::max();
 					}
 
 					double Multiplier = 1.0;
-					const FPCGAttributeAccessorKeysPoints Key(*CurrentNodePoint);
+					FPCGAttributeAccessorKeysPointIndices Key(CurrentNodePoint.SrcPointRange.PointData);
 
-					MultiplierAccessor->Get(Multiplier, Key, EPCGAttributeAccessorFlags::AllowBroadcastAndConstructible);
+					MultiplierAccessor->Get(Multiplier, CurrentNodePoint.SrcPointRange.RangeStartIndex + CurrentNodePoint.PointIndexInRange, Key, EPCGAttributeAccessorFlags::AllowBroadcastAndConstructible);
 					Multiplier = FMath::Max(Multiplier, 1.0);
 
-					return PreviousNodeCost + Multiplier * FMath::Sqrt(DistanceToCurrentSquared);
+					return Multiplier * FMath::Sqrt(DistanceToCurrentSquared);
 				};
 			}
 			else
@@ -423,15 +469,15 @@ bool FPCGPathfindingElement::PrepareDataInternal(FPCGContext* InContext) const
 		else if (Settings->bUsePathTraces)
 		{
 			// Use distance but with line trace
-			OutState.SearchState.CostFunction = [PathTraceTest = MoveTemp(LineTraceTest)](const double PreviousNodeCost, const FPCGPoint* PreviousNodePoint, const double DistanceToCurrentSquared, const FPCGPoint* CurrentNodePoint)
+			OutState.SearchState.CostFunction = [&SearchState = OutState.SearchState, PathTraceTest = MoveTemp(LineTraceTest)](const PCGSpatialAlgo::AStar::FPointDescription& PreviousNodePoint, const double DistanceToCurrentSquared, const PCGSpatialAlgo::AStar::FPointDescription& CurrentNodePoint)
 			{
-				if (!PathTraceTest(PreviousNodePoint->Transform.GetLocation(), CurrentNodePoint->Transform.GetLocation()))
+				if (!PathTraceTest(PreviousNodePoint.GetLocation(), CurrentNodePoint.GetLocation()))
 				{
 					return std::numeric_limits<double>::max();
 				}
 				else
 				{
-					return PCGSpatialAlgo::AStar::Cost::CalculateCost_EuclideanDistance(PreviousNodeCost, PreviousNodePoint, DistanceToCurrentSquared, CurrentNodePoint);
+					return PCGSpatialAlgo::AStar::Cost::CalculateCost_EuclideanDistance(PreviousNodePoint, DistanceToCurrentSquared, CurrentNodePoint);
 				}
 			};
 		}
@@ -469,36 +515,39 @@ bool FPCGPathfindingElement::ExecuteInternal(FPCGContext* InContext) const
 		while (IterationState.PathIterationIndex < ExecutionState.IterationCount)
 		{
 			using PCGSpatialAlgo::AStar::ESearchResult;
-			TArray<FPCGPoint> FinalPath;
+			TArray<PCGSpatialAlgo::AStar::FPointDescription> FinalPath;
 
 			// Initialize the next iteration.
 			if (IterationState.PathIterationIndex != IterationState.LastPathIterationIndex)
 			{
 				// Per start, per goal. GoalPoints was validated as not empty (0) in PrepareData.
 				const int32 StartPointIndex = bCartesianMapping
-					? IterationState.PathIterationIndex / ExecutionState.GoalPoints.Num()
+					? IterationState.PathIterationIndex / ExecutionState.GoalPoints.RangeSize
 					: IterationState.PathIterationIndex;
-				PCGSpatialAlgo::AStar::Initialize(&ExecutionState.StartPoints[StartPointIndex], IterationState.SearchState);
+				IterationState.StartPoints = FPCGPointInputRange{ExecutionState.StartPoints.PointData, ExecutionState.StartPoints.RangeStartIndex + StartPointIndex, 1};
+
+
+				if (Settings->GoalMappingMode == EPCGPathfindingGoalMappingMode::EachStartToNearestGoal)
+				{
+					IterationState.GoalPoints = ExecutionState.GoalPoints;
+				}
+				else // Other mapping modes will do a one-to-one comparison of a single start and single goal.
+				{
+					const int32 GoalPointIndex = bCartesianMapping
+						? IterationState.PathIterationIndex % ExecutionState.GoalPoints.RangeSize
+						: IterationState.PathIterationIndex;
+					IterationState.GoalPoints = FPCGPointInputRange{ExecutionState.GoalPoints.PointData, ExecutionState.GoalPoints.RangeStartIndex + GoalPointIndex, 1};
+				}
+
+				PCGSpatialAlgo::AStar::Initialize(IterationState.StartPoints,IterationState.GoalPoints, IterationState.SearchState);
 				IterationState.LastPathIterationIndex = IterationState.PathIterationIndex;
 			}
 
-			TConstArrayView<FPCGPoint> GoalPointsView;
-			if (Settings->GoalMappingMode == EPCGPathfindingGoalMappingMode::EachStartToNearestGoal)
-			{
-				GoalPointsView = MakeConstArrayView<FPCGPoint>(ExecutionState.GoalPoints.GetData(), ExecutionState.GoalPoints.Num());
-			}
-			else // Other mapping modes will do a one-to-one comparison of a single start and single goal.
-			{
-				const int32 GoalPointIndex = bCartesianMapping
-					? IterationState.PathIterationIndex % ExecutionState.GoalPoints.Num()
-					: IterationState.PathIterationIndex;
-				GoalPointsView = MakeConstArrayView<FPCGPoint>(ExecutionState.GoalPoints.GetData() + GoalPointIndex, /*Size=*/1);
-			}
 
 			ESearchResult SearchResult;
 			do
 			{
-				SearchResult = ExecuteSearchIteration(ExecutionState.Settings, IterationState.SearchState, GoalPointsView, FinalPath);
+				SearchResult = ExecuteSearchIteration(ExecutionState.Settings, IterationState.SearchState, FinalPath);
 
 				if (SearchResult == ESearchResult::Processing && Context->ShouldStop())
 				{
@@ -525,20 +574,35 @@ bool FPCGPathfindingElement::ExecuteInternal(FPCGContext* InContext) const
 			// Output the path as either a spline or points.
 			if (Settings->bOutputAsSpline)
 			{
-				const TArray<FSplinePoint> SplinePoints = PCGPathfindingElement::Helpers::ConvertPathToSplinePoints(FinalPath, Settings->SplineMode);
+				const TArray<FSplinePoint> SplinePoints = PCGPathfindingElement::Helpers::ConvertPathToSplinePoints(IterationState.SearchState, FinalPath, Settings->SplineMode);
 				UPCGSplineData* SplineData = FPCGContext::NewObject_AnyThread<UPCGSplineData>(Context);
 				SplineData->Initialize(SplinePoints, /*bInClosedLoop=*/false, FTransform::Identity);
 				OutputData.Data = SplineData;
 			}
 			else
 			{
-				UPCGPointData* OutputPointData = FPCGContext::NewObject_AnyThread<UPCGPointData>(Context);
+				UPCGBasePointData* OutputPointData = FPCGContext::NewPointData_AnyThread(Context);
 				if (ExecutionState.Settings.bCopyOriginatingPoints)
 				{
-					OutputPointData->InitializeFromData(IterationState.SearchState.OriginatingPointData);
+					FPCGInitializeFromDataParams InitParams(IterationState.SearchState.OriginatingPointData);
+					InitParams.bInheritSpatialData = false;
+					OutputPointData->InitializeFromDataWithParams(InitParams);
+
+					InitParams.Source = IterationState.SearchState.GoalPoints.PointData;
+					OutputPointData->InitializeFromDataWithParams(InitParams);
+
+					InitParams.Source = IterationState.StartPoints.PointData;
+					OutputPointData->InitializeFromDataWithParams(InitParams);
+				}
+
+				OutputPointData->SetNumPoints(FinalPath.Num());
+				const EPCGPointNativeProperties PropertiesToCopy = ExecutionState.Settings.bCopyOriginatingPoints ? EPCGPointNativeProperties::All : EPCGPointNativeProperties::AllProperties;
+				for (int i = 0; i < FinalPath.Num(); i++)
+				{
+					const PCGSpatialAlgo::AStar::FPointDescription& Point = FinalPath[i];
+					Point.SrcPointRange.PointData->CopyPropertiesTo(OutputPointData, Point.SrcPointRange.RangeStartIndex + Point.PointIndexInRange, i, 1, PropertiesToCopy);
 				}
 
-				OutputPointData->GetMutablePoints() = std::move(FinalPath);
 				OutputData.Data = OutputPointData;
 			}
 
diff --git a/Engine/Plugins/PCG/Source/PCG/Private/SpatialAlgo/PCGAStar.cpp b/Engine/Plugins/PCG/Source/PCG/Private/SpatialAlgo/PCGAStar.cpp
--- a/Engine/Plugins/PCG/Source/PCG/Private/SpatialAlgo/PCGAStar.cpp
+++ b/Engine/Plugins/PCG/Source/PCG/Private/SpatialAlgo/PCGAStar.cpp
@@ -7,6 +7,54 @@
 
 namespace PCGSpatialAlgo::AStar
 {
+	struct FNode
+	{
+		explicit FNode(const FPointDescription* InPointDescription, const int32 InPreviousNodeIndex, const double InLocalCost, const double InEstimatedGoalCost)
+			: PointDescription(InPointDescription)
+			, PreviousNodeIndex(InPreviousNodeIndex)
+			, LocalCost(InLocalCost)
+			, EstimatedGoalCost(InEstimatedGoalCost)
+		{}
+
+		const FVector& GetLocation() const
+		{
+			return PointDescription->GetLocation();
+		}
+		
+		const FPointDescription* PointDescription = nullptr;
+
+		/** The currently assigned parent node with the least cost. */
+		int32 PreviousNodeIndex = INDEX_NONE;
+		/** Concrete cost from start to this node. */
+		double LocalCost = 0.0;
+		/** Current total cost when heuristic is added to the local cost. */
+		double EstimatedGoalCost = 0.0;
+	};
+	
+	struct FSearchState::FInternalSearchData
+	{
+		static uint32 constexpr PreAllocNodeCount = 1024u;
+		using FNodeList = TArray<FNode, TInlineAllocator<PreAllocNodeCount>>;
+		/** Storage of the current known node. Source of truth for the nodes, so it should not be resized after its initial reserve. */
+		FNodeList NodeList;
+		/** Used as a priority list of nodes, sorted by cost. */
+		TArray<int32, TInlineAllocator<PreAllocNodeCount>> OpenIndexList;
+		/** Nodes that are no longer viable. */
+		TArray<int32, TInlineAllocator<PreAllocNodeCount>> ClosedIndexList;
+		/** To track nodes to their corresponding points. */
+		TMap<const FPointDescription*, int32> PointDescriptionToNodeIndexMap;
+
+		TArray<FPointDescription> InternalOriginatingPointsDescriptions;
+		FPointDescription InternalStartPointDescription;
+		TArray<FPointDescription> InternalGoalPointsDescriptions; 
+	};
+	
+	FSearchState::FSearchState():
+		InternalSearchData{MakeUnique<FInternalSearchData>()}
+	{}
+
+	FSearchState::~FSearchState() = default;
+
 namespace Helpers
 {
 	bool CompareNodes(const FNode& Node1, const FNode& Node2)
@@ -14,7 +62,7 @@ namespace Helpers
 		return Node1.EstimatedGoalCost < Node2.EstimatedGoalCost;
 	}
 
-	void BuildFinalPath(const FSearchState& SearchState, const FNode& FinalNode, const bool bCopyOriginatingPoints, TArray<FPCGPoint>& OutPath)
+	void BuildFinalPath(const FSearchState& SearchState, const FNode& FinalNode, const bool bCopyOriginatingPoints, TArray<FPointDescription>& OutPath)
 	{
 		check(&FinalNode);
 
@@ -23,17 +71,9 @@ namespace Helpers
 		const FNode* Node = &FinalNode;
 		do // Build path from the goal backwards.
 		{
-			check(Node->PCGPoint);
-			if (bCopyOriginatingPoints)
-			{
-				OutPath.Emplace(*Node->PCGPoint);
-			}
-			else
-			{
-				OutPath.Emplace(Node->PCGPoint->Transform, /*InDensity=*/1, PCGHelpers::ComputeSeedFromPosition(Node->PCGPoint->Transform.GetLocation()));
-			}
+			OutPath.Emplace(*Node->PointDescription);
 
-			Node = Node->PreviousNodeIndex != INDEX_NONE ? &SearchState.NodeList[Node->PreviousNodeIndex] : nullptr;
+			Node = Node->PreviousNodeIndex != INDEX_NONE ? &SearchState.InternalSearchData->NodeList[Node->PreviousNodeIndex] : nullptr;
 		}
 		while (Node);
 
@@ -41,32 +81,34 @@ namespace Helpers
 		Algo::Reverse(OutPath);
 	}
 
-	TTuple<const FPCGPoint*, double> GetNearestGoalToLocation(const FVector& Location, const TConstArrayView<FPCGPoint>& GoalPoints)
+	TTuple<int32, double> GetNearestGoalToLocation(const FSearchState& SearchState, const FPointDescription& FromPosition)
 	{
-		check(!GoalPoints.IsEmpty());
+		check(SearchState.InternalSearchData->InternalGoalPointsDescriptions.Num() > 0);
 
-		double MinSquaredDistanceToGoal = std::numeric_limits<double>::max();
+		double MinCostToGoal = std::numeric_limits<double>::max();
 		int32 CurrentGoalIndex = INDEX_NONE;
-		for (int32 GoalIndex = 0; GoalIndex < GoalPoints.Num(); ++GoalIndex)
+		for (int32 GoalIndex = 0; GoalIndex < SearchState.InternalSearchData->InternalGoalPointsDescriptions.Num(); ++GoalIndex)
 		{
-			const double SquareDistanceToGoal = FVector::DistSquared(Location, GoalPoints[GoalIndex].Transform.GetLocation());
-			if (SquareDistanceToGoal < MinSquaredDistanceToGoal)
+			const FPointDescription& CurGoal = SearchState.InternalSearchData->InternalGoalPointsDescriptions[GoalIndex];
+			const double CurDistToGoalSqr = FVector::DistSquared(FromPosition.GetLocation(), CurGoal.GetLocation());
+			const double CurCostToGoal = SearchState.CostFunction(FromPosition, CurDistToGoalSqr, CurGoal);
+			if (CurCostToGoal < MinCostToGoal)
 			{
-				MinSquaredDistanceToGoal = SquareDistanceToGoal;
+				MinCostToGoal = CurCostToGoal;
 				CurrentGoalIndex = GoalIndex;
 			}
 		}
 		check(CurrentGoalIndex != INDEX_NONE);
 
-		return MakeTuple(&GoalPoints[CurrentGoalIndex], MinSquaredDistanceToGoal);
+		return MakeTuple(CurrentGoalIndex, MinCostToGoal);
 	}
 }
 
 namespace Cost
 {
-	double CalculateCost_EuclideanDistance(const double PreviousNodeCost, const FPCGPoint* PreviousNodePoint, const double DistanceToPreviousNodeSquared, const FPCGPoint* CurrentPoint)
+	double CalculateCost_EuclideanDistance(const FPointDescription& PreviousNodePoint, const double DistanceToPreviousNodeSquared, const FPointDescription& CurrentPoint)
 	{
-		return PreviousNodeCost + FMath::Sqrt(DistanceToPreviousNodeSquared);
+		return FMath::Sqrt(DistanceToPreviousNodeSquared);
 	}
 }
 
@@ -79,56 +121,74 @@ namespace Heuristic
 	}
 }
 
+
 /** Initialize the search state for AStar. Must be called before ExecuteSearchIteration. */
-void Initialize(const FPCGPoint* const StartPoint, FSearchState& OutSearchState)
+void Initialize(const FPCGPointInputRange& StartPoint, const FPCGPointInputRange& GoalPoints, FSearchState& OutSearchState)
 {
-	check(StartPoint);
-
+	OutSearchState.GoalPoints = GoalPoints;
+	OutSearchState.InternalSearchData->InternalStartPointDescription = FPointDescription{StartPoint,0};
 	// Emplace the starting node on the open list. The costs will be 0 at the starting point.
-	OutSearchState.NodeList.Reset();
-	OutSearchState.NodeList.Emplace(StartPoint, /*InParent=*/INDEX_NONE, /*InCost=*/0.0, /*HeuristicCost=*/0.0);
-	OutSearchState.OpenIndexList.Reset();
+	OutSearchState.InternalSearchData->NodeList.Reset();
+	OutSearchState.InternalSearchData->NodeList.Emplace(&OutSearchState.InternalSearchData->InternalStartPointDescription, /*InParent=*/INDEX_NONE, /*InCost=*/0.0, /*HeuristicCost=*/0.0);
+	OutSearchState.InternalSearchData->OpenIndexList.Reset();
 	// Okay to add the first index this way, since it will be "heapified" by default with one element.
-	OutSearchState.OpenIndexList.Add(0);
-	OutSearchState.ClosedIndexList.Reset();
-	OutSearchState.PointToNodeIndexMap.Reset();
+	OutSearchState.InternalSearchData->OpenIndexList.Add(0);
+	OutSearchState.InternalSearchData->ClosedIndexList.Reset();
+	OutSearchState.InternalSearchData->PointDescriptionToNodeIndexMap.Reset();
+	
+	// Since we will use pointer to PontDescription, we should create them all in the array so their ptr is fixed.
+	int32 OrigNumPoint = OutSearchState.OriginatingPointData->GetNumPoints();
+	OutSearchState.InternalSearchData->InternalOriginatingPointsDescriptions.Empty();
+	OutSearchState.InternalSearchData->InternalOriginatingPointsDescriptions.Reserve(OrigNumPoint);
+	for (int32 Index = 0; Index < OrigNumPoint; ++Index)
+	{
+		OutSearchState.InternalSearchData->InternalOriginatingPointsDescriptions.Emplace(FPointDescription{{OutSearchState.OriginatingPointData, 0, OrigNumPoint}, Index});
+	}
+	
+	int32 GoalNumPoint = OutSearchState.GoalPoints.RangeSize;
+	OutSearchState.InternalSearchData->InternalGoalPointsDescriptions.Empty();
+	OutSearchState.InternalSearchData->InternalGoalPointsDescriptions.Reserve(GoalNumPoint);
+	for (int32 Index = 0; Index < GoalNumPoint; ++Index)
+	{
+		OutSearchState.InternalSearchData->InternalGoalPointsDescriptions.Emplace(FPointDescription{OutSearchState.GoalPoints, Index});
+	}
 }
 
 /** Runs a single iteration of the A* algorithm. Intended to be called multiple times, whereupon it will return true when the algorithm is finished. */
-ESearchResult ExecuteSearchIteration(const FSearchSettings& SearchSettings, FSearchState& SearchState, const TConstArrayView<FPCGPoint> GoalPoints, TArray<FPCGPoint>& OutPath)
+ESearchResult ExecuteSearchIteration(const FSearchSettings& SearchSettings, FSearchState& SearchState, TArray<FPointDescription>& OutPath)
 {
 	TRACE_CPUPROFILER_EVENT_SCOPE(FPCGPathfindingElement::Algorithm::ExecuteSearchIteration);
 	check(SearchState.OriginatingPointData);
 
 	// The goal points should be guaranteed by the caller. Must call Initialize first to create the open index list.
-	if (GoalPoints.IsEmpty() || !ensure(!SearchState.OpenIndexList.IsEmpty()))
+	if ((SearchState.GoalPoints.RangeSize==0) || !ensure(!SearchState.InternalSearchData->OpenIndexList.IsEmpty()))
 	{
 		return ESearchResult::Invalid;
 	}
 
-	TArray<FNode, TInlineAllocator<FSearchState::PreAllocNodeCount>>& NodeList = SearchState.NodeList;
+	FSearchState::FInternalSearchData::FNodeList& NodeList = SearchState.InternalSearchData->NodeList;
 
-	auto HeapPredicate = [&SearchState](int32 Index1, int32 Index2) { return Helpers::CompareNodes(SearchState.NodeList[Index1], SearchState.NodeList[Index2]); };
+	auto HeapPredicate = [&SearchState](int32 Index1, int32 Index2) { return Helpers::CompareNodes(SearchState.InternalSearchData->NodeList[Index1], SearchState.InternalSearchData->NodeList[Index2]); };
 
 	// Get the lowest cost point on the list, which has been binary heap sorted.
 	int32 CurrentNodeIndex;
-	SearchState.OpenIndexList.HeapPop(CurrentNodeIndex, HeapPredicate, EAllowShrinking::No);
-	check(NodeList[CurrentNodeIndex].PCGPoint);
+	SearchState.InternalSearchData->OpenIndexList.HeapPop(CurrentNodeIndex, HeapPredicate, EAllowShrinking::No);
 
-	TTuple<const FPCGPoint*, double> GoalInfo = Helpers::GetNearestGoalToLocation(NodeList[CurrentNodeIndex].PCGPoint->Transform.GetLocation(), GoalPoints);
-	const FPCGPoint* CurrentGoal = GoalInfo.Get<0>();
-	const double MinSquaredDistanceToGoal = GoalInfo.Get<1>();
+	TTuple<int32, double> GoalInfo = Helpers::GetNearestGoalToLocation(SearchState, *NodeList[CurrentNodeIndex].PointDescription);
+	const int32 CurrentGoalIndex = (GoalInfo.Get<0>());
+	const double MinDistanceToGoal = GoalInfo.Get<1>();
+	const FVector& CurrentGoal = SearchState.InternalSearchData->InternalGoalPointsDescriptions[CurrentGoalIndex].GetLocation();
 
-	auto PerPointProcessing = [&NodeList, CurrentNodeIndex, CurrentGoal, &HeapPredicate, &SearchSettings, &SearchState, &GoalPoints](const FPCGPoint* Point, const double DistanceToPointSquared)
+	auto PerPointProcessing = [&NodeList, CurrentNodeIndex, &CurrentGoal, &HeapPredicate, &SearchSettings, &SearchState](const FPointDescription* PointDescription, const double DistanceToPointSquared)
 	{
-		const bool bTrackedPoint = SearchState.PointToNodeIndexMap.Contains(Point);
+		const bool bTrackedPoint = SearchState.InternalSearchData->PointDescriptionToNodeIndexMap.Contains(PointDescription);
 		// Node has already been ruled out.
-		if (bTrackedPoint && SearchState.ClosedIndexList.Contains(SearchState.PointToNodeIndexMap[Point]))
+		if (bTrackedPoint && SearchState.InternalSearchData->ClosedIndexList.Contains(SearchState.InternalSearchData->PointDescriptionToNodeIndexMap[PointDescription]))
 		{
 			return;
 		}
 
-		const double TentativeNewLocalCost = SearchState.CostFunction(NodeList[CurrentNodeIndex].LocalCost, NodeList[CurrentNodeIndex].PCGPoint, DistanceToPointSquared, Point);
+		const double TentativeNewLocalCost = NodeList[CurrentNodeIndex].LocalCost + SearchState.CostFunction(*NodeList[CurrentNodeIndex].PointDescription, DistanceToPointSquared, *PointDescription);
 
 		// Blocked paths don't need expansion here, and serve no purpose at this point (at least from that direction) from being visited.
 		if (TentativeNewLocalCost == std::numeric_limits<double>::max())
@@ -140,73 +200,71 @@ ESearchResult ExecuteSearchIteration(const FSearchSettings& SearchSettings, FSea
 		 * can not be changed during the pathfinding process, as neighbor evaluation is tied to the heuristic to the
 		 * current goal at any point. Nodes may be ruled out and added to the Closed List
 		 */
-		auto GetHeuristicCost = [&SearchSettings, &SearchState, &GoalPoints](const FPCGPoint* CurrentPoint, const FPCGPoint* GoalPoint) -> double
+		auto GetHeuristicCost = [&SearchSettings, &SearchState](const FVector& CurrentPoint, const FVector& GoalPoint) -> double
 		{
-			return (GoalPoints.Num() > 1) ? 0.0 : SearchSettings.HeuristicWeight * SearchState.HeuristicFunction(CurrentPoint->Transform.GetLocation(), GoalPoint->Transform.GetLocation());
+			return (SearchState.GoalPoints.RangeSize > 1) ? 0.0 : SearchSettings.HeuristicWeight * SearchState.HeuristicFunction(CurrentPoint, GoalPoint);
 		};
 
 		// Not tracking this point yet--add it to the node list and map it.
 		if (!bTrackedPoint)
 		{
-			const double HeuristicCost = GetHeuristicCost(Point, CurrentGoal);
-			NodeList.Emplace(Point, CurrentNodeIndex, TentativeNewLocalCost, TentativeNewLocalCost + HeuristicCost);
+			const double HeuristicCost = GetHeuristicCost(PointDescription->GetLocation(), CurrentGoal); 
+			NodeList.Emplace(PointDescription, CurrentNodeIndex, TentativeNewLocalCost, TentativeNewLocalCost + HeuristicCost);
+
 			const int32 NewNodeIndex = NodeList.Num() - 1;
-			SearchState.OpenIndexList.HeapPush(NewNodeIndex, HeapPredicate);
-			SearchState.PointToNodeIndexMap.Emplace(Point, NewNodeIndex);
+			SearchState.InternalSearchData->OpenIndexList.HeapPush(NewNodeIndex, HeapPredicate);
+			SearchState.InternalSearchData->PointDescriptionToNodeIndexMap.Emplace(PointDescription, NewNodeIndex);
 
 			return;
 		}
 
 		// Check if the path to this neighbor is a better path than its current one. If so, update accordingly.
-		const int32 NeighborIndex = SearchState.PointToNodeIndexMap[Point];
-		FNode& Neighbor = SearchState.NodeList[NeighborIndex];
+		const int32 NeighborIndex = SearchState.InternalSearchData->PointDescriptionToNodeIndexMap[PointDescription];
+		FNode& Neighbor = SearchState.InternalSearchData->NodeList[NeighborIndex];
 		if (TentativeNewLocalCost < Neighbor.LocalCost)
 		{
-			const int32 CurrentOpenListIndex = SearchState.OpenIndexList.Find(NeighborIndex);
+			const int32 CurrentOpenListIndex = SearchState.InternalSearchData->OpenIndexList.Find(NeighborIndex);
 			check(CurrentOpenListIndex != INDEX_NONE);
 
 			// The priority may have changed, so remove the index and add it again.
-			SearchState.OpenIndexList.HeapRemoveAt(CurrentOpenListIndex, HeapPredicate, EAllowShrinking::No);
+			SearchState.InternalSearchData->OpenIndexList.HeapRemoveAt(CurrentOpenListIndex, HeapPredicate, EAllowShrinking::No);
 
-			const double HeuristicCost = GetHeuristicCost(Point, CurrentGoal);
+			const double HeuristicCost = GetHeuristicCost(PointDescription->GetLocation(), CurrentGoal);
 			Neighbor.PreviousNodeIndex = CurrentNodeIndex;
 			Neighbor.LocalCost = TentativeNewLocalCost;
 			Neighbor.EstimatedGoalCost = TentativeNewLocalCost + HeuristicCost;
 
-			SearchState.OpenIndexList.HeapPush(NeighborIndex, HeapPredicate);
+			SearchState.InternalSearchData->OpenIndexList.HeapPush(NeighborIndex, HeapPredicate);
 		}
 	};
 
 	// Arrived at the goal.
-	if (FMath::IsNearlyZero(MinSquaredDistanceToGoal))
+	if (FMath::IsNearlyZero(MinDistanceToGoal))
 	{
-		Helpers::BuildFinalPath(SearchState, SearchState.NodeList[CurrentNodeIndex], SearchSettings.bCopyOriginatingPoints, OutPath);
+		Helpers::BuildFinalPath(SearchState, SearchState.InternalSearchData->NodeList[CurrentNodeIndex], SearchSettings.bCopyOriginatingPoints, OutPath);
 		return ESearchResult::Complete;
 	}
 	// Close enough to the goal we need to add it as a possibility.
-	else if (MinSquaredDistanceToGoal <= SearchSettings.SearchDistance * SearchSettings.SearchDistance)
+	else if (MinDistanceToGoal <= SearchSettings.SearchDistance)
 	{
-		PerPointProcessing(CurrentGoal, MinSquaredDistanceToGoal);
+		PerPointProcessing(&SearchState.InternalSearchData->InternalGoalPointsDescriptions[CurrentGoalIndex], MinDistanceToGoal*MinDistanceToGoal);
 	}
 
-	const TArray<FPCGPoint>& Points = SearchState.OriginatingPointData->GetPoints();
-
 	// Gather neighbors within the search radius.
 	UPCGOctreeQueries::ForEachPointInsideSphere(
 		SearchState.OriginatingPointData,
-		NodeList[CurrentNodeIndex].PCGPoint->Transform.GetLocation(),
+		NodeList[CurrentNodeIndex].GetLocation(),
  		SearchSettings.SearchDistance,
-		[&Points, &PerPointProcessing](const UPCGBasePointData* BasePointData, int32 PointIndex, const double DistanceToPointSquared)
+		[&PerPointProcessing, &OriginatingPointsDescriptions=SearchState.InternalSearchData->InternalOriginatingPointsDescriptions](const UPCGBasePointData* BasePointData, int32 PointIndex, const double DistanceToPointSquared)
 		{
-			const FPCGPoint* PCGPoint = &Points[PointIndex];
-			PerPointProcessing(PCGPoint, DistanceToPointSquared);
+			PerPointProcessing(&(OriginatingPointsDescriptions[PointIndex]), DistanceToPointSquared);
 		});
 
 	// This node has been completely evaluated.
-	SearchState.ClosedIndexList.Add(CurrentNodeIndex);
+	SearchState.InternalSearchData->ClosedIndexList.Add(CurrentNodeIndex);
 
 	// Final point to check, no path can be found.
-	if (SearchState.OpenIndexList.IsEmpty())
+	if (SearchState.InternalSearchData->OpenIndexList.IsEmpty())
 	{
 		OutPath.Reset();
 
@@ -214,20 +272,19 @@ ESearchResult ExecuteSearchIteration(const FSearchSettings& SearchSettings, FSea
 		{
 			// Find the closest point to any goal and walk back the path from there.
 			int32 ClosestNodeIndex = INDEX_NONE;
-			double ClosestSquaredDistance = std::numeric_limits<double>::max();
-			for (const int32 NodeIndex : SearchState.ClosedIndexList)
+			double ClosestDistance = std::numeric_limits<double>::max();
+			for (const int32 NodeIndex : SearchState.InternalSearchData->ClosedIndexList)
 			{
-				TTuple<const FPCGPoint*, double> ClosestGoalInfo = Helpers::GetNearestGoalToLocation(SearchState.NodeList[NodeIndex].PCGPoint->Transform.GetLocation(), GoalPoints);
-				const double SquaredDistance = ClosestGoalInfo.Get<1>();
-				if (SquaredDistance < ClosestSquaredDistance)
+				TTuple<int32, double> ClosestGoalInfo = Helpers::GetNearestGoalToLocation(SearchState, *SearchState.InternalSearchData->NodeList[NodeIndex].PointDescription);
+				const double CurGoalDistance = ClosestGoalInfo.Get<1>();
+				if (CurGoalDistance < ClosestDistance)
 				{
-					ClosestSquaredDistance = SquaredDistance;
+					ClosestDistance = CurGoalDistance;
 					ClosestNodeIndex = NodeIndex;
 				}
 			}
 
-			Helpers::BuildFinalPath(SearchState, SearchState.NodeList[ClosestNodeIndex], SearchSettings.bCopyOriginatingPoints, OutPath);
-
+			Helpers::BuildFinalPath(SearchState, SearchState.InternalSearchData->NodeList[ClosestNodeIndex], SearchSettings.bCopyOriginatingPoints, OutPath);
 			return ESearchResult::Partial;
 		}
 		else
diff --git a/Engine/Plugins/PCG/Source/PCG/Public/Elements/PCGPathfindingElement.h b/Engine/Plugins/PCG/Source/PCG/Public/Elements/PCGPathfindingElement.h
--- a/Engine/Plugins/PCG/Source/PCG/Public/Elements/PCGPathfindingElement.h
+++ b/Engine/Plugins/PCG/Source/PCG/Public/Elements/PCGPathfindingElement.h
@@ -137,9 +137,10 @@ namespace PCGPathfindingElement
 	struct FExecutionState
 	{
 		/** The starting points of the paths to search. One search iteration will happen for each starting point. */
-		TArray<FPCGPoint> StartPoints;
+		FPCGPointInputRange StartPoints;
 		/** The goal points of the paths to search for. If more than one goal exists, the heuristic will be ignored. */
-		TArray<FPCGPoint> GoalPoints;
+		FPCGPointInputRange GoalPoints;
+
 		/** The total number of paths to find. */
 		int32 IterationCount = 0;
 		PCGSpatialAlgo::AStar::FSearchSettings Settings;
@@ -151,6 +152,12 @@ namespace PCGPathfindingElement
 		int32 PathIterationIndex = 0;
 		/** To compare against PathIterationIndex to know if the iteration has changed, for re-initialization of the search. */
 		int32 LastPathIterationIndex = INDEX_NONE;
+
+		/** The starting points for the current iteration. needs to stay alive until the end of iteration */
+		FPCGPointInputRange StartPoints;
+		/** The goal points for the current iteration. needs to stay alive until the end of iteration */
+		FPCGPointInputRange GoalPoints;
+
 		PCGSpatialAlgo::AStar::FSearchState SearchState;
 	};
 }
@@ -163,4 +170,7 @@ public:
 protected:
 	virtual bool PrepareDataInternal(FPCGContext* InContext) const override;
 	virtual bool ExecuteInternal(FPCGContext* InContext) const override;
+
+	virtual bool SupportsBasePointDataInputs(FPCGContext* InContext) const { return true; }
+
 };
diff --git a/Engine/Plugins/PCG/Source/PCG/Public/PCGPointPropertiesTraits.h b/Engine/Plugins/PCG/Source/PCG/Public/PCGPointPropertiesTraits.h
--- a/Engine/Plugins/PCG/Source/PCG/Public/PCGPointPropertiesTraits.h
+++ b/Engine/Plugins/PCG/Source/PCG/Public/PCGPointPropertiesTraits.h
@@ -20,6 +20,7 @@ enum class EPCGPointNativeProperties : uint32
 	Seed = 1 << 6,
 	MetadataEntry = 1 << 7,
 	All = Transform | Density | BoundsMin | BoundsMax | Color | Steepness | Seed | MetadataEntry,
+	AllProperties = All & (~MetadataEntry),
 };
 ENUM_CLASS_FLAGS(EPCGPointNativeProperties);
 
diff --git a/Engine/Plugins/PCG/Source/PCG/Public/SpatialAlgo/PCGAStar.h b/Engine/Plugins/PCG/Source/PCG/Public/SpatialAlgo/PCGAStar.h
--- a/Engine/Plugins/PCG/Source/PCG/Public/SpatialAlgo/PCGAStar.h
+++ b/Engine/Plugins/PCG/Source/PCG/Public/SpatialAlgo/PCGAStar.h
@@ -2,34 +2,33 @@
 
 #pragma once
 
-#include "PCGPoint.h"
+#include "Data/PCGBasePointData.h"
 #include "Math/Vector.h"
 
-class UPCGPointData;
-
 namespace PCGSpatialAlgo::AStar
 {
-struct FNode
+struct FPointDescription
 {
-	explicit FNode(const FPCGPoint* InPCGPoint, const int32 InPreviousNodeIndex, const double InLocalCost, const double InEstimatedGoalCost)
-		: PCGPoint(InPCGPoint)
-		, PreviousNodeIndex(InPreviousNodeIndex)
-		, LocalCost(InLocalCost)
-		, EstimatedGoalCost(InEstimatedGoalCost)
-	{}
+	FPCGPointInputRange SrcPointRange;
+	int32 PointIndexInRange = 0;
+	FVector CachedLocation;
+	
+	FPointDescription() = default;
+	FPointDescription(const FPCGPointInputRange& InSrcPointRange, int32 InPointIndexInRange):
+		SrcPointRange(InSrcPointRange),
+		PointIndexInRange(InPointIndexInRange),
+		CachedLocation(SrcPointRange.PointData->GetTransform(SrcPointRange.RangeStartIndex + PointIndexInRange).GetLocation()) {}
+	
 
-	const FPCGPoint* PCGPoint = nullptr;
-	/** The currently assigned parent node with the least cost. */
-	int32 PreviousNodeIndex = INDEX_NONE;
-	/** Concrete cost from start to this node. */
-	double LocalCost = 0.0;
-	/** Current total cost when heuristic is added to the local cost. */
-	double EstimatedGoalCost = 0.0;
+	const FVector& GetLocation() const
+	{
+		return CachedLocation;
+	}
 };
-
+	
 namespace Cost
 {
-	double CalculateCost_EuclideanDistance(const double PreviousNodeCost, const FPCGPoint* PreviousNodePoint, const double DistanceToPreviousNodeSquared, const FPCGPoint* CurrentPoint);
+	double CalculateCost_EuclideanDistance(const FPointDescription& PreviousNodePoint, const double DistanceToPreviousNodeSquared, const FPointDescription& CurrentPoint);
 }
 
 namespace Heuristic
@@ -51,30 +50,23 @@ struct FSearchSettings
 
 struct FSearchState
 {
-	static uint32 constexpr PreAllocNodeCount = 1024u;
-	/** Keep track of the point data for the octree search. */
-	const UPCGPointData* OriginatingPointData = nullptr;
+	FSearchState();
+	~FSearchState();
+	/** Keep track of the point data for the search. */
+	const UPCGBasePointData* OriginatingPointData = nullptr;
+	
+	/** Keep the goal point range to index them with FPointDescription */
+	FPCGPointInputRange GoalPoints{};
 	/** Cost function */
-	TFunction<double(const double /*PreviousNodeCost*/, const FPCGPoint* /*PreviousNodePoint*/, const double /*DistanceToPreviousNodeSquared*/, const FPCGPoint* /*CurrentPoint*/)> CostFunction = Cost::CalculateCost_EuclideanDistance;
+	TFunction<double(const FPointDescription& /*PreviousNodePoint*/, const double /*DistanceToPreviousNodeSquared*/, const FPointDescription& /*CurrentPoint*/)> CostFunction = Cost::CalculateCost_EuclideanDistance;
 	/** Heuristic function */
 	TFunction<double(const FVector& /*CurrentLocation*/, const FVector& /*GoalLocation*/)> HeuristicFunction = Heuristic::CalculateHeuristic_EuclideanDistance;
 
-	/** Storage of the current known node. Source of truth for the nodes, so it should not be resized after its initial reserve. */
-	TArray<FNode, TInlineAllocator<PreAllocNodeCount>> NodeList;
-	/** Used as a priority list of nodes, sorted by cost. */
-	TArray<int32, TInlineAllocator<PreAllocNodeCount>> OpenIndexList;
-	/** Nodes that are no longer viable. */
-	TArray<int32, TInlineAllocator<PreAllocNodeCount>> ClosedIndexList;
-	/** To track nodes to their corresponding points. */
-	TMap<const FPCGPoint*, int32> PointToNodeIndexMap;
+	struct FInternalSearchData;
+	TUniquePtr<FInternalSearchData> InternalSearchData;
 };
 
-namespace Helpers
-{
-	bool CompareNodes(const FNode& Node1, const FNode& Node2);
-	void BuildFinalPath(const FSearchState& SearchState, const FNode& FinalNode, bool bCopyOriginatingPoints, TArray<FPCGPoint>& OutPath);
-}
-
+	
 enum class ESearchResult
 {
 	Invalid = 0,
@@ -84,8 +76,8 @@ enum class ESearchResult
 };
 
 /** Initialize the search state for AStar. Must be called before ExecuteSearchIteration. */
-void Initialize(const FPCGPoint* const StartPoint, FSearchState& OutSearchState);
+void Initialize(const FPCGPointInputRange& StartPoint, const FPCGPointInputRange& GoalPoints, FSearchState& OutSearchState);
 
 /** Runs a single iteration of the A* algorithm. Intended to be called multiple times, whereupon it will return true when the algorithm is finished. */
-ESearchResult ExecuteSearchIteration(const FSearchSettings& SearchSettings, FSearchState& SearchState, TConstArrayView<FPCGPoint> GoalPoints, TArray<FPCGPoint>& OutPath);
+ESearchResult ExecuteSearchIteration(const FSearchSettings& SearchSettings, FSearchState& SearchState, TArray<FPointDescription>& OutPath);
 }
diff --git a/Engine/Plugins/PCG/Source/PCG/Public/SpatialAlgo/PCGOctreeQueries.h b/Engine/Plugins/PCG/Source/PCG/Public/SpatialAlgo/PCGOctreeQueries.h
--- a/Engine/Plugins/PCG/Source/PCG/Public/SpatialAlgo/PCGOctreeQueries.h
+++ b/Engine/Plugins/PCG/Source/PCG/Public/SpatialAlgo/PCGOctreeQueries.h
@@ -84,15 +84,16 @@ void UPCGOctreeQueries::ForEachPointInsideSphere(const UPCGBasePointData* InPoin
 	// To support previous callback
 	const UPCGPointData* PointData = Cast<UPCGPointData>(InPointData);
 	const TArray<FPCGPoint>* Points = PointData ? &PointData->GetPoints() : nullptr;
+	const FPCGPointTransform::ConstValueRange TransformValueRange = InPointData->GetConstTransformValueRange();
 		
-	InPointData->GetPointOctree().FindElementsWithBoundsTest(SearchBounds, [InPointData, Points, &InCenter, SquaredRadius, Callback = std::move(Callback)](const PCGPointOctree::FPointRef& PointRef)
+	InPointData->GetPointOctree().FindElementsWithBoundsTest(SearchBounds, [InPointData, Points, &TransformValueRange, &InCenter, SquaredRadius, Callback = std::move(Callback)](const PCGPointOctree::FPointRef& PointRef)
 	{
 		if (!InPointData->IsValidRef(PointRef))
 		{
 			return;
 		}
 
-		const double SquaredDistance = FVector::DistSquared(InCenter, Points ? (*Points)[PointRef.Index].Transform.GetLocation() : InPointData->GetTransform(PointRef.Index).GetLocation());
+		const double SquaredDistance = FVector::DistSquared(InCenter, TransformValueRange[PointRef.Index].GetLocation());
 		if (SquaredDistance <= SquaredRadius)
 		{
 			if constexpr (std::is_invocable_v<Func, const FPCGPoint&, double>)
diff --git a/Engine/Plugins/PCG/Tests/Private/Elements/PCGPathfindingTests.cpp b/Engine/Plugins/PCG/Tests/Private/Elements/PCGPathfindingTests.cpp
new file mode 100644
--- /dev/null
+++ b/Engine/Plugins/PCG/Tests/Private/Elements/PCGPathfindingTests.cpp
@@ -0,0 +1,254 @@
+﻿// Copyright Epic Games, Inc. All Rights Reserved.
+
+#include "CoreMinimal.h"
+#include "TestHarness.h"
+
+#include <catch2/generators/catch_generators.hpp>
+
+#include "PCGContext.h"
+#include "PCGModule.h"
+#include "Data/DataView/PCGDataViewInterface.h"
+#include "Data/PCGPointData.h"
+#include "Elements/PCGPathfindingElement.h"
+#include "Tests/PCGTestsCommon.h"
+#include "PCGTestsCommon.h"
+
+class FPCGPathfindingBaseTest: public FPCGBaseTest
+{
+public:
+	FPCGPathfindingBaseTest(): 
+		FPCGBaseTest(),
+		BaseStartPos(FRotator{}, FVector{ 0, 0, 0 }, FVector{}),
+		BaseGoalPos(FRotator{}, FVector{ 10, 10, 0 }, FVector{})
+	{
+		TestData.Reset(new PCGTestsCommon::FTestData());
+
+		PCGTestsCommon::GenerateSettings<UPCGPathfindingSettings>(*TestData);
+		Settings = CastChecked<UPCGPathfindingSettings>(TestData->Settings);
+		Settings->bOutputAsSpline = false;
+		Settings->bAcceptPartialPath = true;
+		PathfindingElement = TestData->Settings->GetElement();
+
+		TArray<FVector> Points;
+
+		Points.Add({ 0, 0, 0 });
+		Points.Add({ 10, 10, 0 });
+
+		//Short path
+		Points.Add({ 2.5, 2.5, 0 });
+		Points.Add({ 5, 5, 0 });
+		MiddleNodeIndex = Points.Num()-1;
+		Points.Add({ 7.5, 7.5, 0 });
+
+		//Longer path
+		Points.Add({ 1, 0, 0 });
+		Points.Add({ 2, 0, 0 });
+		Points.Add({ 3, 0, 0 });
+		Points.Add({ 4, 0, 0 });
+		Points.Add({ 5, 0, 0 });
+		Points.Add({ 6, 0, 0 });
+		Points.Add({ 7, 0, 0 });
+		Points.Add({ 8, 0, 0 });
+		Points.Add({ 9, 0, 0 });
+		Points.Add({ 10, 0, 0 });
+		Points.Add({ 10, 1, 0 });
+		Points.Add({ 10, 2, 0 });
+		Points.Add({ 10, 3, 0 });
+		Points.Add({ 10, 4, 0 });
+		Points.Add({ 10, 5, 0 });
+		Points.Add({ 10, 6, 0 });
+		Points.Add({ 10, 7, 0 });
+		Points.Add({ 10, 8, 0 });
+		Points.Add({ 10, 9, 0 });
+
+		SimpleData = PCGTestsCommon::CreateEmptyBasePointData();
+		SimpleData->SetNumPoints(Points.Num());
+		SimpleData->SetDensity(1);
+		for (int i = 0; i < Points.Num(); i++)
+		{
+			SimpleData->GetTransformValueRange()[i] = FTransform(FRotator{}, Points[i], FVector{});
+		}
+	}
+
+	int GetPathfindingNumPoint()
+	{
+		Context = TestData->InitializeTestContext();
+
+		while (!PathfindingElement->Execute(Context.Get()))
+		{
+		}
+
+		const TArray<FPCGTaggedData>& Inputs = Context->InputData.GetAllSpatialInputs();
+		const TArray<FPCGTaggedData>& Outputs = Context->OutputData.GetAllSpatialInputs();
+
+		REQUIRE_EQUAL(Outputs.Num(), 1);
+
+		const FPCGTaggedData& Output = Outputs[0];
+
+		const UPCGBasePointData* OutData = Cast<UPCGBasePointData>(Output.Data);
+		REQUIRE(OutData != nullptr);
+
+		return OutData->GetNumPoints();
+	};
+
+
+	int32 MiddleNodeIndex;
+
+	TUniquePtr<PCGTestsCommon::FTestData> TestData;
+	TUniquePtr<FPCGContext> Context;
+	UPCGPathfindingSettings* Settings;
+	FPCGElementPtr PathfindingElement;
+	FTransform BaseStartPos;
+	FTransform BaseGoalPos;
+
+	TObjectPtr<UPCGBasePointData> SimpleData;
+};
+
+TEST_CASE_METHOD(FPCGPathfindingBaseTest, "PCG::PathfindingElement", "[PCG][PathfindingElement]")
+{
+	auto [searchDistance, expectedPathLength] = GENERATE( 
+			table<float, int>({
+				{1000, 2},
+				{7.1, 3},
+				{3.6, 5},
+				{1.1, 21}
+			})
+		);
+	
+	DYNAMIC_SECTION("BaseTest with search distance = " << searchDistance)
+	{
+		FPCGTaggedData& SimpleTaggedData = TestData->InputData.TaggedData.Emplace_GetRef(FPCGTaggedData());
+		SimpleTaggedData.Data = SimpleData;
+		SimpleTaggedData.Pin = PCGPinConstants::DefaultInputLabel;
+
+		Settings->Start = BaseStartPos.GetLocation();
+		Settings->Goal = BaseGoalPos.GetLocation();
+
+		Settings->SearchDistance = searchDistance;
+		CHECK_EQUAL(GetPathfindingNumPoint(), expectedPathLength);
+	}
+}
+
+TEST_CASE_METHOD(FPCGPathfindingBaseTest, "PCG::PathfindingElement:WithCosts", "[PCG][PathfindingElement]")
+{
+	
+	auto [WithCostFunction, expectedPathLength] = GENERATE( 
+			table<bool, int>({
+				{false, 5},
+				{true, 9}
+			})
+		);
+	
+	DYNAMIC_SECTION("With Cost function = " << (WithCostFunction ? "true" : "false"))
+	{
+		FPCGTaggedData& SimpleTaggedData = TestData->InputData.TaggedData.Emplace_GetRef(FPCGTaggedData());
+		SimpleTaggedData.Data = SimpleData;
+		SimpleTaggedData.Pin = PCGPinConstants::DefaultInputLabel;
+
+		Settings->Start = BaseStartPos.GetLocation();
+		Settings->Goal = BaseGoalPos.GetLocation();
+		
+		if (WithCostFunction)
+		{
+			const FString CostAttributeName = "Cost";
+			// Add an absurd cost to the middle point of the fast path.
+			FPCGMetadataAttribute<float>* Attribute = SimpleData->MutableMetadata()->CreateAttribute<float>(CostAttributeName, 1.0, /*bAllowsInterpolation=*/true, /*bOverrideParent=*/false);
+			FPCGMetadataDomain* OutputMetadata = SimpleData->MutableMetadata()->GetMetadataDomain(FPCGMetadataDomainID{});
+			check(OutputMetadata);
+			PCGMetadataEntryKey Key = OutputMetadata->AddEntry();
+			Attribute->SetValue(Key, 1000000.0);
+			SimpleData->GetMetadataEntryValueRange()[MiddleNodeIndex] = Key;
+
+			Settings->CostFunctionMode = EPCGPathfindingCostFunctionMode::CostMultiplier;
+			Settings->CostAttribute =  FPCGAttributePropertySelector::CreateSelectorFromString<FPCGAttributePropertyInputSelector>(CostAttributeName);			
+		}
+
+		Settings->SearchDistance = 3.6;
+		CHECK_EQUAL(GetPathfindingNumPoint(), expectedPathLength);
+	}
+}
+
+TEST_CASE_METHOD(FPCGPathfindingBaseTest, "PCG::PathfindingElement:WithStartPosAsInput", "[PCG][PathfindingElement]")
+{
+	FPCGTaggedData& SimpleTaggedData = TestData->InputData.TaggedData.Emplace_GetRef(FPCGTaggedData());
+	SimpleTaggedData.Data = SimpleData;
+	SimpleTaggedData.Pin = PCGPinConstants::DefaultInputLabel;
+
+	UPCGBasePointData* StartPointData = PCGTestsCommon::CreateEmptyBasePointData();
+	StartPointData->SetNumPoints(1);
+	StartPointData->SetDensity(1);
+	StartPointData->GetTransformValueRange()[0] = BaseStartPos;
+	
+	FPCGTaggedData& StartPointTaggedData = TestData->InputData.TaggedData.Emplace_GetRef(FPCGTaggedData());
+	StartPointTaggedData.Data = StartPointData;
+	StartPointTaggedData.Pin = "Start";
+	Settings->bStartLocationsAsInput = true;
+
+	// Make sure the start pos is not used, so set it to an invalid value.
+	Settings->Start = FVector{ 1000, 1000, 1000 };
+	Settings->Goal = BaseGoalPos.GetLocation();	
+	
+	Settings->SearchDistance = 3.6;
+	CHECK_EQUAL(GetPathfindingNumPoint(), 5);
+}
+
+TEST_CASE_METHOD(FPCGPathfindingBaseTest, "PCG::PathfindingElement:WithStartPosAsAttribute", "[PCG][PathfindingElement]")
+{
+	auto WithCopyFromOriginatingPoints = GENERATE(false, true); 
+	
+	DYNAMIC_SECTION("With Copy From Orig = " << (WithCopyFromOriginatingPoints ? "true" : "false"))
+	{
+		FPCGTaggedData& SimpleTaggedData = TestData->InputData.TaggedData.Emplace_GetRef(FPCGTaggedData());
+		SimpleTaggedData.Data = SimpleData;
+		SimpleTaggedData.Pin = PCGPinConstants::DefaultInputLabel;
+
+		const FString PosAttributeName = "ThePosition";
+
+		UPCGBasePointData* StartPointData = PCGTestsCommon::CreateEmptyBasePointData();
+		StartPointData->SetNumPoints(1);
+		StartPointData->SetDensity(1);
+		StartPointData->GetTransformValueRange()[0] = FTransform(FRotator{}, FVector{ 1000, 1000, 1000 }, FVector{});
+
+		FPCGTaggedData& StartPointTaggedData = TestData->InputData.TaggedData.Emplace_GetRef(FPCGTaggedData());
+		StartPointTaggedData.Data = StartPointData;
+		StartPointTaggedData.Pin = "Start";
+		FPCGMetadataAttribute<FVector>* Attribute = StartPointData->MutableMetadata()->CreateAttribute<FVector>(PosAttributeName, FVector(), /*bAllowsInterpolation=*/true, /*bOverrideParent=*/false);
+		FPCGMetadataDomain* OutputMetadata = StartPointData->MutableMetadata()->GetMetadataDomain(FPCGMetadataDomainID{});
+		check(OutputMetadata);
+		PCGMetadataEntryKey Key = OutputMetadata->AddEntry();
+		Attribute->SetValue(Key, BaseStartPos.GetLocation());
+		StartPointData->GetMetadataEntryValueRange()[0] = Key;
+
+		Settings->bStartLocationsAsInput = true;
+		Settings->StartLocationAttribute = FPCGAttributePropertySelector::CreateSelectorFromString<FPCGAttributePropertyInputSelector>(PosAttributeName);
+
+		// Make sure the start pos is not used, so set it to an invalid value.
+		Settings->Start = FVector{ 1000, 1000, 1000 };
+		Settings->Goal = BaseGoalPos.GetLocation();
+		
+		Settings->SearchDistance = 3.6;
+		Settings->bCopyOriginatingPoints = WithCopyFromOriginatingPoints;
+		
+		CHECK_EQUAL(GetPathfindingNumPoint(), 5);
+		
+		const TArray<FPCGTaggedData>& Outputs = Context->OutputData.GetAllSpatialInputs();
+		const FPCGTaggedData& Output = Outputs[0];
+		const UPCGBasePointData* OutData = Cast<UPCGBasePointData>(Output.Data);
+		REQUIRE(OutData != nullptr);
+		
+		const FPCGMetadataAttribute<FVector>* OutputAttribute = OutData->ConstMetadata()->GetConstTypedAttribute<FVector>(PosAttributeName);
+		if (WithCopyFromOriginatingPoints)
+		{
+			REQUIRE(OutputAttribute);
+			FVector StartLocation = OutData->GetTransform(0).GetLocation();
+			PCGMetadataEntryKey OutStartPointMetaDataEntry = OutData->GetMetadataEntry(0);
+			CHECK(OutStartPointMetaDataEntry == Key);
+		}
+		else
+		{
+			CHECK(OutputAttribute == nullptr);
+		}
+	}
+
+
+}
\ No newline at end of file
