[PCG] PropertyHelpers: treat  UVerseClass as UUserDefinedStruct when searching for properties

#rnx

#Synced-CL 49276028

diff --git a/Engine/Plugins/PCG/Source/PCG/Private/Helpers/PCGPropertyHelpers.cpp b/Engine/Plugins/PCG/Source/PCG/Private/Helpers/PCGPropertyHelpers.cpp
--- a/Engine/Plugins/PCG/Source/PCG/Private/Helpers/PCGPropertyHelpers.cpp
+++ b/Engine/Plugins/PCG/Source/PCG/Private/Helpers/PCGPropertyHelpers.cpp
@@ -14,6 +14,7 @@
 #include "StructUtils/UserDefinedStruct.h"
 #include "UObject/Field.h"
 #include "UObject/UnrealType.h"
+#include "VerseVM/VVMVerseClass.h"
 
 #define LOCTEXT_NAMESPACE "PCGPropertyHelpers"
 
@@ -22,6 +23,8 @@ namespace PCGPropertyHelpers
 	static constexpr uint64 ExcludePropertyFlags = CPF_DisableEditOnInstance;
 	static constexpr uint64 IncludePropertyFlags = CPF_BlueprintVisible;
 
+	static constexpr uint64 VerseIncludePropertyFlags = CPF_BlueprintReadOnly | CPF_Edit;
+
 	/**
 	* Expands container locations to their contents when the property passed in is an array or a set.
 	* This is useful to allow extraction downstream of properties inside of arrays/sets and also to generate the list of addresses/values to look at
@@ -79,7 +82,23 @@ namespace PCGPropertyHelpers
 		}
 
 		// Make sure the property is visible, if requested
-		if (bNeedsToBeVisible && (Property->HasAnyPropertyFlags(ExcludePropertyFlags) || !Property->HasAnyPropertyFlags(IncludePropertyFlags)))
+		bool bVisible = true;
+		if (bNeedsToBeVisible)
+		{
+			if (CurrentClass->IsA<UVerseClass>())
+			{
+				if (!Property->HasAnyPropertyFlags(VerseIncludePropertyFlags))
+				{
+					bVisible = false;
+				}
+			}
+			else if (Property->HasAnyPropertyFlags(ExcludePropertyFlags) || !Property->HasAnyPropertyFlags(IncludePropertyFlags))
+			{
+				bVisible = false;
+			}
+		}
+
+		if (!bVisible)
 		{
 			if (!bQuiet)
 			{
@@ -567,9 +586,9 @@ const FProperty* PCGPropertyHelpers::FindPropertyByName(const UStruct* InStruct,
 	{
 		return nullptr;
 	}
-	else if (const UUserDefinedStruct* UserStruct = Cast<UUserDefinedStruct>(InStruct))
+	else if (InStruct->IsA<UUserDefinedStruct>() || InStruct->IsA<UVerseClass>())
 	{
-		return FindPropertyInUserDefinedStruct(UserStruct, InName); 
+		return FindPropertyByNameEx(InStruct, InName);
 	}
 	else
 	{
@@ -577,7 +596,7 @@ const FProperty* PCGPropertyHelpers::FindPropertyByName(const UStruct* InStruct,
 	}
 }
 
-const FProperty* PCGPropertyHelpers::FindPropertyInUserDefinedStruct(const UUserDefinedStruct* InStruct, const FName InName)
+const FProperty* PCGPropertyHelpers::FindPropertyByNameEx(const UStruct* InStruct, const FName InName)
 {
 	if (!InStruct)
 	{
diff --git a/Engine/Plugins/PCG/Source/PCG/Private/Metadata/Accessors/PCGAttributeAccessorHelpers.cpp b/Engine/Plugins/PCG/Source/PCG/Private/Metadata/Accessors/PCGAttributeAccessorHelpers.cpp
--- a/Engine/Plugins/PCG/Source/PCG/Private/Metadata/Accessors/PCGAttributeAccessorHelpers.cpp
+++ b/Engine/Plugins/PCG/Source/PCG/Private/Metadata/Accessors/PCGAttributeAccessorHelpers.cpp
@@ -19,6 +19,7 @@
 
 #include "StructUtils/UserDefinedStruct.h"
 #include "UObject/EnumProperty.h"
+#include "VerseVM/VVMVerseClass.h"
 
 namespace PCGAttributeAccessorHelpers
 {
@@ -104,15 +105,7 @@ namespace PCGAttributeAccessorHelpers
 			const FProperty* Property = nullptr;
 			const FName PropertyName = InPropertyNames[i];
 
-			// Try to get the property. If it is coming from a user struct, we need to have a special search because the property name might be mangled.
-			if (const UUserDefinedStruct* UserDefinedStruct = Cast<UUserDefinedStruct>(CurrentStruct))
-			{
-				Property = PCGPropertyHelpers::FindPropertyInUserDefinedStruct(UserDefinedStruct, PropertyName);
-			}
-			else
-			{
-				Property = FindFProperty<FProperty>(CurrentStruct, PropertyName);
-			}
+			Property = PCGPropertyHelpers::FindPropertyByName(CurrentStruct, PropertyName);
 
 			if (!Property)
 			{
diff --git a/Engine/Plugins/PCG/Source/PCG/Public/Helpers/PCGPropertyHelpers.h b/Engine/Plugins/PCG/Source/PCG/Public/Helpers/PCGPropertyHelpers.h
--- a/Engine/Plugins/PCG/Source/PCG/Public/Helpers/PCGPropertyHelpers.h
+++ b/Engine/Plugins/PCG/Source/PCG/Public/Helpers/PCGPropertyHelpers.h
@@ -121,10 +121,10 @@ namespace PCGPropertyHelpers
 	PCG_API const FProperty* FindPropertyByName(const UStruct* InStruct, const FName InName);
 
 	/**
-	* More permissive search for a property name in a UserDefinedStruct.
-	* Since the names in editor and at runtime can be different for the properties defined in a UserDefinedStruct, the search will also look for sanitized version of the name if it is not a valid name.
+	* More permissive search for a property name in a UserDefinedStruct/UVerseClass.
+	* Since the names in editor and at runtime can be different for the properties defined in a UserDefinedStruct/VerseClass, the search will also look for sanitized version of the name if it is not a valid name.
 	*/
-	const FProperty* FindPropertyInUserDefinedStruct(const UUserDefinedStruct* InStruct, const FName InName);
+	const FProperty* FindPropertyByNameEx(const UStruct* InStruct, const FName InName);
 
 	namespace Constants
 	{
