[PCG] Add Texture Array support.

Also flesh out the Texture data interface with other useful things:

* Transform for UV->world transformations
* NumMips to help accessing data in higher mips
* Load override with mip index

[RN] [PCG] Added support for texture 2D arrays
#Synced-CL 49338130

diff --git a/Engine/Plugins/PCG/Shaders/Private/PCGTexture2DArrayDataInterface.ush b/Engine/Plugins/PCG/Shaders/Private/PCGTexture2DArrayDataInterface.ush
new file mode 100644
--- /dev/null
+++ b/Engine/Plugins/PCG/Shaders/Private/PCGTexture2DArrayDataInterface.ush
@@ -0,0 +1,116 @@
+// Copyright Epic Games, Inc. All Rights Reserved.
+
+Texture2DArray {DataInterfaceName}_TextureSRV;
+RWTexture2DArray<float4> {DataInterfaceName}_TextureUAV;
+SamplerState {DataInterfaceName}_Sampler;
+uint2 {DataInterfaceName}_Resolution;
+uint {DataInterfaceName}_ArraySize;
+float4x4 {DataInterfaceName}_TextureTransform;
+float4x4 {DataInterfaceName}_TextureInverseTransform;
+float2 {DataInterfaceName}_TexelSizeWorld;
+uint {DataInterfaceName}_NumMips;
+
+// #################### GENERAL HELPERS ##########################
+
+uint GetNumData_{DataInterfaceName}()
+{
+	return 1u;
+}
+
+uint2 GetResolution_{DataInterfaceName}(uint InDataIndex)
+{
+	return {DataInterfaceName}__Resolution;
+}
+
+uint GetArraySize_{DataInterfaceName}(uint InDataIndex)
+{
+	return {DataInterfaceName}_ArraySize;
+}
+
+uint3 GetNumElements_{DataInterfaceName}(uint InDataIndex)
+{
+	return uint3({DataInterfaceName}_Resolution.xy, {DataInterfaceName}_ArraySize);
+}
+
+bool GetThreadData_{DataInterfaceName}(uint InThreadIndex, out uint OutDataIndex, out uint3 OutElementIndex)
+{
+	const uint TexelsPerSlice = {DataInterfaceName}_Resolution.x * {DataInterfaceName}_Resolution.y;
+	const uint SliceIndex = InThreadIndex / TexelsPerSlice;
+
+	if (SliceIndex < 0 || SliceIndex >= {DataInterfaceName}_ArraySize)
+	{
+		OutDataIndex = (uint)-1;
+		OutElementIndex = (uint3)-1;
+		return false;
+	}
+
+	// Linear texel index within the current slice.
+	const uint LinearTexelIndex = InThreadIndex - SliceIndex * TexelsPerSlice;
+
+	OutElementIndex.x = LinearTexelIndex % {DataInterfaceName}_Resolution.x;
+	OutElementIndex.y = LinearTexelIndex / {DataInterfaceName}_Resolution.x;
+	OutElementIndex.z = SliceIndex;
+
+	// This data interface is single-data.
+	OutDataIndex = 0;
+
+	return true;
+}
+
+// #################### INPUT FUNCTIONS ##########################
+
+float2 GetTexCoords_{DataInterfaceName}(uint InDataIndex, float3 InWorldPos)
+{
+	const float4 LocalPos = mul({DataInterfaceName}_TextureInverseTransform, float4(InWorldPos, 1.0));
+
+	// LocalPos will be in [-1, 1] so transform to a texture coordinate in [0, 1]
+	return LocalPos.xy * 0.5f + 0.5f;
+}
+
+float2 GetTexelSize_{DataInterfaceName}(uint InDataIndex)
+{
+	return 1.0f / float2(GetNumElements_{DataInterfaceName}(InDataIndex));
+}
+
+float2 GetTexelSizeWorld_{DataInterfaceName}(uint InDataIndex)
+{
+	return {DataInterfaceName}_TexelSizeWorld;
+}
+
+float4x4 GetTransform_{DataInterfaceName}(uint InDataIndex)
+{
+	return {DataInterfaceName}_TextureTransform;
+}
+
+uint GetNumMips_{DataInterfaceName}(uint InDataIndex)
+{
+	return {DataInterfaceName}_NumMips;
+}
+
+float4 Sample_{DataInterfaceName}(uint InDataIndex, float2 InTextureUV, uint InSliceIndex)
+{
+	return {DataInterfaceName}_TextureSRV.SampleLevel({DataInterfaceName}_Sampler, float3(InTextureUV, InSliceIndex), 0);
+}
+
+float4 SampleWorldPos_{DataInterfaceName}(uint InDataIndex, float3 WorldPos, uint InSliceIndex)
+{
+	const float2 TextureUV = GetTexCoords_{DataInterfaceName}(InDataIndex, WorldPos);
+	return {DataInterfaceName}_TextureSRV.SampleLevel({DataInterfaceName}_Sampler, float3(TextureUV, InSliceIndex), 0);
+}
+
+float4 Load_{DataInterfaceName}(uint InDataIndex, uint3 InElementIndex)
+{
+	return {DataInterfaceName}_TextureSRV[InElementIndex];
+}
+
+float4 Load_{DataInterfaceName}(uint InDataIndex, uint3 InElementIndex, uint InMipIndex)
+{
+	return {DataInterfaceName}_TextureSRV.Load(int4(InElementIndex, InMipIndex));
+}
+
+// #################### OUTPUT FUNCTIONS ##########################
+
+void Store_{DataInterfaceName}(uint InDataIndex, uint3 InElementIndex, float4 InValue)
+{
+	{DataInterfaceName}_TextureUAV[InElementIndex] = InValue;
+}
diff --git a/Engine/Plugins/PCG/Shaders/Private/PCGTextureDataInterface.ush b/Engine/Plugins/PCG/Shaders/Private/PCGTextureDataInterface.ush
--- a/Engine/Plugins/PCG/Shaders/Private/PCGTextureDataInterface.ush
+++ b/Engine/Plugins/PCG/Shaders/Private/PCGTextureDataInterface.ush
@@ -34,6 +34,7 @@ PCG_TEXTURE_UAV_DECLARATION(0);
 
 SamplerState {DataInterfaceName}_Sampler;
 
+float4x4 {DataInterfaceName}_TextureTransforms[PCG_MAX_NUM_SRV_BINDINGS];
 float4x4 {DataInterfaceName}_TextureInverseTransforms[PCG_MAX_NUM_SRV_BINDINGS];
 
 /** Layout of TextureInfos buffer:
@@ -41,7 +42,7 @@ float4x4 {DataInterfaceName}_TextureInverseTransforms[PCG_MAX_NUM_SRV_BINDINGS];
  *   - uint32 - BindingIndex
  *   - uint32 - SliceIndex (aka TextureIndex if bound to a TextureArray, 0 otherwise.)
  *   - uint32 - Dimension
- *   - uint32 - UNUSED
+ *   - uint32 - NumMips
  *   - uint32 - ResolutionX
  *   - uint32 - ResolutionY
  *   - float  - TexelSizeWorldX (cm)
@@ -55,7 +56,7 @@ uint {DataInterfaceName}_NumTextureInfos;
 
 // #################### INTERNAL HELPERS ##########################
 
-void GetTextureInfoInternal_{DataInterfaceName}(uint InDataIndex, out uint OutBindingIndex, out uint OutSliceIndex, out uint OutTextureDimension)
+void GetTextureInfoInternal_{DataInterfaceName}(uint InDataIndex, out uint OutBindingIndex, out uint OutSliceIndex, out uint OutTextureDimension, out uint OutNumMips)
 {
 	if (InDataIndex >= 0u && InDataIndex < {DataInterfaceName}_NumTextureInfos)
 	{
@@ -63,12 +64,14 @@ void GetTextureInfoInternal_{DataInterfaceName}(uint InDataIndex, out uint OutBi
 		OutBindingIndex = TextureInfo.x;
 		OutSliceIndex = TextureInfo.y;
 		OutTextureDimension = TextureInfo.z;
+		OutNumMips = TextureInfo.w;
 	}
 	else
 	{
 		OutBindingIndex = 0u;
 		OutSliceIndex = 0u;
 		OutTextureDimension = 0u;
+		OutNumMips = 1u;
 	}
 }
 
@@ -133,7 +136,7 @@ bool GetThreadData_{DataInterfaceName}(uint InThreadIndex, out uint OutDataIndex
 float2 GetTexCoords_{DataInterfaceName}(uint InDataIndex, float3 InWorldPos)
 {
 	const float4x4 InverseTransform = {DataInterfaceName}_TextureInverseTransforms[InDataIndex];
-	const float4 LocalPos = mul(float4(InWorldPos, 1.0), InverseTransform);
+	const float4 LocalPos = mul(InverseTransform, float4(InWorldPos, 1.0));
 
 	// LocalPos will be in [-1, 1] so transform to a texture coordinate in [0, 1]
 	return LocalPos.xy * 0.5 + 0.5;
@@ -152,6 +155,21 @@ float2 GetTexelSizeWorld_{DataInterfaceName}(uint InDataIndex)
 	return asfloat({DataInterfaceName}_TextureInfos[InDataIndex * 2 + 1].zw);
 }
 
+uint GetNumMips_{DataInterfaceName}(uint InDataIndex)
+{
+	uint BindingIndex = 0;
+	uint SliceIndex = 0;
+	uint TextureDimension = 0;
+	uint NumMips = 1;
+	GetTextureInfoInternal_{DataInterfaceName}(InDataIndex, BindingIndex, SliceIndex, TextureDimension, NumMips);
+	return NumMips;
+}
+
+float4x4 GetTransform_{DataInterfaceName}(uint InDataIndex)
+{
+	return {DataInterfaceName}_TextureTransforms[InDataIndex];
+}
+
 #define PCG_CONDITIONAL_TEXTURE_SAMPLE(Index) \
 	if (BindingIndex == Index) \
 	{ \
@@ -175,7 +193,8 @@ float4 Sample_{DataInterfaceName}(uint InDataIndex, float2 InTextureUV)
 	uint BindingIndex = 0;
 	uint SliceIndex = 0;
 	uint TextureDimension = 0;
-	GetTextureInfoInternal_{DataInterfaceName}(InDataIndex, BindingIndex, SliceIndex, TextureDimension);
+	uint NumMips = 1;
+	GetTextureInfoInternal_{DataInterfaceName}(InDataIndex, BindingIndex, SliceIndex, TextureDimension, NumMips);
 
 #if PCG_ENABLE_MULTIPLE_TEXTURE_OBJECTS
 	PCG_CONDITIONAL_TEXTURE_SAMPLE(0)
@@ -232,7 +251,8 @@ float4 SampleWorldPos_{DataInterfaceName}(uint InDataIndex, float3 WorldPos)
 	uint BindingIndex = 0;
 	uint SliceIndex = 0;
 	uint TextureDimension = 0;
-	GetTextureInfoInternal_{DataInterfaceName}(InDataIndex, BindingIndex, SliceIndex, TextureDimension);
+	uint NumMips = 1;
+	GetTextureInfoInternal_{DataInterfaceName}(InDataIndex, BindingIndex, SliceIndex, TextureDimension, NumMips);
 
 #if PCG_ENABLE_MULTIPLE_TEXTURE_OBJECTS
 	PCG_CONDITIONAL_TEXTURE_SAMPLE_WORLD_POS(0)
@@ -290,7 +310,8 @@ float4 Load_{DataInterfaceName}(uint InDataIndex, uint2 InElementIndex)
 	uint BindingIndex = 0;
 	uint SliceIndex = 0;
 	uint TextureDimension = 0;
-	GetTextureInfoInternal_{DataInterfaceName}(InDataIndex, BindingIndex, SliceIndex, TextureDimension);
+	uint NumMips = 1;
+	GetTextureInfoInternal_{DataInterfaceName}(InDataIndex, BindingIndex, SliceIndex, TextureDimension, NumMips);
 
 #if PCG_ENABLE_MULTIPLE_TEXTURE_OBJECTS
 	PCG_CONDITIONAL_TEXTURE_LOAD(0)
@@ -337,7 +358,8 @@ void Store_{DataInterfaceName}(uint InDataIndex, uint2 InElementIndex, float4 In
 	uint BindingIndex = 0;
 	uint SliceIndex = 0;
 	uint TextureDimension = 0;
-	GetTextureInfoInternal_{DataInterfaceName}(InDataIndex, BindingIndex, SliceIndex, TextureDimension);
+	uint NumMips = 1;
+	GetTextureInfoInternal_{DataInterfaceName}(InDataIndex, BindingIndex, SliceIndex, TextureDimension, NumMips);
 
 	if (TextureDimension != PCG_TextureDimension_Texture2D)
 	{
diff --git a/Engine/Plugins/PCG/Source/PCG/Private/Compute/DataInterfaces/PCGTexture2DArrayDataInterface.cpp b/Engine/Plugins/PCG/Source/PCG/Private/Compute/DataInterfaces/PCGTexture2DArrayDataInterface.cpp
new file mode 100644
--- /dev/null
+++ b/Engine/Plugins/PCG/Source/PCG/Private/Compute/DataInterfaces/PCGTexture2DArrayDataInterface.cpp
@@ -0,0 +1,422 @@
+// Copyright Epic Games, Inc. All Rights Reserved.
+
+#include "PCGTexture2DArrayDataInterface.h"
+
+#include "PCGContext.h"
+#include "Compute/PCGDataBinding.h"
+#include "Data/PCGTexture2DArrayData.h"
+
+#include "RHIStaticStates.h"
+#include "RenderGraphBuilder.h"
+#include "ShaderCompilerCore.h"
+#include "ShaderParameterMetadataBuilder.h"
+#include "SystemTextures.h"
+#include "ComputeFramework/ShaderParamTypeDefinition.h"
+
+#include UE_INLINE_GENERATED_CPP_BY_NAME(PCGTexture2DArrayDataInterface)
+
+#define LOCTEXT_NAMESPACE "PCGTexture2DArrayDataInterface"
+
+namespace PCGTexture2DArrayDataInterfaceHelpers
+{
+	void GetSharedFunctions(TArray<FShaderFunctionDefinition>& OutFunctions)
+	{
+		OutFunctions.AddDefaulted_GetRef()
+			.SetName(TEXT("GetNumData"))
+			.AddReturnType(EShaderFundamentalType::Uint);
+
+		OutFunctions.AddDefaulted_GetRef()
+			.SetName(TEXT("GetNumElements"))
+			.AddReturnType(EShaderFundamentalType::Uint, 3)
+			.AddParam(EShaderFundamentalType::Uint); // InDataIndex
+
+		OutFunctions.AddDefaulted_GetRef()
+			.SetName(TEXT("GetResolution"))
+			.AddReturnType(EShaderFundamentalType::Uint, 2)
+			.AddParam(EShaderFundamentalType::Uint); // InDataIndex
+
+		OutFunctions.AddDefaulted_GetRef()
+			.SetName(TEXT("GetArraySize"))
+			.AddReturnType(EShaderFundamentalType::Uint)
+			.AddParam(EShaderFundamentalType::Uint); // InDataIndex
+
+		OutFunctions.AddDefaulted_GetRef()
+			.SetName(TEXT("GetNumMips"))
+			.AddReturnType(EShaderFundamentalType::Uint)
+			.AddParam(EShaderFundamentalType::Uint); // InDataIndex
+
+		OutFunctions.AddDefaulted_GetRef()
+			.SetName(TEXT("GetThreadData"))
+			.AddReturnType(EShaderFundamentalType::Bool)
+			.AddParam(EShaderFundamentalType::Uint) // InThreadIndex
+			.AddParam(EShaderFundamentalType::Uint, 0, 0, EShaderParamModifier::Out) // OutDataIndex
+			.AddParam(EShaderFundamentalType::Uint, 3, 0, EShaderParamModifier::Out); // OutElementIndex
+
+		OutFunctions.AddDefaulted_GetRef()
+			.SetName(TEXT("GetTexCoords"))
+			.AddReturnType(EShaderFundamentalType::Float, 2)
+			.AddParam(EShaderFundamentalType::Uint) // InDataIndex
+			.AddParam(EShaderFundamentalType::Float, 3); // InWorldPos
+
+		OutFunctions.AddDefaulted_GetRef()
+			.SetName(TEXT("GetTexelSize"))
+			.AddReturnType(EShaderFundamentalType::Float, 2)
+			.AddParam(EShaderFundamentalType::Uint); // InDataIndex
+
+		OutFunctions.AddDefaulted_GetRef()
+			.SetName(TEXT("GetTexelSizeWorld"))
+			.AddReturnType(EShaderFundamentalType::Float, 2)
+			.AddParam(EShaderFundamentalType::Uint); // InDataIndex
+
+		OutFunctions.AddDefaulted_GetRef()
+			.SetName(TEXT("GetTransform"))
+			.AddReturnType(EShaderFundamentalType::Float, 4, 4)
+			.AddParam(EShaderFundamentalType::Uint); // InDataIndex
+	}
+}
+
+void UPCGTexture2DArrayDataInterface::GetSupportedInputs(TArray<FShaderFunctionDefinition>& OutFunctions) const
+{
+	PCGTexture2DArrayDataInterfaceHelpers::GetSharedFunctions(OutFunctions);
+
+	OutFunctions.AddDefaulted_GetRef()
+		.SetName(TEXT("Sample"))
+		.AddReturnType(EShaderFundamentalType::Float, 4)
+		.AddParam(EShaderFundamentalType::Uint) // InDataIndex
+		.AddParam(EShaderFundamentalType::Float, 2) // InTextureUV
+		.AddParam(EShaderFundamentalType::Uint); // InSliceIndex
+
+	OutFunctions.AddDefaulted_GetRef()
+		.SetName(TEXT("SampleWorldPos"))
+		.AddReturnType(EShaderFundamentalType::Float, 4)
+		.AddParam(EShaderFundamentalType::Uint) // InDataIndex
+		.AddParam(EShaderFundamentalType::Float, 3) // WorldPos
+		.AddParam(EShaderFundamentalType::Uint); // InSliceIndex
+
+	OutFunctions.AddDefaulted_GetRef()
+		.SetName(TEXT("Load"))
+		.AddReturnType(EShaderFundamentalType::Float, 4)
+		.AddParam(EShaderFundamentalType::Uint) // InDataIndex
+		.AddParam(EShaderFundamentalType::Uint, 3); // InElementIndex = (TexelIndex, SliceIndex)
+
+	OutFunctions.AddDefaulted_GetRef()
+		.SetName(TEXT("Load"))
+		.AddReturnType(EShaderFundamentalType::Float, 4)
+		.AddParam(EShaderFundamentalType::Uint) // InDataIndex
+		.AddParam(EShaderFundamentalType::Uint, 3) // InElementIndex = (TexelIndex, SliceIndex)
+		.AddParam(EShaderFundamentalType::Uint); // InMipIndex
+}
+
+void UPCGTexture2DArrayDataInterface::GetSupportedOutputs(TArray<FShaderFunctionDefinition>& OutFunctions) const
+{
+	PCGTexture2DArrayDataInterfaceHelpers::GetSharedFunctions(OutFunctions);
+
+	// No override to select a mip as that is not supported in HLSL.
+	OutFunctions.AddDefaulted_GetRef()
+		.SetName(TEXT("Store"))
+		.AddParam(EShaderFundamentalType::Uint) // InDataIndex
+		.AddParam(EShaderFundamentalType::Uint, 3) // InElementIndex = (TexelIndex, SliceIndex)
+		.AddParam(EShaderFundamentalType::Float, 4); // InValue
+}
+
+BEGIN_SHADER_PARAMETER_STRUCT(FPCGTexture2DArrayDataInterfaceParameters, )
+	SHADER_PARAMETER_RDG_TEXTURE_SRV(Texture2DArray, TextureSRV)
+	SHADER_PARAMETER_RDG_TEXTURE_UAV(RWTexture2DArray<float4>, TextureUAV)
+	SHADER_PARAMETER_SAMPLER(SamplerState, Sampler)
+	SHADER_PARAMETER(FUintVector2, Resolution)
+	SHADER_PARAMETER(uint32, ArraySize)
+	SHADER_PARAMETER(uint32, NumMips)
+	SHADER_PARAMETER(FMatrix44f, TextureTransform)
+	SHADER_PARAMETER(FMatrix44f, TextureInverseTransform)
+	SHADER_PARAMETER(FVector2f, TexelSizeWorld)
+END_SHADER_PARAMETER_STRUCT()
+
+void UPCGTexture2DArrayDataInterface::GetShaderParameters(TCHAR const* UID, FShaderParametersMetadataBuilder& InOutBuilder, FShaderParametersMetadataAllocations& InOutAllocations) const
+{
+	InOutBuilder.AddNestedStruct<FPCGTexture2DArrayDataInterfaceParameters>(UID);
+}
+
+TCHAR const* UPCGTexture2DArrayDataInterface::TemplateFilePath = TEXT("/Plugin/PCG/Private/PCGTexture2DArrayDataInterface.ush");
+
+TCHAR const* UPCGTexture2DArrayDataInterface::GetShaderVirtualPath() const
+{
+	return TemplateFilePath;
+}
+
+void UPCGTexture2DArrayDataInterface::GetShaderHash(FString& InOutKey) const
+{
+	GetShaderFileHash(TemplateFilePath, EShaderPlatform::SP_PCD3D_SM5).AppendString(InOutKey);
+}
+
+void UPCGTexture2DArrayDataInterface::GetHLSL(FString& OutHLSL, FString const& InDataInterfaceName) const
+{
+	TMap<FString, FStringFormatArg> TemplateArgs =
+	{
+		{ TEXT("DataInterfaceName"), InDataInterfaceName },
+	};
+
+	FString TemplateFile;
+	if (ensure(LoadShaderSourceFile(TemplateFilePath, EShaderPlatform::SP_PCD3D_SM5, &TemplateFile, nullptr)))
+	{
+		OutHLSL += FString::Format(*TemplateFile, TemplateArgs);
+	}
+}
+
+UComputeDataProvider* UPCGTexture2DArrayDataInterface::CreateDataProvider() const
+{
+	return NewObject<UPCGTexture2DArrayDataProvider>();
+}
+
+void UPCGTexture2DArrayDataProvider::Initialize(const UComputeDataInterface* InDataInterface, UObject* InBinding, uint64 InInputMask, uint64 InOutputMask)
+{
+	Super::Initialize(InDataInterface, InBinding, InInputMask, InOutputMask);
+
+	const UPCGTexture2DArrayDataInterface* DataInterface = CastChecked<UPCGTexture2DArrayDataInterface>(InDataInterface);
+	bInitializeFromDataCollection = DataInterface->GetInitializeFromDataCollection();
+}
+
+FComputeDataProviderRenderProxy* UPCGTexture2DArrayDataProvider::GetRenderProxy()
+{
+	FPCGTexture2DArrayDataProviderProxy::FData Data;
+	Data.BindingInfo = BindingInfo;
+	Data.ExportMode = GetExportMode();
+	Data.PinDesc = GetPinDescription();
+	Data.OutputPinLabel = GetOutputPinLabel();
+	Data.OutputPinLabelAlias = GetOutputPinLabelAlias();
+	Data.OriginatingGenerationCount = GetGenerationCounter();
+	Data.WeakDataProvider_GT = this;
+
+	return new FPCGTexture2DArrayDataProviderProxy(Data);
+}
+
+void UPCGTexture2DArrayDataProvider::Reset()
+{
+	Super::Reset();
+
+	bInitializeFromDataCollection = false;
+	BindingInfo = {};
+}
+
+bool UPCGTexture2DArrayDataProvider::PrepareForExecute_GameThread(UPCGDataBinding* InBinding)
+{
+	TRACE_CPUPROFILER_EVENT_SCOPE(UPCGTexture2DArrayDataProvider::PrepareForExecute_GameThread);
+
+	check(InBinding);
+
+	if (!Super::PrepareForExecute_GameThread(InBinding))
+	{
+		return false;
+	}
+
+	if (GetInitializeFromDataCollection())
+	{
+		InitializeFromDataCollection(InBinding);
+	}
+	else
+	{
+		InitializeFromDataDescription(InBinding);
+	}
+
+	return true;
+}
+
+void UPCGTexture2DArrayDataProvider::InitializeFromDataCollection(UPCGDataBinding* InBinding)
+{
+	check(InBinding);
+
+	// Take any input pin label alias to obtain the data from the input data collection.
+	check(!GetDownstreamInputPinLabelAliases().IsEmpty());
+	const TArray<FPCGTaggedData> InputTaggedData = InBinding->GetInputDataCollection().GetInputsByPin(GetDownstreamInputPinLabelAliases()[0]);
+
+	if (!InputTaggedData.IsEmpty())
+	{
+		const FPCGTaggedData& TaggedData = InputTaggedData[0];
+		const UPCGTexture2DArrayData* TextureData = Cast<UPCGTexture2DArrayData>(TaggedData.Data);
+
+		if (ensure(TextureData))
+		{
+			BindingInfo = FPCGTextureBindingInfo(TextureData);
+		}
+		else
+		{
+			UE_LOG(LogPCG, Error, TEXT("Unsupported data type encountered by Texture2DArray data interface: '%s'"), TaggedData.Data ? *TaggedData.Data->GetName() : TEXT("NULL"));
+		}
+	}
+}
+
+void UPCGTexture2DArrayDataProvider::InitializeFromDataDescription(UPCGDataBinding* InBinding)
+{
+	check(InBinding);
+
+	TSharedPtr<FPCGContextHandle> ContextHandle = InBinding->GetContextHandle().Pin();
+	FPCGContext* Context = ContextHandle ? ContextHandle->GetContext() : nullptr;
+
+	if (!ensure(Context))
+	{
+		return;
+	}
+
+	TSharedPtr<const FPCGDataCollectionDesc> PinDesc = GetPinDescription();
+
+	if (ensure(PinDesc) && !PinDesc->GetDataDescriptions().IsEmpty())
+	{
+		BindingInfo = FPCGTextureBindingInfo(PinDesc->GetDataDescriptions()[0]);
+	}
+}
+
+bool FPCGTexture2DArrayDataProviderProxy::IsValid(FValidationData const& InValidationData) const
+{
+	if (!Data.BindingInfo.IsValid())
+	{
+		UE_LOG(LogPCG, Warning, TEXT("FPCGTexture2DArrayDataProviderProxy invalid due to invalid texture binding."));
+		return false;
+	}
+
+	return InValidationData.ParameterStructSize == sizeof(FParameters);
+}
+
+void FPCGTexture2DArrayDataProviderProxy::GatherDispatchData(FDispatchData const& InDispatchData)
+{
+	const FPCGTextureBindingInfo& BindingInfo = Data.BindingInfo;
+	const FMatrix44f Transform = FMatrix44f(BindingInfo.Transform.ToMatrixWithScale()).GetTransposed();
+	const FMatrix44f InverseTransform = !FMath::IsNearlyZero(BindingInfo.Transform.GetDeterminant()) ? FMatrix44f(BindingInfo.Transform.ToInverseMatrixWithScale().GetTransposed()) : FMatrix44f::Identity;
+	const FVector& TextureWorldScale = BindingInfo.Transform.GetScale3D();
+	const float TexelSizeWorldX = BindingInfo.Size.X != 0 ? (TextureWorldScale.X * 2.0) / BindingInfo.Size.X : 0.0f;
+	const float TexelSizeWorldY = BindingInfo.Size.Y != 0 ? (TextureWorldScale.Y * 2.0) / BindingInfo.Size.Y : 0.0f;
+
+	const TStridedView<FParameters> ParameterArray = MakeStridedParameterView<FParameters>(InDispatchData);
+	for (int32 InvocationIndex = 0; InvocationIndex < ParameterArray.Num(); ++InvocationIndex)
+	{
+		FParameters& Parameters = ParameterArray[InvocationIndex];
+
+		Parameters.TextureSRV = TextureSRV;
+		Parameters.TextureUAV = TextureUAV;
+		Parameters.Sampler = BindingInfo.Filter == EPCGTextureFilter::Point ? TStaticSamplerState<SF_Point, AM_Clamp, AM_Clamp, AM_Clamp>::GetRHI() : TStaticSamplerState<SF_Bilinear, AM_Clamp, AM_Clamp, AM_Clamp>::GetRHI();
+		Parameters.Resolution = FUintVector2(BindingInfo.Size.X, BindingInfo.Size.Y);
+		Parameters.ArraySize = static_cast<uint32>(BindingInfo.ArraySize);
+		Parameters.NumMips = static_cast<uint32>(BindingInfo.NumMips);
+		Parameters.TextureTransform = Transform;
+		Parameters.TextureInverseTransform = InverseTransform;
+		Parameters.TexelSizeWorld.X = TexelSizeWorldX;
+		Parameters.TexelSizeWorld.Y = TexelSizeWorldY;
+	}
+}
+
+void FPCGTexture2DArrayDataProviderProxy::AllocateResources(FRDGBuilder& GraphBuilder, FAllocationData const& InAllocationData)
+{
+	LLM_SCOPE_BYTAG(PCG);
+
+	CreateDefaultTextures(GraphBuilder);
+	CreateTextures(GraphBuilder);
+}
+
+void FPCGTexture2DArrayDataProviderProxy::CreateDefaultTextures(FRDGBuilder& GraphBuilder)
+{
+	FRDGTextureRef DummyTextureSRV = GSystemTextures.GetDefaultTexture(GraphBuilder, ETextureDimension::Texture2DArray, EPixelFormat::PF_G8, 0.0f);
+	FRDGTextureRef DummyTextureUAV = GraphBuilder.CreateTexture(
+		FRDGTextureDesc::Create2DArray(FIntPoint(1, 1), PF_G8, FClearValueBinding::Black, ETextureCreateFlags::ShaderResource | ETextureCreateFlags::UAV, /*ArraySize=*/1),
+		TEXT("PCGTexture2DArrayDI_DummyUAV"));
+
+	TextureSRV = GraphBuilder.CreateSRV(DummyTextureSRV);
+	TextureUAV = GraphBuilder.CreateUAV(DummyTextureUAV);
+}
+
+void FPCGTexture2DArrayDataProviderProxy::CreateTextures(FRDGBuilder& GraphBuilder)
+{
+	FRDGTextureRef ExportableTexture = nullptr;
+
+	FRDGTextureRef Texture = nullptr;
+	bool bCanCreateUAV = false;
+
+	if (Data.BindingInfo.ResourceType == EPCGTextureResourceType::TextureObject && Data.BindingInfo.Texture)
+	{
+		Texture = GraphBuilder.RegisterExternalTexture(CreateRenderTarget(Data.BindingInfo.Texture, TEXT("PCGTexture2DArrayDI_SRV")));
+	}
+	else if (Data.BindingInfo.ResourceType == EPCGTextureResourceType::ExportedTexture && Data.BindingInfo.ExportedTexture)
+	{
+		Texture = GraphBuilder.RegisterExternalTexture(Data.BindingInfo.ExportedTexture);
+	}
+	else
+	{
+		Texture = GraphBuilder.CreateTexture(
+			FRDGTextureDesc::Create2DArray(
+				Data.BindingInfo.Size,
+				Data.BindingInfo.Format,
+				FClearValueBinding::Black,
+				ETextureCreateFlags::ShaderResource | ETextureCreateFlags::UAV,
+				Data.BindingInfo.ArraySize),
+			TEXT("PCGTexture2DArrayDI_UAV"));
+
+		bCanCreateUAV = true;
+	}
+
+	TextureSRV = GraphBuilder.CreateSRV(Texture);
+
+	if (bCanCreateUAV)
+	{
+		TextureUAV = GraphBuilder.CreateUAV(Texture);
+	}
+
+	ExportableTexture = Texture;
+
+	if (Data.ExportMode != EPCGExportMode::NoExport)
+	{
+		ExportTextureUAV(GraphBuilder.ConvertToExternalTexture(ExportableTexture));
+	}
+}
+
+void FPCGTexture2DArrayDataProviderProxy::ExportTextureUAV(TRefCountPtr<IPooledRenderTarget> ExportedTexture)
+{
+	LLM_SCOPE_BYTAG(PCG);
+
+	// Export textures and pass a reference back to the main thread where it will be picked up by the compute graph element.
+	ExecuteOnGameThread(UE_SOURCE_LOCATION, [Data = Data, ExportedTexture = ExportedTexture]()
+	{
+		LLM_SCOPE_BYTAG(PCG);
+
+		// Obtain objects. No ensures added because a graph cancellation could feasibly destroy some or all of these.
+		UPCGTexture2DArrayDataProvider* DataProvider = Data.WeakDataProvider_GT.Get();
+		if (!DataProvider)
+		{
+			UE_LOG(LogPCG, Error, TEXT("Could not resolve UPCGTexture2DArrayDataProvider object to pass back buffer handle."));
+			return;
+		}
+
+		if (DataProvider->GetGenerationCounter() != Data.OriginatingGenerationCount)
+		{
+			return;
+		}
+
+		if (DataProvider->GetInitializeFromDataCollection())
+		{
+			DataProvider->OnDataExported_GameThread().Broadcast();
+			return;
+		}
+
+		if (!ensure(Data.PinDesc) || Data.PinDesc->GetDataDescriptions().IsEmpty())
+		{
+			return;
+		}
+
+		if (UPCGDataBinding* Binding = DataProvider->GetDataBinding())
+		{
+			const TArray<FString>& StringTable = Binding->GetStringTable();
+			FTransform TextureTransform = Data.BindingInfo.Transform;
+			const FVector::FReal XSize = 2.0 * TextureTransform.GetScale3D().X;
+			const FVector::FReal YSize = 2.0 * TextureTransform.GetScale3D().Y;
+			const FVector2D TexelSize = FVector2D(XSize, YSize) / Data.BindingInfo.Size;
+
+			FPCGTexture2DArrayDataInitParams InitParams;
+			InitParams.Transform = MoveTemp(TextureTransform);
+
+			UPCGTexture2DArrayData* ExportedData = NewObject<UPCGTexture2DArrayData>();
+			ExportedData->Initialize(ExportedTexture, InitParams);
+
+			Binding->ReceiveDataFromGPU_GameThread(ExportedData, DataProvider->GetProducerSettings(), Data.ExportMode, Data.OutputPinLabel, Data.OutputPinLabelAlias, Data.PinDesc->GetDataDescriptions()[0]);
+		}
+
+		DataProvider->OnDataExported_GameThread().Broadcast();
+	});
+}
+
+#undef LOCTEXT_NAMESPACE
diff --git a/Engine/Plugins/PCG/Source/PCG/Private/Compute/DataInterfaces/PCGTexture2DArrayDataInterface.h b/Engine/Plugins/PCG/Source/PCG/Private/Compute/DataInterfaces/PCGTexture2DArrayDataInterface.h
new file mode 100644
--- /dev/null
+++ b/Engine/Plugins/PCG/Source/PCG/Private/Compute/DataInterfaces/PCGTexture2DArrayDataInterface.h
@@ -0,0 +1,116 @@
+// Copyright Epic Games, Inc. All Rights Reserved.
+
+#pragma once
+
+#include "Compute/DataInterfaces/PCGTextureInfo.h"
+#include "Compute/DataInterfaces/PCGComputeDataInterface.h"
+
+#include "RHIResources.h"
+
+#include "PCGTexture2DArrayDataInterface.generated.h"
+
+class FPCGTexture2DArrayDataInterfaceParameters;
+class UPCGTexture2DArrayData;
+
+/** Data Interface allowing read/write of a UPCGTexture2DArrayData. */
+UCLASS(ClassGroup = (Procedural))
+class UPCGTexture2DArrayDataInterface : public UPCGExportableDataInterface
+{
+	GENERATED_BODY()
+
+public:
+	//~ Begin UComputeDataInterface Interface
+	TCHAR const* GetClassName() const override { return TEXT("PCGTexture2DArray"); }
+	void GetSupportedInputs(TArray<FShaderFunctionDefinition>& OutFunctions) const override;
+	void GetSupportedOutputs(TArray<FShaderFunctionDefinition>& OutFunctions) const override;
+	void GetShaderParameters(TCHAR const* UID, FShaderParametersMetadataBuilder& InOutBuilder, FShaderParametersMetadataAllocations& InOutAllocations) const override;
+	TCHAR const* GetShaderVirtualPath() const override;
+	void GetShaderHash(FString& InOutKey) const override;
+	void GetHLSL(FString& OutHLSL, FString const& InDataInterfaceName) const override;
+	UComputeDataProvider* CreateDataProvider() const override;
+	//~ End UComputeDataInterface Interface
+
+	void SetInitializeFromDataCollection(bool bInInitializeFromDataCollection) { bInitializeFromDataCollection = bInInitializeFromDataCollection; }
+	bool GetInitializeFromDataCollection() const { return bInitializeFromDataCollection; }
+
+protected:
+	UPROPERTY()
+	bool bInitializeFromDataCollection = false;
+
+private:
+	static TCHAR const* TemplateFilePath;
+};
+
+/** Data Provider allowing read/write of a UPCGTexture2DArrayData. */
+UCLASS()
+class UPCGTexture2DArrayDataProvider : public UPCGExportableDataProvider
+{
+	GENERATED_BODY()
+
+public:
+	//~ Begin UComputeDataProvider Interface
+	virtual void Initialize(const UComputeDataInterface* InDataInterface, UObject* InBinding, uint64 InInputMask, uint64 InOutputMask) override;
+	FComputeDataProviderRenderProxy* GetRenderProxy() override;
+	virtual void Reset() override;
+	//~ End UComputeDataProvider Interface
+
+	//~ Begin UPCGComputeDataProvider Interface
+	virtual bool PrepareForExecute_GameThread(UPCGDataBinding* InBinding) override;
+	//~ End UPCGComputeDataProvider Interface
+
+	bool GetInitializeFromDataCollection() const { return bInitializeFromDataCollection; }
+
+private:
+	void InitializeFromDataCollection(UPCGDataBinding* InBinding);
+	void InitializeFromDataDescription(UPCGDataBinding* InBinding);
+
+private:
+	bool bInitializeFromDataCollection = false;
+
+	FPCGTextureBindingInfo BindingInfo;
+};
+
+class FPCGTexture2DArrayDataProviderProxy : public FComputeDataProviderRenderProxy
+{
+public:
+	struct FData
+	{
+		FPCGTextureBindingInfo BindingInfo;
+
+		// For export
+		EPCGExportMode ExportMode = EPCGExportMode::NoExport;
+		TSharedPtr<const FPCGDataCollectionDesc> PinDesc = nullptr;
+		FName OutputPinLabel;
+		FName OutputPinLabelAlias;
+	
+		/** Generation count of the data provider when the proxy was created. */
+		uint64 OriginatingGenerationCount = 0;
+
+		// Weak pointer useful for passing back texture handles. Do not access this directly from the render thread.
+		TWeakObjectPtr<UPCGTexture2DArrayDataProvider> WeakDataProvider_GT;
+	};
+
+	explicit FPCGTexture2DArrayDataProviderProxy(FData InData)
+		: Data(MoveTemp(InData))
+	{}
+
+	//~ Begin FComputeDataProviderRenderProxy Interface
+	bool IsValid(FValidationData const& InValidationData) const override;
+	void GatherDispatchData(FDispatchData const& InDispatchData) override;
+	void AllocateResources(FRDGBuilder& GraphBuilder, FAllocationData const& InAllocationData) override;
+	//~ End FComputeDataProviderRenderProxy Interface
+
+private:
+	void CreateDefaultTextures(FRDGBuilder& GraphBuilder);
+	void CreateTextures(FRDGBuilder& GraphBuilder);
+	void ExportTextureUAV(TRefCountPtr<IPooledRenderTarget> ExportedTexture);
+
+private:
+	using FParameters = FPCGTexture2DArrayDataInterfaceParameters;
+
+	FData Data;
+
+	// Allocated resources
+	FRDGTextureSRVRef TextureSRV;
+	FRDGTextureUAVRef TextureUAV;
+};
diff --git a/Engine/Plugins/PCG/Source/PCG/Private/Compute/DataInterfaces/PCGTextureDataInterface.cpp b/Engine/Plugins/PCG/Source/PCG/Private/Compute/DataInterfaces/PCGTextureDataInterface.cpp
--- a/Engine/Plugins/PCG/Source/PCG/Private/Compute/DataInterfaces/PCGTextureDataInterface.cpp
+++ b/Engine/Plugins/PCG/Source/PCG/Private/Compute/DataInterfaces/PCGTextureDataInterface.cpp
@@ -70,60 +70,14 @@ namespace PCGTextureDataInterfaceHelpers
 			.SetName(TEXT("GetTexelSizeWorld"))
 			.AddReturnType(EShaderFundamentalType::Float, 2)
 			.AddParam(EShaderFundamentalType::Uint); // InDataIndex
-	}
-}
-
-FPCGTextureBindingInfo::FPCGTextureBindingInfo(const UPCGBaseTextureData* InTextureData)
-{
-	check(InTextureData);
-
-	ResourceType = InTextureData->GetTextureResourceType();
-	Texture = InTextureData->GetTextureRHI();
-	ExportedTexture = InTextureData->GetRefCountedTexture();
-	Transform = InTextureData->GetTransform();
-	Filter = InTextureData->Filter;
 
-	if (Texture.IsValid())
-	{
-		const FRHITextureDesc& Desc = Texture->GetDesc();
-		Size = Desc.Extent;
-		Dimension = Desc.Dimension;
-		Format = Desc.Format;
-	}
-	else
-	{
-		UE_LOG(LogPCG, Error, TEXT("PCGTextureDataInterface: Texture data '%s' has an invalid texture."), InTextureData ? *InTextureData->GetName() : TEXT("NULL"));
+		OutFunctions.AddDefaulted_GetRef()
+			.SetName(TEXT("GetNumMips"))
+			.AddReturnType(EShaderFundamentalType::Uint)
+			.AddParam(EShaderFundamentalType::Uint); // InDataIndex
 	}
 }
 
-FPCGTextureBindingInfo::FPCGTextureBindingInfo(const FPCGDataDesc& InDataDesc, const FTransform& InTransform)
-{
-	// @todo_pcg: Need to support other information here too, like format, slice index, resource type, etc.
-	Size = FIntPoint(FMath::Max(InDataDesc.GetElementCount().X, 1), FMath::Max(InDataDesc.GetElementCount().Y, 1));
-	Format = PCGComputeHelpers::GetPixelFormatFromPCGRenderTargetFormat(InDataDesc.GetRenderTargetFormat());
-	Filter = InDataDesc.GetTextureFilter();
-	Transform = InTransform;
-}
-
-bool FPCGTextureBindingInfo::operator==(const FPCGTextureBindingInfo& Other) const
-{
-	return ResourceType == Other.ResourceType
-		&& Texture == Other.Texture
-		&& ExportedTexture == Other.ExportedTexture
-		&& Size == Other.Size
-		&& Dimension == Other.Dimension
-		&& Format == Other.Format
-		&& Filter == Other.Filter
-		&& Transform.Equals(Other.Transform);
-}
-
-bool FPCGTextureBindingInfo::IsValid() const
-{
-	const int32 MaxTextureDimension = GetMax2DTextureDimension();
-
-	return ResourceType != EPCGTextureResourceType::Invalid && Size.X > 0 && Size.Y > 0 && Size.X <= MaxTextureDimension && Size.Y <= MaxTextureDimension;
-}
-
 void UPCGTextureDataInterface::GetSupportedInputs(TArray<FShaderFunctionDefinition>& OutFunctions) const
 {
 	// TODO: Support texture3d, etc. UNiagaraDataInterfaceRenderTarget2D and related are useful references.
@@ -148,6 +102,11 @@ void UPCGTextureDataInterface::GetSupportedInputs(TArray<FShaderFunctionDefiniti
 		.AddParam(EShaderFundamentalType::Uint) // InDataIndex
 		.AddParam(EShaderFundamentalType::Uint, 2); // InElementIndex
 
+	OutFunctions.AddDefaulted_GetRef()
+		.SetName(TEXT("GetTransform"))
+		.AddReturnType(EShaderFundamentalType::Float, 4, 4)
+		.AddParam(EShaderFundamentalType::Uint); // InDataIndex
+
 	// Deprecated 5.6. Assumes data index is 0.
 	OutFunctions.AddDefaulted_GetRef()
 		.SetName(TEXT("Sample"))
@@ -194,6 +153,7 @@ BEGIN_SHADER_PARAMETER_STRUCT(FPCGTextureDataInterfaceParameters,)
 	SHADER_PARAMETER_RDG_TEXTURE_SRV_ARRAY(Texture2DArray, TextureArray, [PCGTextureDataInterfaceConstants::MAX_NUM_SRV_BINDINGS])
 	SHADER_PARAMETER_RDG_TEXTURE_UAV_ARRAY(RWTexture2D<float4>, TextureUAV, [PCGTextureDataInterfaceConstants::MAX_NUM_UAV_BINDINGS])
 	SHADER_PARAMETER_SAMPLER(SamplerState, Sampler)
+	SHADER_PARAMETER_ARRAY(FMatrix44f, TextureTransforms, [PCGTextureDataInterfaceConstants::MAX_NUM_SRV_BINDINGS])
 	SHADER_PARAMETER_ARRAY(FMatrix44f, TextureInverseTransforms, [PCGTextureDataInterfaceConstants::MAX_NUM_SRV_BINDINGS])
 	SHADER_PARAMETER(uint32, NumTextureInfos)
 	SHADER_PARAMETER_RDG_BUFFER_SRV(StructuredBuffer<FUintVector4>, TextureInfos)
@@ -364,13 +324,9 @@ void UPCGTextureDataProvider::BuildInfosFromDataDescription(UPCGDataBinding* InB
 	TSharedPtr<const FPCGDataCollectionDesc> PinDesc = GetPinDescription();
 	check(PinDesc);
 
-	const FBox OriginalLocalBounds = Context->ExecutionSource->GetExecutionState().GetOriginalLocalSpaceBounds();
-	FTransform Transform = Context->ExecutionSource->GetExecutionState().GetOriginalTransform();
-	Transform.SetScale3D(Transform.GetScale3D() * 0.5 * (OriginalLocalBounds.Max - OriginalLocalBounds.Min));
-
 	for (const FPCGDataDesc& DataDesc : PinDesc->GetDataDescriptions())
 	{
-		FPCGTextureBindingInfo BindingInfo(DataDesc, Transform);
+		FPCGTextureBindingInfo BindingInfo(DataDesc);
 		int32 BindingIndex = BindingInfos.Find(BindingInfo);
 
 		if (BindingIndex == INDEX_NONE)
@@ -474,10 +430,12 @@ void FPCGTextureDataProviderProxy::GatherDispatchData(FDispatchData const& InDis
 			if (BindingIndex < BindingInfos.Num())
 			{
 				const FTransform& Transform = BindingInfos[BindingIndex].Transform;
-				Parameters.TextureInverseTransforms[BindingIndex] = !FMath::IsNearlyZero(Transform.GetDeterminant()) ? FMatrix44f(Transform.ToInverseMatrixWithScale()) : FMatrix44f::Identity;
+				Parameters.TextureTransforms[BindingIndex] = FMatrix44f(Transform.ToMatrixWithScale()).GetTransposed();
+				Parameters.TextureInverseTransforms[BindingIndex] = !FMath::IsNearlyZero(Transform.GetDeterminant()) ? FMatrix44f(Transform.ToInverseMatrixWithScale()).GetTransposed() : FMatrix44f::Identity;
 			}
 			else
 			{
+				Parameters.TextureTransforms[BindingIndex] = FMatrix44f::Identity;
 				Parameters.TextureInverseTransforms[BindingIndex] = FMatrix44f::Identity;
 			}
 		}
@@ -670,7 +628,7 @@ void FPCGTextureDataProviderProxy::PackTextureInfos(FRDGBuilder& GraphBuilder)
 	//   - uint32 - BindingIndex
 	//   - uint32 - SliceIndex (aka TextureIndex if bound to a TextureArray, 0 otherwise.)
 	//   - uint32 - Dimension
-	//   - uint32 - UNUSED
+	//   - uint32 - NumMips
 	//   - uint32 - ResolutionX
 	//   - uint32 - ResolutionY
 	//   - float  - TexelSizeWorldX (cm)
@@ -697,7 +655,7 @@ void FPCGTextureDataProviderProxy::PackTextureInfos(FRDGBuilder& GraphBuilder)
 			const float TexelWorldSizeX = BindingInfo.Size.X != 0 ? (TextureWorldScale.X * 2.0) / BindingInfo.Size.X : 0.0f;
 			const float TexelWorldSizeY = BindingInfo.Size.Y != 0 ? (TextureWorldScale.Y * 2.0) / BindingInfo.Size.Y : 0.0f;
 
-			TextureInfosPackedData[TextureInfoIndex * 2 + 0] = FUintVector4(TextureInfo.BindingIndex, TextureInfo.SliceIndex, static_cast<uint32>(BindingInfo.Dimension), /*UNUSEd*/0u);
+			TextureInfosPackedData[TextureInfoIndex * 2 + 0] = FUintVector4(TextureInfo.BindingIndex, TextureInfo.SliceIndex, static_cast<uint32>(BindingInfo.Dimension), BindingInfo.NumMips);
 			TextureInfosPackedData[TextureInfoIndex * 2 + 1] = FUintVector4(BindingInfo.Size.X, BindingInfo.Size.Y, FMath::AsUInt(TexelWorldSizeX), FMath::AsUInt(TexelWorldSizeY));
 		}
 
diff --git a/Engine/Plugins/PCG/Source/PCG/Private/Compute/DataInterfaces/PCGTextureDataInterface.h b/Engine/Plugins/PCG/Source/PCG/Private/Compute/DataInterfaces/PCGTextureDataInterface.h
--- a/Engine/Plugins/PCG/Source/PCG/Private/Compute/DataInterfaces/PCGTextureDataInterface.h
+++ b/Engine/Plugins/PCG/Source/PCG/Private/Compute/DataInterfaces/PCGTextureDataInterface.h
@@ -4,6 +4,7 @@
 
 #include "Compute/PCGDataDescription.h"
 #include "Compute/DataInterfaces/PCGComputeDataInterface.h"
+#include "Compute/DataInterfaces/PCGTextureInfo.h"
 #include "Data/PCGTextureData.h"
 
 #include "RHIResources.h"
@@ -25,29 +26,6 @@ namespace PCGTextureDataInterfaceConstants
 	constexpr uint32 MAX_NUM_UAV_BINDINGS = 1;
 }
 
-struct FPCGTextureBindingInfo
-{
-	explicit FPCGTextureBindingInfo(const UPCGBaseTextureData* InTextureData);
-	FPCGTextureBindingInfo(const FPCGDataDesc& InDataDesc, const FTransform& InTransform);
-	bool operator==(const FPCGTextureBindingInfo& Other) const;
-	bool IsValid() const;
-
-	EPCGTextureResourceType ResourceType = EPCGTextureResourceType::TextureObject;
-	FTextureRHIRef Texture = nullptr;
-	TRefCountPtr<IPooledRenderTarget> ExportedTexture = nullptr;
-	FIntPoint Size = FIntPoint::ZeroValue;
-	ETextureDimension Dimension = ETextureDimension::Texture2D;
-	EPixelFormat Format = EPixelFormat::PF_Unknown;
-	EPCGTextureFilter Filter = EPCGTextureFilter::Bilinear;
-	FTransform Transform = FTransform::Identity;
-};
-
-struct FPCGTextureInfo
-{
-	int BindingIndex = 0;
-	int SliceIndex = 0;
-};
-
 /** Data Interface allowing sampling of a texture. */
 UCLASS(ClassGroup = (Procedural))
 class UPCGTextureDataInterface : public UPCGExportableDataInterface
diff --git a/Engine/Plugins/PCG/Source/PCG/Private/Compute/DataInterfaces/PCGTextureInfo.cpp b/Engine/Plugins/PCG/Source/PCG/Private/Compute/DataInterfaces/PCGTextureInfo.cpp
new file mode 100644
--- /dev/null
+++ b/Engine/Plugins/PCG/Source/PCG/Private/Compute/DataInterfaces/PCGTextureInfo.cpp
@@ -0,0 +1,96 @@
+// Copyright Epic Games, Inc. All Rights Reserved.
+
+#include "PCGTextureInfo.h"
+
+#include "Compute/PCGDataDescription.h"
+#include "Data/PCGTexture2DArrayData.h"
+#include "Data/PCGTextureData.h"
+
+#include "RHIDefinitions.h"
+
+FPCGTextureBindingInfo::FPCGTextureBindingInfo(const UPCGBaseTextureData* InTextureData)
+{
+	check(InTextureData);
+
+	ResourceType = InTextureData->GetTextureResourceType();
+	Texture = InTextureData->GetTextureRHI();
+	ExportedTexture = InTextureData->GetRefCountedTexture();
+	Transform = InTextureData->GetTransform();
+	Filter = InTextureData->Filter;
+
+	if (Texture.IsValid())
+	{
+		const FRHITextureDesc& Desc = Texture->GetDesc();
+		Size = Desc.Extent;
+		ArraySize = Desc.ArraySize;
+		Dimension = Desc.Dimension;
+		Format = Desc.Format;
+		NumMips = Desc.NumMips;
+	}
+	else
+	{
+		UE_LOG(LogPCG, Error, TEXT("PCGTextureDataInterface: Texture data '%s' has an invalid texture."), InTextureData ? *InTextureData->GetName() : TEXT("NULL"));
+	}
+}
+
+FPCGTextureBindingInfo::FPCGTextureBindingInfo(const UPCGTexture2DArrayData* InTextureData)
+{
+	check(InTextureData);
+
+	ResourceType = InTextureData->GetResourceType();
+	Texture = InTextureData->GetRHI();
+	ExportedTexture = InTextureData->GetRefCountedTexture();
+	Filter = InTextureData->GetFilter();
+	Transform = InTextureData->GetTransform();
+
+	if (Texture.IsValid())
+	{
+		const FRHITextureDesc& Desc = Texture->GetDesc();
+		Size = Desc.Extent;
+		ArraySize = Desc.ArraySize;
+		Dimension = Desc.Dimension;
+		Format = Desc.Format;
+		NumMips = Desc.NumMips;
+	}
+	else
+	{
+		UE_LOG(LogPCG, Error, TEXT("PCGTexture2DArrayDataInterface: Texture2DArray data '%s' has an invalid texture."), InTextureData ? *InTextureData->GetName() : TEXT("NULL"));
+	}
+}
+
+FPCGTextureBindingInfo::FPCGTextureBindingInfo(const FPCGDataDesc& InDataDesc)
+{
+	Size = FIntPoint(FMath::Max(InDataDesc.GetElementCount().X, 1), FMath::Max(InDataDesc.GetElementCount().Y, 1));
+	ArraySize = InDataDesc.GetTextureArraySize();
+	Dimension = InDataDesc.GetType().IsChildOf(FPCGDataTypeInfoTexture2DArray::AsId()) ? ETextureDimension::Texture2DArray : ETextureDimension::Texture2D;
+	Format = PCGComputeHelpers::GetPixelFormatFromPCGRenderTargetFormat(InDataDesc.GetRenderTargetFormat());
+	Filter = InDataDesc.GetTextureFilter();
+	Transform = InDataDesc.GetTextureTransform();
+	NumMips = InDataDesc.GetTextureNumMips();
+}
+
+bool FPCGTextureBindingInfo::operator==(const FPCGTextureBindingInfo& Other) const
+{
+	return ResourceType == Other.ResourceType
+		&& Texture == Other.Texture
+		&& ExportedTexture == Other.ExportedTexture
+		&& Size == Other.Size
+		&& ArraySize == Other.ArraySize
+		&& NumMips == Other.NumMips
+		&& Dimension == Other.Dimension
+		&& Format == Other.Format
+		&& Filter == Other.Filter
+		&& Transform.Equals(Other.Transform);
+}
+
+bool FPCGTextureBindingInfo::IsValid() const
+{
+	const int32 MaxTextureDimension = GetMax2DTextureDimension();
+	const int32 MaxTextureArrayLayers = GetMaxTextureArrayLayers();
+
+	return ResourceType != EPCGTextureResourceType::Invalid
+		&& Size.X > 0 && Size.Y > 0
+		&& Size.X <= MaxTextureDimension && Size.Y <= MaxTextureDimension
+		&& (Dimension != ETextureDimension::Texture2DArray || (ArraySize > 0 && ArraySize <= MaxTextureArrayLayers))
+		&& NumMips > 0 && NumMips <= MAX_TEXTURE_MIP_COUNT;
+}
diff --git a/Engine/Plugins/PCG/Source/PCG/Private/Compute/DataInterfaces/PCGTextureInfo.h b/Engine/Plugins/PCG/Source/PCG/Private/Compute/DataInterfaces/PCGTextureInfo.h
new file mode 100644
--- /dev/null
+++ b/Engine/Plugins/PCG/Source/PCG/Private/Compute/DataInterfaces/PCGTextureInfo.h
@@ -0,0 +1,36 @@
+// Copyright Epic Games, Inc. All Rights Reserved.
+
+#pragma once
+
+#include "Data/PCGTextureData.h"
+
+class UPCGBaseTextureData;
+class UPCGTexture2DArrayData;
+struct FPCGDataDesc;
+
+struct FPCGTextureBindingInfo
+{
+	FPCGTextureBindingInfo() = default;
+	explicit FPCGTextureBindingInfo(const UPCGBaseTextureData* InTextureData);
+	explicit FPCGTextureBindingInfo(const UPCGTexture2DArrayData* InTextureData);
+	FPCGTextureBindingInfo(const FPCGDataDesc& InDataDesc);
+	bool operator==(const FPCGTextureBindingInfo& Other) const;
+	bool IsValid() const;
+
+	EPCGTextureResourceType ResourceType = EPCGTextureResourceType::TextureObject;
+	FTextureRHIRef Texture = nullptr;
+	TRefCountPtr<IPooledRenderTarget> ExportedTexture = nullptr;
+	FIntPoint Size = FIntPoint::ZeroValue;
+	uint16 ArraySize = 1;
+	uint8 NumMips = 1;
+	ETextureDimension Dimension = ETextureDimension::Texture2D;
+	EPixelFormat Format = EPixelFormat::PF_Unknown;
+	EPCGTextureFilter Filter = EPCGTextureFilter::Bilinear;
+	FTransform Transform = FTransform::Identity;
+};
+
+struct FPCGTextureInfo
+{
+	int BindingIndex = 0;
+	int SliceIndex = 0;
+};
diff --git a/Engine/Plugins/PCG/Source/PCG/Private/Compute/Elements/PCGComputeGraphElement.cpp b/Engine/Plugins/PCG/Source/PCG/Private/Compute/Elements/PCGComputeGraphElement.cpp
--- a/Engine/Plugins/PCG/Source/PCG/Private/Compute/Elements/PCGComputeGraphElement.cpp
+++ b/Engine/Plugins/PCG/Source/PCG/Private/Compute/Elements/PCGComputeGraphElement.cpp
@@ -7,6 +7,7 @@
 #include "PCGModule.h"
 #include "Components/PCGProceduralISMComponent.h"
 #include "Compute/PCGComputeCommon.h"
+#include "Data/PCGTexture2DArrayData.h"
 #include "Data/PCGTextureData.h"
 #include "Compute/Data/PCGProxyForGPUData.h"
 #include "Compute/DataInterfaces/PCGDataCollectionDataInterface.h"
@@ -812,10 +813,14 @@ void FPCGComputeGraphElement::ExecuteDebugDraw(FPCGComputeGraphContext* InContex
 
 			DataToDisplay = Readback.TaggedData.Data.Get();
 		}
-		else if (const UPCGTextureData* Data = Cast<UPCGTextureData>(DataToDebug.Data.Get()))
+		else if (Cast<UPCGTextureData>(DataToDebug.Data.Get()))
 		{
 			DataToDisplay = DataToDebug.DataPendingInit;
 		}
+		else if (Cast<UPCGTexture2DArrayData>(DataToDebug.Data.Get()))
+		{
+			DataToDisplay = DataToDebug.Data;
+		}
 
 		AllDataToDebug.RemoveAtSwap(Index);
 
@@ -855,7 +860,7 @@ void FPCGComputeGraphElement::StoreDataForInspection(FPCGComputeGraphContext* In
 				TaggedData.Tags.Append(DataToInspect.Tags);
 				DataCollection.TaggedData.Add(TaggedData);
 			}
-			else if (const UPCGTextureData* TextureData = Cast<UPCGTextureData>(DataToInspect.Data.Get()))
+			else if (Cast<UPCGTextureData>(DataToInspect.Data.Get()))
 			{
 				const bool bReadbackTextureDataOnInspect = CVarReadbackTextureDataOnInspect.GetValueOnGameThread();
 
@@ -865,6 +870,14 @@ void FPCGComputeGraphElement::StoreDataForInspection(FPCGComputeGraphContext* In
 				TaggedData.Tags = DataToInspect.Tags;
 				DataCollection.TaggedData.Add(TaggedData);
 			}
+			else if (Cast<UPCGTexture2DArrayData>(DataToInspect.Data.Get()))
+			{
+				FPCGTaggedData TaggedData;
+				TaggedData.Data = DataToInspect.Data;
+				TaggedData.Pin = DataToInspect.PinLabel;
+				TaggedData.Tags = DataToInspect.Tags;
+				DataCollection.TaggedData.Add(TaggedData);
+			}
 		}
 	}
 
diff --git a/Engine/Plugins/PCG/Source/PCG/Private/Compute/Elements/PCGCustomHLSL.cpp b/Engine/Plugins/PCG/Source/PCG/Private/Compute/Elements/PCGCustomHLSL.cpp
--- a/Engine/Plugins/PCG/Source/PCG/Private/Compute/Elements/PCGCustomHLSL.cpp
+++ b/Engine/Plugins/PCG/Source/PCG/Private/Compute/Elements/PCGCustomHLSL.cpp
@@ -11,7 +11,10 @@
 #include "Compute/PCGComputeSource.h"
 #include "Compute/PCGKernelHelpers.h"
 #include "Compute/Elements/PCGCustomHLSLKernel.h"
+#include "Data/PCGTexture2DArrayData.h"
+#include "Data/PCGLandscapeData.h"
 #include "Data/PCGPointData.h"
+#include "Data/PCGTextureData.h"
 #include "Graph/PCGGPUGraphCompilationContext.h"
 
 #include "Internationalization/Regex.h"
@@ -425,6 +428,7 @@ void UPCGCustomHLSLSettings::UpdateInputDeclarations()
 	TArray<FString> PointDataPins;
 	TArray<FString> LandscapeDataPins;
 	TArray<FString> TextureDataPins;
+	TArray<FString> TextureArrayDataPins;
 	TArray<FString> VirtualTextureDataPins;
 	TArray<FString> RawBufferDataPins;
 	TArray<FString> StaticMeshDataPins;
@@ -451,6 +455,11 @@ void UPCGCustomHLSLSettings::UpdateInputDeclarations()
 			TextureDataPins.Add(Pin.Label.ToString());
 		}
 
+		if (!!(Pin.AllowedTypes & FPCGDataTypeInfoTexture2DArray::AsId()))
+		{
+			TextureArrayDataPins.Add(Pin.Label.ToString());
+		}
+
 		if (!!(Pin.AllowedTypes & EPCGDataType::VirtualTexture))
 		{
 			VirtualTextureDataPins.Add(Pin.Label.ToString());
@@ -538,7 +547,7 @@ void UPCGCustomHLSLSettings::UpdateInputDeclarations()
 
 	if (!TextureDataPins.IsEmpty())
 	{
-		InputDeclarations += TEXT("/*** INPUT TEXTURE DATA FUNCTIONS ***/\n\n");
+		InputDeclarations += TEXT("/*** INPUT TEXTURE 2D DATA FUNCTIONS ***/\n\n");
 
 		const bool bMultiPin = TextureDataPins.Num() > 1;
 
@@ -553,16 +562,51 @@ void UPCGCustomHLSLSettings::UpdateInputDeclarations()
 			"// Computes an unclamped texture coordinate.\n"
 			"float2 {0}_GetTexCoords(uint DataIndex, float3 WorldPos);\n"
 			"float4 {0}_Sample(uint DataIndex, float2 TexCoords);\n"
-			"// Computes sample coordinates of the WorldPos relative to the texture data's bounds.\n"
+			"// Computes sample coordinates of the WorldPos relative to the texture data's transform.\n"
 			"float4 {0}_SampleWorldPos(uint DataIndex, float3 WorldPos);\n"
 			"float4 {0}_Load(uint DataIndex, uint2 ElementIndex);\n"
 			"float2 {0}_GetTexelSize(uint DataIndex);\n"
-			"float2 {0}_GetTexelSizeWorld(uint DataIndex);\n"),
+			"float2 {0}_GetTexelSizeWorld(uint DataIndex);\n"
+			"float4x4 {0}_GetTransform(uint DataIndex);\n"
+			"uint {0}_GetNumMips(uint DataIndex);\n"),
 			{ bMultiPin ? PCGHLSLElement::PinDeclTemplateStr : TextureDataPins[0] });
 
 		InputDeclarations += TEXT("\n");
 	}
 
+	if (!TextureArrayDataPins.IsEmpty())
+	{
+		InputDeclarations += TEXT("/*** INPUT TEXTURE 2D ARRAY DATA FUNCTIONS ***/\n\n");
+
+		const bool bMultiPin = TextureArrayDataPins.Num() > 1;
+
+		if (bMultiPin)
+		{
+			InputDeclarations += TEXT("// Valid pins: ") + FString::Join(TextureArrayDataPins, TEXT(", ")) + TEXT("\n\n");
+		}
+
+		InputDeclarations += FString::Format(TEXT(
+			"uint {0}_GetNumData();\n"
+			"// NumElements: Size in XY, ArraySize in Z.\n"
+			"uint3 {0}_GetNumElements(uint DataIndex);\n"
+			"uint2 {0}_GetSize(uint DataIndex);\n"
+			"uint {0}_GetArraySize(uint DataIndex);\n"
+			"// Computes an unclamped texture coordinate.\n"
+			"float2 {0}_GetTexCoords(uint DataIndex, float3 WorldPos);\n"
+			"float4 {0}_Sample(uint DataIndex, float2 TexCoords, uint InSliceIndex);\n"
+			"// Computes sample coordinates of the WorldPos relative to the texture data's transform.\n"
+			"float4 {0}_SampleWorldPos(uint DataIndex, float3 WorldPos, uint InSliceIndex);\n"
+			"float4 {0}_Load(uint DataIndex, uint ElementIndex, uint SliceIndex);\n"
+			"float4 {0}_Load(uint DataIndex, uint ElementIndex, uint SliceIndex, uint MipIndex);\n"
+			"float2 {0}_GetTexelSize(uint DataIndex);\n"
+			"float2 {0}_GetTexelSizeWorld(uint DataIndex);\n"
+			"float4x4 {0}_GetTransform(uint DataIndex);\n"
+			"uint {0}_GetNumMips(uint DataIndex);\n"),
+			{ bMultiPin ? PCGHLSLElement::PinDeclTemplateStr : TextureArrayDataPins[0] });
+
+		InputDeclarations += TEXT("\n");
+	}
+
 	if (!VirtualTextureDataPins.IsEmpty())
 	{
 		InputDeclarations += TEXT("/*** INPUT VIRTUAL TEXTURE DATA FUNCTIONS ***/\n\n");
@@ -654,6 +698,7 @@ void UPCGCustomHLSLSettings::UpdateOutputDeclarations()
 	TArray<FString> DataCollectionDataPins;
 	TArray<FString> PointDataPins;
 	TArray<FString> TextureDataPins;
+	TArray<FString> TextureArrayDataPins;
 	TArray<FString> RawBufferDataPins;
 
 	for (const FPCGPinProperties& Pin : OutputPinProperties())
@@ -672,6 +717,11 @@ void UPCGCustomHLSLSettings::UpdateOutputDeclarations()
 		{
 			TextureDataPins.Add(Pin.Label.ToString());
 		}
+
+		if (!!(Pin.AllowedTypes & FPCGDataTypeInfoTexture2DArray::AsId()))
+		{
+			TextureArrayDataPins.Add(Pin.Label.ToString());
+		}
 	}
 
 	if (!DataCollectionDataPins.IsEmpty())
@@ -733,7 +783,7 @@ void UPCGCustomHLSLSettings::UpdateOutputDeclarations()
 
 	if (!TextureDataPins.IsEmpty())
 	{
-		OutputDeclarations += TEXT("/*** OUTPUT TEXTURE DATA FUNCTIONS ***/\n\n");
+		OutputDeclarations += TEXT("/*** OUTPUT TEXTURE 2D DATA FUNCTIONS ***/\n\n");
 
 		const bool bMultiPin = TextureDataPins.Num() > 1;
 
@@ -748,12 +798,43 @@ void UPCGCustomHLSLSettings::UpdateOutputDeclarations()
 			"float2 {0}_GetTexCoords(uint DataIndex, float3 WorldPos);\n"
 			"float2 {0}_GetTexelSize(uint DataIndex);\n"
 			"float2 {0}_GetTexelSizeWorld(uint DataIndex);\n"
+			"uint {0}_GetNumMips(uint DataIndex);\n"
+			"float4x4 {0}_GetTransform(uint DataIndex);\n"
+			"uint {0}_GetNumMips(uint DataIndex);\n"
 			"void {0}_Store(uint DataIndex, uint2 ElementIndex, float4 Value);\n"),
 			{ bMultiPin ? PCGHLSLElement::PinDeclTemplateStr : TextureDataPins[0] });
 
 		OutputDeclarations += TEXT("\n");
 	}
 
+	if (!TextureArrayDataPins.IsEmpty())
+	{
+		OutputDeclarations += TEXT("/*** OUTPUT TEXTURE 2D ARRAY DATA FUNCTIONS ***/\n\n");
+
+		const bool bMultiPin = TextureArrayDataPins.Num() > 1;
+
+		if (bMultiPin)
+		{
+			OutputDeclarations += TEXT("// Valid pins: ") + FString::Join(TextureArrayDataPins, TEXT(", ")) + TEXT("\n\n");
+		}
+
+		OutputDeclarations += FString::Format(TEXT(
+			"uint {0}_GetNumData();\n"
+			"// NumElements: Size in XY, ArraySize in Z.\n"
+			"uint3 {0}_GetNumElements(uint DataIndex);\n"
+			"uint2 {0}_GetSize(uint DataIndex);\n"
+			"uint {0}_GetArraySize(uint DataIndex);\n"
+			"float2 {0}_GetTexCoords(uint DataIndex, float3 WorldPos);\n"
+			"float2 {0}_GetTexelSize(uint DataIndex);\n"
+			"float2 {0}_GetTexelSizeWorld(uint DataIndex);\n"
+			"float4x4 {0}_GetTransform(uint DataIndex);\n"
+			"uint {0}_GetNumMips(uint DataIndex);\n"
+			"void {0}_Store(uint DataIndex, uint2 ElementIndex, uint SliceIndex, float4 Value);\n"),
+			{ bMultiPin ? PCGHLSLElement::PinDeclTemplateStr : TextureArrayDataPins[0] });
+
+		OutputDeclarations += TEXT("\n");
+	}
+
 	if (!RawBufferDataPins.IsEmpty())
 	{
 		OutputDeclarations += TEXT("/*** OUTPUT BYTE ADDRESS BUFFER DATA FUNCTIONS ***/\n\n");
@@ -791,6 +872,7 @@ void UPCGCustomHLSLSettings::UpdateHelperDeclarations()
 		{
 			TArray<FString> DataCollectionPinNames;
 			TArray<FString> TextureDataPinNames;
+			TArray<FString> TextureArrayDataPinNames;
 
 			for (const FPCGPinProperties& Properties : InputPinProperties())
 			{
@@ -803,6 +885,11 @@ void UPCGCustomHLSLSettings::UpdateHelperDeclarations()
 				{
 					TextureDataPinNames.Add(Properties.Label.ToString());
 				}
+
+				if (!!(Properties.AllowedTypes & FPCGDataTypeInfoTexture2DArray::AsId()))
+				{
+					TextureArrayDataPinNames.Add(Properties.Label.ToString());
+				}
 			}
 
 			for (const FPCGPinProperties& Properties : OutputPinProperties())
@@ -816,6 +903,11 @@ void UPCGCustomHLSLSettings::UpdateHelperDeclarations()
 				{
 					TextureDataPinNames.Add(Properties.Label.ToString());
 				}
+
+				if (!!(Properties.AllowedTypes & FPCGDataTypeInfoTexture2DArray::AsId()))
+				{
+					TextureArrayDataPinNames.Add(Properties.Label.ToString());
+				}
 			}
 
 			if (!DataCollectionPinNames.IsEmpty())
@@ -853,6 +945,24 @@ void UPCGCustomHLSLSettings::UpdateHelperDeclarations()
 
 				HelperDeclarations += TEXT("\n");
 			}
+
+			if (!TextureArrayDataPinNames.IsEmpty())
+			{
+				HelperDeclarations += TEXT("\n// Returns false if thread has no data to operate on.\n");
+
+				const bool bMultiPin = TextureArrayDataPinNames.Num() > 1;
+
+				if (bMultiPin)
+				{
+					HelperDeclarations += TEXT("// Valid pins: ") + FString::Join(TextureArrayDataPinNames, TEXT(", ")) + TEXT("\n");
+				}
+
+				HelperDeclarations += FString::Format(
+					TEXT("bool {0}_GetThreadData(uint ThreadIndex, out uint OutDataIndex, out uint3 OutElementIndex);\n"),
+					{ bMultiPin ? PCGHLSLElement::PinDeclTemplateStr : TextureArrayDataPinNames[0] });
+
+				HelperDeclarations += TEXT("\n");
+			}
 		}
 
 		HelperDeclarations += TEXT(
@@ -938,12 +1048,10 @@ void UPCGCustomHLSLSettings::UpdatePinSettings()
 			}
 		}
 
-		if (!!(Properties.AllowedTypes & EPCGDataType::Landscape))
+		// Enforce single data on pin types that don't support a dynamic number of data.
+		if (Properties.AllowedTypes.IsChildOf(FPCGDataTypeInfoLandscape::AsId()) || Properties.AllowedTypes.IsChildOf(FPCGDataTypeInfoTexture2DArray::AsId()))
 		{
-			// Don't allow multiple data on this pin because we do not support a dynamic number of landscapes bound to a
-			// compute kernel.
 			Properties.bAllowMultipleData = false;
-
 			Properties.bAllowEditMultipleData = false;
 		}
 		else
@@ -992,16 +1100,29 @@ void UPCGCustomHLSLSettings::UpdatePinSettings()
 			Properties.bShowPropertiesGPU = !IsTextureKernel();
 		}
 
-		if (!!(Properties.AllowedTypes & EPCGDataType::BaseTexture))
+		// Enforce single data on pin types that don't support a dynamic number of data.
+		if (Properties.AllowedTypes.IsChildOf(FPCGDataTypeInfoBaseTexture2D::AsId()))
+		{
+			Properties.bAllowMultipleData = false;
+			Properties.bAllowEditMultipleData = false;
+			Properties.PropertiesGPU.bShowTexturePinSettings = true;
+
+			// Don't show texture array settings for non-array textures.
+			Properties.PropertiesGPU.bShowTextureArrayPinSettings = false;
+		}
+		else if (Properties.AllowedTypes.IsChildOf(FPCGDataTypeInfoTexture2DArray::AsId()))
 		{
-			// Texture outputs are single-data for now.
 			Properties.bAllowMultipleData = false;
 			Properties.bAllowEditMultipleData = false;
+			Properties.PropertiesGPU.bShowTextureArrayPinSettings = true;
+
+			// Also show normal texture settings for texture arrays.
 			Properties.PropertiesGPU.bShowTexturePinSettings = true;
 		}
 		else
 		{
 			Properties.PropertiesGPU.bShowTexturePinSettings = false;
+			Properties.PropertiesGPU.bShowTextureArrayPinSettings = false;
 		}
 
 		Properties.PropertiesGPU.bAllowEditInitMode = !bInitModeDrivenByKernel;
diff --git a/Engine/Plugins/PCG/Source/PCG/Private/Compute/Elements/PCGCustomHLSL.h b/Engine/Plugins/PCG/Source/PCG/Private/Compute/Elements/PCGCustomHLSL.h
--- a/Engine/Plugins/PCG/Source/PCG/Private/Compute/Elements/PCGCustomHLSL.h
+++ b/Engine/Plugins/PCG/Source/PCG/Private/Compute/Elements/PCGCustomHLSL.h
@@ -162,6 +162,12 @@ public:
 	UPROPERTY(EditAnywhere, Category = "Settings", meta = (EditCondition = "KernelType == EPCGKernelType::TextureGenerator", EditConditionHides, PCG_OverridableCPUAndGPUWithReadback))
 	EPCGTextureFilter TextureFilter = EPCGTextureFilter::Bilinear;
 
+	UPROPERTY(EditAnywhere, Category = "Settings", meta = (EditCondition = "KernelType == EPCGKernelType::TextureGenerator", EditConditionHides))
+	bool bOverrideTextureTransform = false;
+
+	UPROPERTY(EditAnywhere, Category = "Settings", meta = (EditCondition = "KernelType == EPCGKernelType::TextureGenerator && bOverrideTextureTransform", EditConditionHides))
+	FTransform TextureTransform = FTransform::Identity;
+
 	UPROPERTY(EditAnywhere, Category = "Settings|Thread Count", meta = (EditCondition = "KernelType == EPCGKernelType::Custom", EditConditionHides))
 	EPCGDispatchThreadCount DispatchThreadCount = EPCGDispatchThreadCount::FromFirstOutputPin;
 
diff --git a/Engine/Plugins/PCG/Source/PCG/Private/Compute/Elements/PCGCustomHLSLKernel.cpp b/Engine/Plugins/PCG/Source/PCG/Private/Compute/Elements/PCGCustomHLSLKernel.cpp
--- a/Engine/Plugins/PCG/Source/PCG/Private/Compute/Elements/PCGCustomHLSLKernel.cpp
+++ b/Engine/Plugins/PCG/Source/PCG/Private/Compute/Elements/PCGCustomHLSLKernel.cpp
@@ -10,6 +10,7 @@
 #include "Compute/PCGDataBinding.h"
 #include "Compute/Elements/PCGCustomHLSL.h"
 #include "Compute/DataInterfaces/Elements/PCGCustomHLSLDataInterface.h"
+#include "Data/PCGTexture2DArrayData.h"
 #include "Graph/PCGGPUGraphCompilationContext.h"
 
 #include "ComputeFramework/ComputeSource.h"
@@ -289,7 +290,22 @@ TSharedPtr<const FPCGDataCollectionDesc> UPCGCustomHLSLKernel::ComputeOutputBind
 		// Generators always produce a single texture data with known size.
 		OutDataDesc->GetDataDescriptionsMutable().Emplace(EPCGDataType::BaseTexture, FIntVector4(NumElementsX, NumElementsY, 0, 0));
 		OutDataDesc->SetRenderTargetFormatForAllData(TextureFormat);
-		OutDataDesc->SetTextureFilterForAllData(TextureFilter);
+
+		if (CustomHLSLSettings->bOverrideTextureTransform)
+		{
+			OutDataDesc->SetTextureTransformForAllData(CustomHLSLSettings->TextureTransform);
+		}
+		else
+		{
+			TSharedPtr<FPCGContextHandle> ContextHandle = InBinding->GetContextHandle().Pin();
+			if (FPCGContext* Context = ContextHandle ? ContextHandle->GetContext() : nullptr)
+			{
+				const FBox OriginalLocalBounds = Context->ExecutionSource->GetExecutionState().GetOriginalLocalSpaceBounds();
+				FTransform Transform = Context->ExecutionSource->GetExecutionState().GetOriginalTransform();
+				Transform.SetScale3D(Transform.GetScale3D() * 0.5 * (OriginalLocalBounds.Max - OriginalLocalBounds.Min));
+				OutDataDesc->SetTextureTransformForAllData(Transform);
+			}
+		}
 	}
 	else
 	{
@@ -868,7 +884,7 @@ void UPCGCustomHLSLKernel::ParseShaderSource()
 
 			InitializeKeywords.Add(PinStr + TEXT("_") + PCGCustomHLSLKernel::InitializeFunctionKeyword);
 		}
-		else if (!!(PinProps.AllowedTypes & EPCGDataType::BaseTexture))
+		else if (!!(PinProps.AllowedTypes & EPCGDataType::BaseTexture) || PinProps.AllowedTypes.IsChildOf(FPCGDataTypeInfoTexture2DArray::AsId()))
 		{
 			InitializeKeywords.Add(PinStr + TEXT("_") + PCGCustomHLSLKernel::StoreFunctionKeyword);
 		}
diff --git a/Engine/Plugins/PCG/Source/PCG/Private/Compute/PCGComputeCommon.cpp b/Engine/Plugins/PCG/Source/PCG/Private/Compute/PCGComputeCommon.cpp
--- a/Engine/Plugins/PCG/Source/PCG/Private/Compute/PCGComputeCommon.cpp
+++ b/Engine/Plugins/PCG/Source/PCG/Private/Compute/PCGComputeCommon.cpp
@@ -16,11 +16,13 @@
 #include "Compute/DataInterfaces/PCGLandscapeDataInterface.h"
 #include "Compute/DataInterfaces/PCGRawBufferDataInterface.h"
 #include "Compute/DataInterfaces/PCGStaticMeshDataInterface.h"
+#include "Compute/DataInterfaces/PCGTexture2DArrayDataInterface.h"
 #include "Compute/DataInterfaces/PCGTextureDataInterface.h"
 #include "Compute/DataInterfaces/PCGVirtualTextureDataInterface.h"
 #include "Compute/Packing/PCGDataCollectionPacking.h"
 #include "Data/PCGBasePointData.h"
 #include "Data/PCGTextureData.h"
+#include "Data/PCGTexture2DArrayData.h"
 #include "Graph/PCGGPUGraphCompilationContext.h"
 #include "Editor/IPCGEditorModule.h"
 
@@ -67,6 +69,13 @@ namespace PCGComputeHelpers
 			ElementCount.X = TextureSize.X;
 			ElementCount.Y = TextureSize.Y;
 		}
+		else if (const UPCGTexture2DArrayData* Texture2DArrayData = Cast<UPCGTexture2DArrayData>(InData))
+		{
+			const FIntPoint TextureResolution = Texture2DArrayData->GetResolution();
+			ElementCount.X = TextureResolution.X;
+			ElementCount.Y = TextureResolution.Y;
+			ElementCount.Z = Texture2DArrayData->GetArraySize();
+		}
 		else if (const UPCGRawBufferData* RawBufferData = Cast<UPCGRawBufferData>(InData))
 		{
 			ElementCount.X = RawBufferData->GetNumUint32s();
@@ -82,7 +91,30 @@ namespace PCGComputeHelpers
 
 	EPCGElementDimension GetElementDimension(const FPCGDataTypeIdentifier& InDataType)
 	{
-		return !!(InDataType & EPCGDataType::BaseTexture) ? EPCGElementDimension::Two : EPCGElementDimension::One;
+		if (InDataType.IsChildOf(EPCGDataType::BaseTexture))
+		{
+			return EPCGElementDimension::Two;
+		}
+		else if (InDataType.IsChildOf(FPCGDataTypeInfoTexture2DArray::AsId()))
+		{
+			return EPCGElementDimension::Three;
+		}
+
+		return EPCGElementDimension::One;
+	}
+
+	int32 GetMaxOutputDataCount(const FPCGDataTypeIdentifier& InDataType)
+	{
+		if (InDataType.IsChildOf(EPCGDataType::BaseTexture) || InDataType.IsChildOf(FPCGDataTypeInfoRawBuffer::AsId()))
+		{
+			// These data types can only emit a single data.
+			return 1;
+		}
+		else
+		{
+			// No limit.
+			return INT32_MAX;
+		}
 	}
 
 	const TArray<FPCGDataTypeIdentifier>& GetAllowedInputTypesList()
@@ -96,7 +128,8 @@ namespace PCGComputeHelpers
 			EPCGDataType::VirtualTexture,
 			EPCGDataType::StaticMeshResource,
 			EPCGDataType::ProxyForGPU,
-			FPCGDataTypeInfoRawBuffer::AsId()
+			FPCGDataTypeInfoRawBuffer::AsId(),
+			FPCGDataTypeInfoTexture2DArray::AsId()
 		};
 
 		return AllowedTypes;
@@ -110,7 +143,8 @@ namespace PCGComputeHelpers
 			EPCGDataType::Param,
 			EPCGDataType::ProxyForGPU,
 			EPCGDataType::BaseTexture,
-			FPCGDataTypeInfoRawBuffer::AsId()
+			FPCGDataTypeInfoRawBuffer::AsId(),
+			FPCGDataTypeInfoTexture2DArray::AsId()
 		};
 
 		return AllowedTypes;
@@ -368,6 +402,14 @@ namespace PCGComputeHelpers
 
 			DataInterface = TextureDataInterface;
 		}
+		else if (PinProperties.AllowedTypes.IsChildOf(FPCGDataTypeInfoTexture2DArray::AsId()))
+		{
+			UPCGTexture2DArrayDataInterface* TextureDataInterface = Context.NewObject_AnyThread<UPCGTexture2DArrayDataInterface>(InParams.ObjectOuter);
+			TextureDataInterface->SetRequiresExport(InParams.bRequiresExport);
+			TextureDataInterface->SetInitializeFromDataCollection(InParams.bProducedByCPU);
+
+			DataInterface = TextureDataInterface;
+		}
 		else if (PinProperties.AllowedTypes == EPCGDataType::Landscape)
 		{
 			DataInterface = Context.NewObject_AnyThread<UPCGLandscapeDataInterface>(InParams.ObjectOuter);
diff --git a/Engine/Plugins/PCG/Source/PCG/Private/Compute/PCGComputeGraph.cpp b/Engine/Plugins/PCG/Source/PCG/Private/Compute/PCGComputeGraph.cpp
--- a/Engine/Plugins/PCG/Source/PCG/Private/Compute/PCGComputeGraph.cpp
+++ b/Engine/Plugins/PCG/Source/PCG/Private/Compute/PCGComputeGraph.cpp
@@ -11,6 +11,8 @@
 #include "Compute/PCGPinPropertiesGPU.h"
 #include "Compute/Data/PCGProxyForGPUData.h"
 #include "Compute/Packing/PCGDataCollectionPacking.h"
+#include "Data/PCGTexture2DArrayData.h"
+#include "Data/PCGTextureData.h"
 #include "Metadata/PCGMetadata.h"
 
 #include "RenderGraphResources.h"
@@ -219,10 +221,12 @@ void UPCGComputeGraph::DebugLogDataDescriptions(const UPCGDataBinding* InBinding
 			UE_LOG(LogPCG, Display, TEXT("\t\t\tType: %s"), *Type.ToString());
 			UE_LOG(LogPCG, Display, TEXT("\t\t\tNum Elements: %s (%s)"), *DataDesc.GetElementCount().ToString(), *UEnum::GetValueAsString(DataDesc.GetElementDimension()));
 
-			if (Type.IsChildOf(FPCGDataTypeInfoBaseTexture2D::AsId()))
+			if (Type.IsChildOf(FPCGDataTypeInfoBaseTexture2D::AsId()) || Type.IsChildOf(FPCGDataTypeInfoTexture2DArray::AsId()))
 			{
 				UE_LOG(LogPCG, Display, TEXT("\t\t\tFormat: %s"), *UEnum::GetValueAsString(DataDesc.GetRenderTargetFormat()));
+				UE_LOG(LogPCG, Display, TEXT("\t\t\tTransform: %s"), *DataDesc.GetTextureTransform().ToString());
 				UE_LOG(LogPCG, Display, TEXT("\t\t\tFilter: %s"), *UEnum::GetValueAsString(DataDesc.GetTextureFilter()));
+				UE_LOG(LogPCG, Display, TEXT("\t\t\tNumMips: %d"), DataDesc.GetTextureNumMips());
 			}
 			else
 			{
diff --git a/Engine/Plugins/PCG/Source/PCG/Private/Compute/PCGComputeKernel.cpp b/Engine/Plugins/PCG/Source/PCG/Private/Compute/PCGComputeKernel.cpp
--- a/Engine/Plugins/PCG/Source/PCG/Private/Compute/PCGComputeKernel.cpp
+++ b/Engine/Plugins/PCG/Source/PCG/Private/Compute/PCGComputeKernel.cpp
@@ -13,6 +13,7 @@
 #include "Compute/PCGPinPropertiesGPU.h"
 #include "Compute/Elements/PCGComputeGraphElement.h"
 #include "Compute/Packing/PCGDataCollectionPacking.h"
+#include "Data/PCGTexture2DArrayData.h"
 #include "Graph/PCGGPUGraphCompilationContext.h"
 #include "Helpers/PCGSettingsHelpers.h"
 
@@ -256,6 +257,8 @@ void UPCGComputeKernel::ComputeDataDescFromPinProperties(const FPCGPinProperties
 
 	const FPCGPinPropertiesGPUStruct& Props = OutputPinProps.PropertiesGPU;
 
+	const int32 MaximumOutputDataCount = PCGComputeHelpers::GetMaxOutputDataCount(OutputPinProps.AllowedTypes);
+
 	if (Props.InitializationMode == EPCGPinInitMode::FromInputPins)
 	{
 		TArray<FName> InitPins;
@@ -411,16 +414,23 @@ void UPCGComputeKernel::ComputeDataDescFromPinProperties(const FPCGPinProperties
 
 					DataDesc.SetTagStringKeys(InputDataDesc.GetTagStringKeys());
 					DataDesc.AllocateProperties(InputDataDesc.GetAllocatedProperties());
+
+					if (OutPinDesc->GetDataDescriptions().Num() >= MaximumOutputDataCount)
+					{
+						break;
+					}
 				}
 			}
 			// Take pairs of datas, where the pairs are given by each data of each pin to each data of every other pin.
 			else if (Props.DataMultiplicity == EPCGDataMultiplicity::CartesianProduct)
 			{
-				for (int InputPinIndex = 0; InputPinIndex < NumInitPins; ++InputPinIndex)
+				bool bDone = false;
+
+				for (int InputPinIndex = 0; InputPinIndex < NumInitPins && !bDone; ++InputPinIndex)
 				{
 					const TSharedPtr<const FPCGDataCollectionDesc> InputDesc = InputDescs[InputPinIndex];
 
-					for (int OtherInputPinIndex = InputPinIndex + 1; OtherInputPinIndex < NumInitPins; ++OtherInputPinIndex)
+					for (int OtherInputPinIndex = InputPinIndex + 1; OtherInputPinIndex < NumInitPins && !bDone; ++OtherInputPinIndex)
 					{
 						const TSharedPtr<const FPCGDataCollectionDesc> OtherInputDesc = InputDescs[OtherInputPinIndex];
 
@@ -475,6 +485,17 @@ void UPCGComputeKernel::ComputeDataDescFromPinProperties(const FPCGPinProperties
 								}
 
 								DataDesc.AllocateProperties(InputDataDesc.GetAllocatedProperties() | OtherInputDataDesc.GetAllocatedProperties());
+
+								bDone = OutPinDesc->GetDataDescriptions().Num() >= MaximumOutputDataCount;
+								if (bDone)
+								{
+									break;
+								}
+							}
+
+							if (bDone)
+							{
+								break;
 							}
 						}
 					}
@@ -500,7 +521,7 @@ void UPCGComputeKernel::ComputeDataDescFromPinProperties(const FPCGPinProperties
 					MaxDataCount = FMath::Max(MaxDataCount, DataCount);
 				}
 
-				AddDataPairwise(MaxDataCount);
+				AddDataPairwise(FMath::Min(MaxDataCount, MaximumOutputDataCount));
 			}
 			else
 			{
@@ -509,7 +530,7 @@ void UPCGComputeKernel::ComputeDataDescFromPinProperties(const FPCGPinProperties
 		}
 		else if (Props.DataCountMode == EPCGDataCountMode::Fixed)
 		{
-			AddDataPairwise(Props.DataCount);
+			AddDataPairwise(FMath::Min(Props.DataCount, MaximumOutputDataCount));
 		}
 		else
 		{
@@ -523,9 +544,10 @@ void UPCGComputeKernel::ComputeDataDescFromPinProperties(const FPCGPinProperties
 			Desc.ScaleElementCount(Multiplier);
 		}
 
-		if (OutputPinProps.AllowedTypes.IsChildOf(EPCGDataType::BaseTexture))
+		if (OutputPinProps.AllowedTypes.IsChildOf(EPCGDataType::BaseTexture) || OutputPinProps.AllowedTypes.IsChildOf(FPCGDataTypeInfoTexture2DArray::AsId()))
 		{
 			// This code assumes texture output pins are single data.
+			ensure(MaximumOutputDataCount == 1);
 			const FPCGDataDesc* InputDesc = (!InputDescs.IsEmpty() && InputDescs[0] && !InputDescs[0]->GetDataDescriptions().IsEmpty()) ? &InputDescs[0]->GetDataDescriptions()[0] : nullptr;
 
 			if (OutputPinProps.PropertiesGPU.TextureFormatInheritanceMode == EPCGPropertyInheritanceMode::FromInput && InputDesc)
@@ -545,21 +567,60 @@ void UPCGComputeKernel::ComputeDataDescFromPinProperties(const FPCGPinProperties
 			{
 				OutPinDesc->SetTextureFilterForAllData(OutputPinProps.PropertiesGPU.TextureFilter);
 			}
+
+			if (OutputPinProps.PropertiesGPU.TextureArraySizeInheritanceMode == EPCGPropertyInheritanceMode::FromInput && InputDesc)
+			{
+				OutPinDesc->SetTextureArraySizeForAllData(InputDesc->GetTextureArraySize());
+			}
+			else
+			{
+				OutPinDesc->SetTextureArraySizeForAllData(OutputPinProps.PropertiesGPU.TextureArraySize);
+			}
+
+			if (OutputPinProps.PropertiesGPU.TextureTransformInheritanceMode == EPCGPropertyInheritanceMode::FromInput && InputDesc)
+			{
+				OutPinDesc->SetTextureTransformForAllData(InputDesc->GetTextureTransform());
+			}
+			else if (!OutputPinProps.PropertiesGPU.bInitializeTextureTransformFromGenerationVolume && OutputPinProps.PropertiesGPU.TextureTransformInheritanceMode == EPCGPropertyInheritanceMode::Custom)
+			{
+				OutPinDesc->SetTextureTransformForAllData(OutputPinProps.PropertiesGPU.TextureTransform);
+			}
+			else
+			{
+				TSharedPtr<FPCGContextHandle> ContextHandle = InBinding->GetContextHandle().Pin();
+				if (FPCGContext* Context = ContextHandle ? ContextHandle->GetContext() : nullptr)
+				{
+					const FBox OriginalLocalBounds = Context->ExecutionSource->GetExecutionState().GetOriginalLocalSpaceBounds();
+					FTransform Transform = Context->ExecutionSource->GetExecutionState().GetOriginalTransform();
+					Transform.SetScale3D(Transform.GetScale3D() * 0.5 * (OriginalLocalBounds.Max - OriginalLocalBounds.Min));
+					OutPinDesc->SetTextureTransformForAllData(Transform);
+				}
+			}
 		}
 	}
 	else if (Props.InitializationMode == EPCGPinInitMode::Custom)
 	{
-		for (int I = 0; I < Props.DataCount; ++I)
+		TArray<FPCGDataDesc>& DataDescriptions = OutPinDesc->GetDataDescriptionsMutable();
+
+		const int32 DataCount = FMath::Min(Props.DataCount, MaximumOutputDataCount);
+
+		for (int I = 0; I < DataCount; ++I)
 		{
+			FPCGDataDesc* DataDesc = nullptr;
+
 			const EPCGElementDimension ElementDimension = PCGComputeHelpers::GetElementDimension(OutputPinProps.AllowedTypes);
 
 			if (ElementDimension == EPCGElementDimension::One)
 			{
-				OutPinDesc->GetDataDescriptionsMutable().Emplace(OutputPinProps.AllowedTypes, Props.ElementCount);
+				DataDesc = &OutPinDesc->GetDataDescriptionsMutable().Emplace_GetRef(OutputPinProps.AllowedTypes, Props.ElementCount);
 			}
 			else if (ElementDimension == EPCGElementDimension::Two)
 			{
-				OutPinDesc->GetDataDescriptionsMutable().Emplace(OutputPinProps.AllowedTypes, Props.NumElements2D);
+				DataDesc = &OutPinDesc->GetDataDescriptionsMutable().Emplace_GetRef(OutputPinProps.AllowedTypes, Props.NumElements2D);
+			}
+			else if (ElementDimension == EPCGElementDimension::Three && OutputPinProps.AllowedTypes.IsChildOf(FPCGDataTypeInfoTexture2DArray::AsId()))
+			{
+				DataDesc = &OutPinDesc->GetDataDescriptionsMutable().Emplace_GetRef(OutputPinProps.AllowedTypes, FIntVector3(Props.NumElements2D.X, Props.NumElements2D.Y, Props.TextureArraySize));
 			}
 			else
 			{
@@ -568,10 +629,27 @@ void UPCGComputeKernel::ComputeDataDescFromPinProperties(const FPCGPinProperties
 			}
 		}
 
-		if (OutputPinProps.AllowedTypes.IsChildOf(EPCGDataType::BaseTexture))
+		if (OutputPinProps.AllowedTypes.IsChildOf(EPCGDataType::BaseTexture) || OutputPinProps.AllowedTypes.IsChildOf(FPCGDataTypeInfoTexture2DArray::AsId()))
 		{
 			OutPinDesc->SetRenderTargetFormatForAllData(OutputPinProps.PropertiesGPU.TextureFormat);
 			OutPinDesc->SetTextureFilterForAllData(OutputPinProps.PropertiesGPU.TextureFilter);
+			OutPinDesc->SetTextureArraySizeForAllData(OutputPinProps.PropertiesGPU.TextureArraySize);
+
+			if (OutputPinProps.PropertiesGPU.bInitializeTextureTransformFromGenerationVolume)
+			{
+				TSharedPtr<FPCGContextHandle> ContextHandle = InBinding->GetContextHandle().Pin();
+				if (FPCGContext* Context = ContextHandle ? ContextHandle->GetContext() : nullptr)
+				{
+					const FBox OriginalLocalBounds = Context->ExecutionSource->GetExecutionState().GetOriginalLocalSpaceBounds();
+					FTransform Transform = Context->ExecutionSource->GetExecutionState().GetOriginalTransform();
+					Transform.SetScale3D(Transform.GetScale3D() * 0.5 * (OriginalLocalBounds.Max - OriginalLocalBounds.Min));
+					OutPinDesc->SetTextureTransformForAllData(Transform);
+				}
+			}
+			else
+			{
+				OutPinDesc->SetTextureTransformForAllData(OutputPinProps.PropertiesGPU.TextureTransform);
+			}
 		}
 	}
 	else
diff --git a/Engine/Plugins/PCG/Source/PCG/Private/Compute/PCGDataDescription.cpp b/Engine/Plugins/PCG/Source/PCG/Private/Compute/PCGDataDescription.cpp
--- a/Engine/Plugins/PCG/Source/PCG/Private/Compute/PCGDataDescription.cpp
+++ b/Engine/Plugins/PCG/Source/PCG/Private/Compute/PCGDataDescription.cpp
@@ -6,6 +6,7 @@
 #include "Compute/PCGPinPropertiesGPU.h"
 #include "Compute/Packing/PCGDataCollectionPacking.h"
 #include "Data/PCGBasePointData.h"
+#include "Data/PCGTexture2DArrayData.h"
 #include "Elements/Metadata/PCGMetadataElementCommon.h"
 #include "Metadata/PCGMetadata.h"
 #include "Metadata/Accessors/PCGAttributeAccessorHelpers.h"
@@ -259,10 +260,20 @@ FPCGDataDesc::FPCGDataDesc(const FPCGTaggedData& InTaggedData, const UPCGDataBin
 	ElementCount = PCGComputeHelpers::GetElementCount(InTaggedData.Data);
 	ElementDimension = PCGComputeHelpers::GetElementDimension(InTaggedData.Data);
 
+	// @todo_pcg: Implement a way for a UPCGData to generate its data description.
+	if (const UPCGTexture2DArrayData* Texture2DArrayData = Cast<UPCGTexture2DArrayData>(InTaggedData.Data))
+	{
+		SetTextureArraySize(Texture2DArrayData->GetArraySize());
+		SetRenderTargetFormat(PCGComputeHelpers::GetPCGRenderTargetFormatFromPixelFormat(Texture2DArrayData->GetFormat()));
+		SetTextureFilter(Texture2DArrayData->GetFilter());
+		SetTextureTransform(Texture2DArrayData->GetTransform());
+	}
+
 	if (const UPCGBaseTextureData* TextureData = Cast<UPCGBaseTextureData>(InTaggedData.Data))
 	{
 		SetTextureFilter(TextureData->Filter);
 		SetRenderTargetFormat(PCGComputeHelpers::GetPCGRenderTargetFormatFromPixelFormat(TextureData->GetTextureFormat()));
+		SetTextureTransform(TextureData->GetTransform());
 	}
 
 	TagStringKeys.Reserve(InTaggedData.Tags.Num());
@@ -791,4 +802,20 @@ void FPCGDataCollectionDesc::SetTextureFilterForAllData(EPCGTextureFilter InText
 	}
 }
 
+void FPCGDataCollectionDesc::SetTextureArraySizeForAllData(uint16 InTextureArraySize)
+{
+	for (FPCGDataDesc& DataDesc : DataDescs)
+	{
+		DataDesc.SetTextureArraySize(InTextureArraySize);
+	}
+}
+
+void FPCGDataCollectionDesc::SetTextureTransformForAllData(const FTransform& InTransform)
+{
+	for (FPCGDataDesc& DataDesc : DataDescs)
+	{
+		DataDesc.SetTextureTransform(InTransform);
+	}
+}
+
 #undef LOCTEXT_NAMESPACE
diff --git a/Engine/Plugins/PCG/Source/PCG/Private/Data/PCGTexture2DArrayData.cpp b/Engine/Plugins/PCG/Source/PCG/Private/Data/PCGTexture2DArrayData.cpp
new file mode 100644
--- /dev/null
+++ b/Engine/Plugins/PCG/Source/PCG/Private/Data/PCGTexture2DArrayData.cpp
@@ -0,0 +1,106 @@
+// Copyright Epic Games, Inc. All Rights Reserved.
+
+#include "Data/PCGTexture2DArrayData.h"
+
+#include "PCGContext.h"
+#include "Data/PCGBasePointData.h"
+#include "Data/PCGPointArrayData.h"
+#include "Data/PCGPointData.h"
+
+#include "TextureResource.h"
+#include "Engine/Texture.h"
+
+#include UE_INLINE_GENERATED_CPP_BY_NAME(PCGTexture2DArrayData)
+
+PCG_DEFINE_TYPE_INFO(FPCGDataTypeInfoTexture2DArray, UPCGTexture2DArrayData)
+
+UPCGSpatialData* UPCGTexture2DArrayData::CopyInternal(FPCGContext* Context) const
+{
+	UPCGTexture2DArrayData* NewData = FPCGContext::NewObject_AnyThread<UPCGTexture2DArrayData>(Context);
+
+	CopyTexture2DArrayData(NewData);
+
+	return NewData;
+}
+
+const UPCGPointData* UPCGTexture2DArrayData::CreatePointData(FPCGContext* InContext) const
+{
+	return CastChecked<UPCGPointData>(CreateBasePointData(InContext, UPCGPointData::StaticClass()));
+}
+
+const UPCGPointArrayData* UPCGTexture2DArrayData::CreatePointArrayData(FPCGContext* InContext, const FBox& InBounds) const
+{
+	return CastChecked<UPCGPointArrayData>(CreateBasePointData(InContext, UPCGPointArrayData::StaticClass()));
+}
+
+bool UPCGTexture2DArrayData::Initialize(UTexture* InTexture, const FPCGTexture2DArrayDataInitParams& InParams)
+{
+	Texture = InTexture;
+	ResourceType = EPCGTextureResourceType::TextureObject;
+
+	return InitializeInternal(InParams);
+}
+
+bool UPCGTexture2DArrayData::Initialize(TRefCountPtr<IPooledRenderTarget> InExportedTextureHandle, const FPCGTexture2DArrayDataInitParams& InParams)
+{
+	ExportedTextureHandle = InExportedTextureHandle;
+	ResourceType = EPCGTextureResourceType::ExportedTexture;
+
+	return InitializeInternal(InParams);
+}
+
+FTextureRHIRef UPCGTexture2DArrayData::GetRHI() const
+{
+	if (ResourceType == EPCGTextureResourceType::TextureObject)
+	{
+		FTextureResource* Resource = Texture.IsValid() ? Texture->GetResource() : nullptr;
+		return Resource ? Resource->GetTextureRHI() : nullptr;
+	}
+	else
+	{
+		return ExportedTextureHandle ? ExportedTextureHandle->GetRHI() : nullptr;
+	}
+}
+
+bool UPCGTexture2DArrayData::InitializeInternal(const FPCGTexture2DArrayDataInitParams& InParams)
+{
+	if (ensure(GetRHI()))
+	{
+		const FRHITextureDesc& Desc = GetRHI()->GetDesc();
+		Resolution = Desc.Extent;
+		ArraySize = Desc.ArraySize;
+		Format = Desc.Format;
+	}
+
+	Transform = InParams.Transform;
+	Filter = InParams.Filter;
+
+	Bounds = FBox(EForceInit::ForceInit);
+	Bounds += FVector(-1.0f, -1.0f, 0.0f);
+	Bounds += FVector(1.0f, 1.0f, 0.0f);
+	Bounds = Bounds.TransformBy(Transform);
+
+	return true;
+}
+
+const UPCGBasePointData* UPCGTexture2DArrayData::CreateBasePointData(FPCGContext* InContext, TSubclassOf<UPCGBasePointData> InPointDataClass) const
+{
+	UPCGBasePointData* Data = FPCGContext::NewObject_AnyThread<UPCGBasePointData>(InContext, GetTransientPackage(), InPointDataClass);
+	Data->InitializeFromData(this);
+
+	return Data;
+}
+
+void UPCGTexture2DArrayData::CopyTexture2DArrayData(UPCGTexture2DArrayData* NewTextureData) const
+{
+	CopyBaseSurfaceData(NewTextureData);
+
+	NewTextureData->Texture = Texture;
+	NewTextureData->ExportedTextureHandle = ExportedTextureHandle;
+	NewTextureData->ResourceType = ResourceType;
+	NewTextureData->Resolution = Resolution;
+	NewTextureData->ArraySize = ArraySize;
+	NewTextureData->Format = Format;
+	NewTextureData->Filter = Filter;
+	NewTextureData->Bounds = Bounds;
+}
diff --git a/Engine/Plugins/PCG/Source/PCG/Private/Elements/PCGGetTexture2DArrayData.cpp b/Engine/Plugins/PCG/Source/PCG/Private/Elements/PCGGetTexture2DArrayData.cpp
new file mode 100644
--- /dev/null
+++ b/Engine/Plugins/PCG/Source/PCG/Private/Elements/PCGGetTexture2DArrayData.cpp
@@ -0,0 +1,145 @@
+// Copyright Epic Games, Inc. All Rights Reserved.
+
+#include "Elements/PCGGetTexture2DArrayData.h"
+
+#include "PCGContext.h"
+#include "PCGModule.h"
+#include "Data/PCGTexture2DArrayData.h"
+
+#if WITH_EDITOR
+#include "Helpers/PCGDynamicTrackingHelpers.h"
+#endif // WITH_EDITOR
+
+#include "Engine/Texture2DArray.h"
+#include "Engine/TextureRenderTarget2DArray.h"
+
+#include UE_INLINE_GENERATED_CPP_BY_NAME(PCGGetTexture2DArrayData)
+
+#define LOCTEXT_NAMESPACE "PCGGetTexture2DArrayDataElement"
+
+#if WITH_EDITOR
+void UPCGGetTexture2DArrayDataSettings::GetStaticTrackedKeys(FPCGSelectionKeyToSettingsMap& OutKeysToSettings, TArray<TObjectPtr<const UPCGGraph>>& OutVisitedGraphs) const
+{
+	if (IsPropertyOverriddenByPin(GET_MEMBER_NAME_CHECKED(UPCGGetTexture2DArrayDataSettings, Texture)) || Texture.IsNull())
+	{
+		// Dynamic tracking or null settings
+		return;
+	}
+
+	FPCGSelectionKey Key = FPCGSelectionKey::CreateFromPath(Texture.ToSoftObjectPath());
+
+	OutKeysToSettings.FindOrAdd(Key).Emplace(this, /*bCulling=*/false);
+}
+#endif // WITH_EDITOR
+
+TArray<FPCGPinProperties> UPCGGetTexture2DArrayDataSettings::OutputPinProperties() const
+{
+	TArray<FPCGPinProperties> Properties;
+	FPCGPinProperties& Pin = Properties.Emplace_GetRef(PCGPinConstants::DefaultOutputLabel, FPCGDataTypeInfoTexture2DArray::AsId());
+	Pin.bAllowMultipleData = false;
+
+	return Properties;
+}
+
+FPCGElementPtr UPCGGetTexture2DArrayDataSettings::CreateElement() const
+{
+	return MakeShared<FPCGGetTexture2DArrayDataElement>();
+}
+
+bool FPCGGetTexture2DArrayDataElement::CanExecuteOnlyOnMainThread(FPCGContext* InContext) const
+{
+	// PrepareData does loading, which requires the game thread.
+	return !InContext || InContext->CurrentPhase == EPCGExecutionPhase::PrepareData;
+}
+
+void FPCGGetTexture2DArrayDataElement::GetDependenciesCrc(const FPCGGetDependenciesCrcParams& InParams, FPCGCrc& OutCrc) const
+{
+	FPCGCrc Crc;
+	IPCGElement::GetDependenciesCrc(InParams, Crc);
+
+	if (InParams.ExecutionSource)
+	{
+		if (const UPCGData* Data = InParams.ExecutionSource->GetExecutionState().GetSelfData())
+		{
+			Crc.Combine(Data->GetOrComputeCrc(/*bFullDataCrc=*/false));
+		}
+	}
+
+	OutCrc = Crc;
+}
+
+bool FPCGGetTexture2DArrayDataElement::PrepareDataInternal(FPCGContext* InContext) const
+{
+	TRACE_CPUPROFILER_EVENT_SCOPE(FPCGGetTexture2DArrayDataElement::PrepareDataInternal);
+	check(InContext);
+
+	FPCGGetTexture2DArrayDataContext* Context = static_cast<FPCGGetTexture2DArrayDataContext*>(InContext);
+
+	const UPCGGetTexture2DArrayDataSettings* Settings = Context->GetInputSettings<UPCGGetTexture2DArrayDataSettings>();
+	check(Settings);
+
+	if (Settings->Texture.IsNull())
+	{
+		return true;
+	}
+
+	if (!Context->WasLoadRequested())
+	{
+		return Context->RequestResourceLoad(Context, { Settings->Texture.ToSoftObjectPath() }, !Settings->bSynchronousLoad);
+	}
+
+	return true;
+}
+
+bool FPCGGetTexture2DArrayDataElement::ExecuteInternal(FPCGContext* Context) const
+{
+	TRACE_CPUPROFILER_EVENT_SCOPE(FPCGGetTexture2DArrayDataElement::Execute);
+	check(Context);
+
+	const UPCGGetTexture2DArrayDataSettings* Settings = Context->GetInputSettings<UPCGGetTexture2DArrayDataSettings>();
+	check(Settings);
+
+	if (Settings->Texture.IsNull())
+	{
+		return true;
+	}
+
+	UTexture* Texture = Settings->Texture.Get();
+
+	if (!Texture)
+	{
+		PCGE_LOG(Error, GraphAndLog, FText::Format(LOCTEXT("CouldNotResolveTexture", "Texture at path '{0}' could not be loaded."), FText::FromString(Settings->Texture.ToString())));
+		return true;
+	}
+
+	if (!Texture->IsA<UTexture2DArray>() && !Texture->IsA<UTextureRenderTarget2DArray>())
+	{
+		PCGE_LOG(Error, GraphAndLog, FText::Format(LOCTEXT("InvalidType", "Texture at path '{0}' was not a valid type. Must be UTexture2DArray or UTextureRenderTarget2DArray."), FText::FromString(Settings->Texture.ToString())));
+		return true;
+	}
+
+	FTransform Transform = Context->ExecutionSource->GetExecutionState().GetOriginalTransform();
+	const FBox OriginalActorLocalBounds = Context->ExecutionSource->GetExecutionState().GetOriginalLocalSpaceBounds();
+	Transform.SetScale3D(Transform.GetScale3D() * 0.5 * (OriginalActorLocalBounds.Max - OriginalActorLocalBounds.Min));
+
+	FPCGTexture2DArrayDataInitParams InitParams;
+	InitParams.Filter = Settings->Filter;
+	InitParams.Transform = MoveTemp(Transform);
+
+	UPCGTexture2DArrayData* Texture2DArrayData = FPCGContext::NewObject_AnyThread<UPCGTexture2DArrayData>(Context);
+	Texture2DArrayData->Initialize(Texture, InitParams);
+
+	Context->OutputData.TaggedData.Emplace_GetRef().Data = Texture2DArrayData;
+
+#if WITH_EDITOR
+	// If we have an override, register for dynamic tracking.
+	if (Context->IsValueOverriden(GET_MEMBER_NAME_CHECKED(UPCGGetTexture2DArrayDataSettings, Texture)))
+	{
+		FPCGDynamicTrackingHelper::AddSingleDynamicTrackingKey(Context, FPCGSelectionKey::CreateFromPath(Texture), /*bIsCulled=*/false);
+	}
+#endif // WITH_EDITOR
+
+	return true;
+}
+
+#undef LOCTEXT_NAMESPACE
diff --git a/Engine/Plugins/PCG/Source/PCG/Public/Compute/PCGComputeCommon.h b/Engine/Plugins/PCG/Source/PCG/Public/Compute/PCGComputeCommon.h
--- a/Engine/Plugins/PCG/Source/PCG/Public/Compute/PCGComputeCommon.h
+++ b/Engine/Plugins/PCG/Source/PCG/Public/Compute/PCGComputeCommon.h
@@ -104,6 +104,9 @@ namespace PCGComputeHelpers
 	EPCGElementDimension GetElementDimension(const UPCGData* InData);
 	EPCGElementDimension GetElementDimension(const FPCGDataTypeIdentifier& InDataType);
 
+	/** Maximum number of data that a GPU node pin of the given type can receive/emit. */
+	int32 GetMaxOutputDataCount(const FPCGDataTypeIdentifier& InDataType);
+
 	/** PCG data types supported in GPU node inputs. */
 	const TArray<FPCGDataTypeIdentifier>& GetAllowedInputTypesList();
 
diff --git a/Engine/Plugins/PCG/Source/PCG/Public/Compute/PCGDataDescription.h b/Engine/Plugins/PCG/Source/PCG/Public/Compute/PCGDataDescription.h
--- a/Engine/Plugins/PCG/Source/PCG/Public/Compute/PCGDataDescription.h
+++ b/Engine/Plugins/PCG/Source/PCG/Public/Compute/PCGDataDescription.h
@@ -191,6 +191,15 @@ public:
 	bool IsAttributeAllocated(const FPCGKernelAttributeDesc& InAttributeDesc) const { return IsAttributeAllocated(InAttributeDesc.GetAttributeId()); }
 	PCG_API bool IsAttributeAllocated(int32 InAttributeId) const;
 
+	uint16 GetTextureArraySize() const { return TextureArraySize; }
+	void SetTextureArraySize(uint16 InTextureArraySize) { TextureArraySize = InTextureArraySize; }
+
+	const FTransform& GetTextureTransform() const { return TextureTransform; }
+	void SetTextureTransform(const FTransform& InTransform) { TextureTransform = InTransform; }
+
+	// There's no store shader function to write to particular mips so we always generate textures with 1 mip for now.
+	const int32 GetTextureNumMips() const { return 1; }
+
 	TConstArrayView<int32> GetTagStringKeys() const { return MakeConstArrayView(TagStringKeys); }
 	TArrayView<int32> GetTagStringKeysMutable();
 	void SetTagStringKeys(TConstArrayView<int32> InTagStringKeys) { TagStringKeys = InTagStringKeys; }
@@ -220,6 +229,8 @@ private:
 	/** Texture format used to create a render target UAV from this data. If this description represents a CPU texture data, then this format will be a narrowed conversion from its EPixelFormat. */
 	EPCGRenderTargetFormat RenderTargetFormat = EPCGRenderTargetFormat::RGBA16f;
 	EPCGTextureFilter TextureFilter = EPCGTextureFilter::Bilinear;
+	FTransform TextureTransform = FTransform::Identity;
+	uint16 TextureArraySize = 0;
 };
 
 /** Data description for a data collection (FPCGDataCollection). */
@@ -262,6 +273,10 @@ struct FPCGDataCollectionDesc
 
 	PCG_API void SetTextureFilterForAllData(EPCGTextureFilter InTextureFilter);
 
+	PCG_API void SetTextureArraySizeForAllData(uint16 InTextureArraySize);
+
+	PCG_API void SetTextureTransformForAllData(const FTransform& InTransform);
+
 private:
 	FPCGDataCollectionDesc() = default;
 
diff --git a/Engine/Plugins/PCG/Source/PCG/Public/Compute/PCGPinPropertiesGPU.h b/Engine/Plugins/PCG/Source/PCG/Public/Compute/PCGPinPropertiesGPU.h
--- a/Engine/Plugins/PCG/Source/PCG/Public/Compute/PCGPinPropertiesGPU.h
+++ b/Engine/Plugins/PCG/Source/PCG/Public/Compute/PCGPinPropertiesGPU.h
@@ -119,6 +119,16 @@ struct FPCGPinPropertiesGPUStruct
 	UPROPERTY(BlueprintReadWrite, EditAnywhere, Category = Settings, meta = (EditCondition = "bAllowEditInitMode && (InitializationMode == EPCGPinInitMode::Custom || TextureFormatInheritanceMode == EPCGPropertyInheritanceMode::Custom) && bShowTexturePinSettings", EditConditionHides, HideEditConditionToggle))
 	EPCGRenderTargetFormat TextureFormat = EPCGRenderTargetFormat::RGBA16f;
 
+	UPROPERTY(BlueprintReadWrite, EditAnywhere, Category = Settings, meta = (EditCondition = "bAllowEditInitMode && bShowTexturePinSettings", EditConditionHides, HideEditConditionToggle))
+	bool bInitializeTextureTransformFromGenerationVolume = false;
+
+	/** Method for inheriting texture transform. */
+	UPROPERTY(BlueprintReadWrite, EditAnywhere, Category = Settings, meta = (EditCondition = "bAllowEditInitMode && bShowTexturePinSettings && !bInitializeTextureTransformFromGenerationVolume", EditConditionHides, HideEditConditionToggle))
+	EPCGPropertyInheritanceMode TextureTransformInheritanceMode = EPCGPropertyInheritanceMode::FromInput;
+
+	UPROPERTY(BlueprintReadWrite, EditAnywhere, Category = Settings, meta = (EditCondition = "bAllowEditInitMode && (InitializationMode == EPCGPinInitMode::Custom || TextureTransformInheritanceMode == EPCGPropertyInheritanceMode::Custom) && bShowTexturePinSettings && !bInitializeTextureTransformFromGenerationVolume", EditConditionHides, HideEditConditionToggle))
+	FTransform TextureTransform = FTransform::Identity;
+
 	/** Method for inheriting texture filter. */
 	UPROPERTY(BlueprintReadWrite, EditAnywhere, Category = Settings, meta = (EditCondition = "bAllowEditInitMode && bShowTexturePinSettings && InitializationMode == EPCGPinInitMode::FromInputPins", EditConditionHides, HideEditConditionToggle))
 	EPCGPropertyInheritanceMode TextureFilterInheritanceMode = EPCGPropertyInheritanceMode::FromInput;
@@ -127,6 +137,14 @@ struct FPCGPinPropertiesGPUStruct
 	UPROPERTY(BlueprintReadWrite, EditAnywhere, Category = Settings, meta = (EditCondition = "bAllowEditInitMode && (InitializationMode == EPCGPinInitMode::Custom || TextureFilterInheritanceMode == EPCGPropertyInheritanceMode::Custom) && bShowTexturePinSettings", EditConditionHides, HideEditConditionToggle))
 	EPCGTextureFilter TextureFilter = EPCGTextureFilter::Bilinear;
 
+	/** Method for inheriting texture array size. */
+	UPROPERTY(BlueprintReadWrite, EditAnywhere, Category = Settings, meta = (EditCondition = "bAllowEditInitMode && bShowTextureArrayPinSettings && InitializationMode == EPCGPinInitMode::FromInputPins", EditConditionHides, HideEditConditionToggle))
+	EPCGPropertyInheritanceMode TextureArraySizeInheritanceMode = EPCGPropertyInheritanceMode::FromInput;
+
+	/** Number of layers in the texture array created by this pin. */
+	UPROPERTY(BlueprintReadWrite, EditAnywhere, Category = Settings, meta = (EditCondition = "bAllowEditInitMode && (InitializationMode == EPCGPinInitMode::Custom || TextureArraySizeInheritanceMode == EPCGPropertyInheritanceMode::Custom) && bShowTextureArrayPinSettings", EditConditionHides, ClampMin = 1))
+	int TextureArraySize = 1;
+
 	/** Scalar on the fixed number of elements to create in each output data. */
 	UPROPERTY(BlueprintReadWrite, EditAnywhere, Category = Settings, meta = (EditCondition = "bAllowEditInitMode && (InitializationMode == EPCGPinInitMode::Custom || ElementCountMode == EPCGElementCountMode::FromInputData)", EditConditionHides, ClampMin = 1))
 	int ElementCountMultiplier = 1;
@@ -157,6 +175,9 @@ struct FPCGPinPropertiesGPUStruct
 	UPROPERTY(Transient)
 	bool bShowTexturePinSettings = false;
 
+	UPROPERTY(Transient)
+	bool bShowTextureArrayPinSettings = false;
+
 	UPROPERTY(Transient)
 	bool bShowInitializeToZeroSetting = false;
 #endif // WITH_EDITORONLY_DATA
diff --git a/Engine/Plugins/PCG/Source/PCG/Public/Data/PCGTexture2DArrayData.h b/Engine/Plugins/PCG/Source/PCG/Public/Data/PCGTexture2DArrayData.h
new file mode 100644
--- /dev/null
+++ b/Engine/Plugins/PCG/Source/PCG/Public/Data/PCGTexture2DArrayData.h
@@ -0,0 +1,107 @@
+// Copyright Epic Games, Inc. All Rights Reserved.
+
+#pragma once
+
+#include "Data/PCGSurfaceData.h"
+
+#include "PCGCommon.h"
+
+#include "RendererInterface.h"
+#include "RHI.h"
+
+#include "PCGTexture2DArrayData.generated.h"
+
+class UPCGBasePointData;
+class UPCGMetadata;
+class UPCGPointArrayData;
+class UPCGPointData;
+class UTexture;
+struct FPCGContext;
+struct FPCGPoint;
+struct IPooledRenderTarget;
+
+USTRUCT()
+struct FPCGDataTypeInfoTexture2DArray : public FPCGDataTypeInfoSurface
+{
+	GENERATED_BODY()
+
+	PCG_DECLARE_TYPE_INFO()
+};
+
+struct FPCGTexture2DArrayDataInitParams
+{
+	FTransform Transform = FTransform::Identity;
+	EPCGTextureFilter Filter = EPCGTextureFilter::Bilinear;
+};
+
+UCLASS(MinimalAPI)
+class UPCGTexture2DArrayData : public UPCGSurfaceData
+{
+	GENERATED_BODY()
+
+public:
+	// ~Begin UPCGData interface
+	PCG_ASSIGN_TYPE_INFO(FPCGDataTypeInfoTexture2DArray);
+	virtual bool IsCacheable() const { return Super::IsCacheable() && ResourceType != EPCGTextureResourceType::ExportedTexture; }
+	virtual void ReleaseTransientResources(const TCHAR* InReason = nullptr) override { ExportedTextureHandle = nullptr; }
+	// ~End UPCGData interface
+
+	//~Begin UPCGSpatialData interface
+	virtual FBox GetBounds() const override { return Bounds; }
+	virtual FBox GetStrictBounds() const override { return Bounds; }
+	virtual bool SamplePoint(const FTransform& InTransform, const FBox& InBounds, FPCGPoint& OutPoint, UPCGMetadata* OutMetadata) const override { return false; }
+
+protected:
+	virtual UPCGSpatialData* CopyInternal(FPCGContext* Context) const override;
+	//~End UPCGSpatialData interface
+
+public:
+	//~Begin UPCGSpatialDataWithPointCache interface
+	PCG_API virtual const UPCGPointData* CreatePointData(FPCGContext* InContext) const override;
+	PCG_API virtual const UPCGPointArrayData* CreatePointArrayData(FPCGContext* InContext, const FBox& InBounds) const override;
+	//~End UPCGSpatialDataWithPointCache interface
+
+	PCG_API bool Initialize(UTexture* InTexture, const FPCGTexture2DArrayDataInitParams& InTransform);
+	PCG_API bool Initialize(TRefCountPtr<IPooledRenderTarget> InExportedTextureHandle, const FPCGTexture2DArrayDataInitParams& InTransform);
+	UTexture* GetTexture() const { return Texture.Get(); }
+	PCG_API FTextureRHIRef GetRHI() const;
+	EPCGTextureResourceType GetResourceType() const { return ResourceType; }
+	TRefCountPtr<IPooledRenderTarget> GetRefCountedTexture() const { return ExportedTextureHandle; }
+	FIntPoint GetResolution() const { return Resolution; }
+	uint16 GetArraySize() const { return ArraySize; }
+	EPixelFormat GetFormat() const { return Format; }
+	EPCGTextureFilter GetFilter() const { return Filter; }
+
+protected:
+	PCG_API bool InitializeInternal(const FPCGTexture2DArrayDataInitParams& InParams);
+	PCG_API const UPCGBasePointData* CreateBasePointData(FPCGContext* InContext, TSubclassOf<UPCGBasePointData> InPointDataClass) const;
+	PCG_API void CopyTexture2DArrayData(UPCGTexture2DArrayData* NewTextureData) const;
+
+protected:
+	UPROPERTY()
+	TWeakObjectPtr<UTexture> Texture = nullptr;
+
+	/** If initialized from an exported texture this holds a reference to the resource. */
+	TRefCountPtr<IPooledRenderTarget> ExportedTextureHandle = nullptr;
+
+	/** The type of underlying resource that this texture data represents. */
+	UPROPERTY()
+	EPCGTextureResourceType ResourceType = EPCGTextureResourceType::TextureObject;
+
+	UPROPERTY()
+	FIntPoint Resolution = FIntPoint::ZeroValue;
+
+	UPROPERTY()
+	uint16 ArraySize = 0;
+
+	/** Format of the underlying texture resource. */
+	UPROPERTY()
+	TEnumAsByte<EPixelFormat> Format = EPixelFormat::PF_Unknown;
+
+	/** Method used to determine the value for a sample based on the value of nearby texels. */
+	UPROPERTY()
+	EPCGTextureFilter Filter = EPCGTextureFilter::Bilinear;
+
+	UPROPERTY()
+	FBox Bounds = FBox(EForceInit::ForceInit);
+};
diff --git a/Engine/Plugins/PCG/Source/PCG/Public/Data/PCGTextureData.h b/Engine/Plugins/PCG/Source/PCG/Public/Data/PCGTextureData.h
--- a/Engine/Plugins/PCG/Source/PCG/Public/Data/PCGTextureData.h
+++ b/Engine/Plugins/PCG/Source/PCG/Public/Data/PCGTextureData.h
@@ -4,6 +4,8 @@
 
 #include "Data/PCGSurfaceData.h"
 
+#include "PCGCommon.h"
+
 #include "PixelFormat.h"
 #include "RendererInterface.h"
 #include "RHI.h"
@@ -32,14 +34,6 @@ enum class UE_DEPRECATED(5.5, "EPCGTextureDensityFunction has been deprecated.")
 	Multiply
 };
 
-/** Method used to determine the value for a sample based on the value of nearby texels. */
-UENUM(BlueprintType)
-enum class EPCGTextureFilter : uint8
-{
-	Point UMETA(Tooltip="Takes the value of whatever texel the sample lands in."),
-	Bilinear UMETA(Tooltip="Bilinearly interpolates the values of the four nearest texels to the sample location.")
-};
-
 UENUM()
 enum class EPCGTextureAddressMode : uint8
 {
@@ -47,14 +41,6 @@ enum class EPCGTextureAddressMode : uint8
 	Wrap UMETA(ToolTip = "Tiles the texture to fit.")
 };
 
-UENUM()
-enum class EPCGTextureResourceType : uint8
-{
-	TextureObject UMETA(ToolTip = "UObject texture such as UTexture2D or UTextureRenderTarget2D."),
-	ExportedTexture UMETA(ToolTip = "Texture handle exported from a texture on the GPU."),
-	Invalid UMETA(Hidden)
-};
-
 namespace PCGTextureSamplingHelpers
 {
 	/** Returns true if a texture is CPU-accessible. */
@@ -232,7 +218,7 @@ public:
 	//~Begin UPCGBaseTextureData interface
 	virtual UTexture* GetTexture() const override { return Texture.IsValid() ? Texture.Get() : nullptr; }
 	virtual FTextureRHIRef GetTextureRHI() const override;
-	virtual TRefCountPtr<IPooledRenderTarget> GetRefCountedTexture() const { return TextureHandle; }
+	virtual TRefCountPtr<IPooledRenderTarget> GetRefCountedTexture() const override { return TextureHandle; }
 	virtual int GetTextureSlice() const override { return TextureIndex; }
 	//~End UPCGBaseTextureData interface
 	
diff --git a/Engine/Plugins/PCG/Source/PCG/Public/Elements/PCGGetTexture2DArrayData.h b/Engine/Plugins/PCG/Source/PCG/Public/Elements/PCGGetTexture2DArrayData.h
new file mode 100644
--- /dev/null
+++ b/Engine/Plugins/PCG/Source/PCG/Public/Elements/PCGGetTexture2DArrayData.h
@@ -0,0 +1,62 @@
+// Copyright Epic Games, Inc. All Rights Reserved.
+
+#pragma once
+
+#include "PCGContext.h"
+#include "PCGSettings.h"
+
+#include "Async/PCGAsyncLoadingContext.h"
+#include "Data/PCGTexture2DArrayData.h"
+
+#include "PCGGetTexture2DArrayData.generated.h"
+
+class UTexture;
+
+/** Creates texture 2D array data from a soft object path. */
+UCLASS(MinimalAPI, BlueprintType, ClassGroup = (Procedural))
+class UPCGGetTexture2DArrayDataSettings : public UPCGSettings
+{
+	GENERATED_BODY()
+
+public:
+	//~ Begin UPCGSettings interface
+#if WITH_EDITOR
+	virtual FName GetDefaultNodeName() const override { return FName(TEXT("GetTexture2DArrayData")); }
+	virtual FText GetDefaultNodeTitle() const override { return NSLOCTEXT("PCGGetTexture2DArrayDataElement", "NodeTitle", "Get Texture 2D Array Data"); }
+	virtual FText GetNodeTooltipText() const override { return NSLOCTEXT("PCGGetTexture2DArrayDataElement", "NodeTooltip", "Creates texture 2D array data from a soft object path."); }
+	virtual EPCGSettingsType GetType() const override { return EPCGSettingsType::Spatial; }
+	virtual void GetStaticTrackedKeys(FPCGSelectionKeyToSettingsMap& OutKeysToSettings, TArray<TObjectPtr<const UPCGGraph>>& OutVisitedGraphs) const override;
+	virtual bool CanDynamicallyTrackKeys() const override { return true; }
+#endif
+
+protected:
+	virtual TArray<FPCGPinProperties> InputPinProperties() const override { return TArray<FPCGPinProperties>(); }
+	virtual TArray<FPCGPinProperties> OutputPinProperties() const override;
+	virtual FPCGElementPtr CreateElement() const override;
+	//~ End UPCGSettings interface
+
+public:
+	UPROPERTY(BlueprintReadWrite, EditAnywhere, Category = Settings, meta = (EditCondition = "!bOverrideFromInput", AllowedClasses = "/Script/Engine.Texture2DArray, /Script/Engine.TextureRenderTarget2DArray", PCG_Overridable))
+	TSoftObjectPtr<UTexture> Texture;
+
+	/** Method used to determine the value for a sample based on the value of nearby texels. */
+	UPROPERTY(BlueprintReadWrite, EditAnywhere, Category = Settings, meta = (PCG_Overridable))
+	EPCGTextureFilter Filter = EPCGTextureFilter::Bilinear;
+
+	/** By default, texture loading is asynchronous, can force it synchronous if needed. */
+	UPROPERTY(BlueprintReadWrite, EditAnywhere, Category = "Settings|Debug")
+	bool bSynchronousLoad = false;
+};
+
+struct FPCGGetTexture2DArrayDataContext : public FPCGContext, public IPCGAsyncLoadingContext {};
+
+class FPCGGetTexture2DArrayDataElement : public IPCGElementWithCustomContext<FPCGGetTexture2DArrayDataContext>
+{
+public:
+	virtual bool CanExecuteOnlyOnMainThread(FPCGContext* InContext) const override;
+	virtual void GetDependenciesCrc(const FPCGGetDependenciesCrcParams& InParams, FPCGCrc& OutCrc) const override;
+
+protected:
+	virtual bool PrepareDataInternal(FPCGContext* InContext) const override;
+	virtual bool ExecuteInternal(FPCGContext* InContext) const override;
+};
diff --git a/Engine/Plugins/PCG/Source/PCG/Public/PCGCommon.h b/Engine/Plugins/PCG/Source/PCG/Public/PCGCommon.h
--- a/Engine/Plugins/PCG/Source/PCG/Public/PCGCommon.h
+++ b/Engine/Plugins/PCG/Source/PCG/Public/PCGCommon.h
@@ -837,4 +837,20 @@ private:
 
 #endif // PCG_EXECUTION_CACHE_VALIDATION_ENABLED
 
+UENUM()
+enum class EPCGTextureResourceType : uint8
+{
+	TextureObject UMETA(ToolTip = "UObject texture such as UTexture2D or UTextureRenderTarget2D."),
+	ExportedTexture UMETA(ToolTip = "Texture handle exported from a texture on the GPU."),
+	Invalid UMETA(Hidden)
+};
+
+/** Method used to determine the value for a sample based on the value of nearby texels. */
+UENUM(BlueprintType)
+enum class EPCGTextureFilter : uint8
+{
+	Point UMETA(Tooltip="Takes the value of whatever texel the sample lands in."),
+	Bilinear UMETA(Tooltip="Bilinearly interpolates the values of the four nearest texels to the sample location.")
+};
+
 #undef UE_API
diff --git a/Engine/Plugins/PCG/Source/PCGEditor/Private/DataVisualizations/PCGBaseTextureDataVisualization.cpp b/Engine/Plugins/PCG/Source/PCGEditor/Private/DataVisualizations/PCGBaseTextureDataVisualization.cpp
--- a/Engine/Plugins/PCG/Source/PCGEditor/Private/DataVisualizations/PCGBaseTextureDataVisualization.cpp
+++ b/Engine/Plugins/PCG/Source/PCGEditor/Private/DataVisualizations/PCGBaseTextureDataVisualization.cpp
@@ -13,17 +13,11 @@
 #include "Engine/Texture.h"
 #include "Materials/MaterialInstanceDynamic.h"
 
-namespace PCGBaseTextureVisualizationConstants
-{
-	static const FSoftObjectPath PlaneMeshPath(TEXT("/Engine/BasicShapes/Plane.Plane"));
-	static const FSoftObjectPath DebugMaterialPath(TEXT("Material'/PCG/DebugObjects/PCG_DebugMaterialTexture.PCG_DebugMaterialTexture'"));
-}
-
 TArray<TSharedPtr<FStreamableHandle>> FPCGBaseTextureDataVisualization::LoadRequiredResources(const UPCGData* Data) const
 {
 	TArray<TSharedPtr<FStreamableHandle>> LoadHandles;
-	LoadHandles.Add(UAssetManager::GetStreamableManager().RequestAsyncLoad(PCGBaseTextureVisualizationConstants::PlaneMeshPath, nullptr));
-	LoadHandles.Add(UAssetManager::GetStreamableManager().RequestAsyncLoad(PCGBaseTextureVisualizationConstants::DebugMaterialPath, nullptr));
+	LoadHandles.Add(UAssetManager::GetStreamableManager().RequestAsyncLoad(PCGTextureVisualizationConstants::PlaneMeshPath, nullptr));
+	LoadHandles.Add(UAssetManager::GetStreamableManager().RequestAsyncLoad(PCGTextureVisualizationConstants::DebugMaterialPath, nullptr));
 
 	return LoadHandles;
 }
@@ -57,7 +51,12 @@ FPCGSetupSceneFunc FPCGBaseTextureDataVisualization::GetViewportSetupFunc(const
 		}
 		else if (WeakData->GetTextureResourceType() == EPCGTextureResourceType::ExportedTexture)
 		{
-			Texture = UPCGVisualizationTexture2D::Create(WeakData);
+			FPCGVisualizationTexture2DParams VisParams;
+			VisParams.TextureRHI = WeakData->GetTextureRHI();
+			VisParams.Resolution = WeakData->GetTextureSize();
+			VisParams.SliceIndex = WeakData->GetTextureSlice();
+
+			Texture = UPCGVisualizationTexture2D::Create(VisParams);
 			InOutParams.ManagedResources.Add(Texture);
 		}
 		else
diff --git a/Engine/Plugins/PCG/Source/PCGEditor/Private/DataVisualizations/PCGTexture2DArrayDataVisualization.cpp b/Engine/Plugins/PCG/Source/PCGEditor/Private/DataVisualizations/PCGTexture2DArrayDataVisualization.cpp
new file mode 100644
--- /dev/null
+++ b/Engine/Plugins/PCG/Source/PCGEditor/Private/DataVisualizations/PCGTexture2DArrayDataVisualization.cpp
@@ -0,0 +1,116 @@
+// Copyright Epic Games, Inc. All Rights Reserved.
+
+#include "DataVisualizations/PCGTexture2DArrayDataVisualization.h"
+
+#include "DataVisualizations/PCGVisualizationTexture2D.h"
+
+#include "Data/PCGTexture2DArrayData.h"
+
+#include "AdvancedPreviewScene.h"
+#include "EditorViewportClient.h"
+#include "Components/StaticMeshComponent.h"
+#include "Engine/AssetManager.h"
+#include "Engine/Texture.h"
+#include "Materials/MaterialInstanceDynamic.h"
+
+TArray<TSharedPtr<FStreamableHandle>> FPCGTexture2DArrayDataVisualization::LoadRequiredResources(const UPCGData* Data) const
+{
+	TArray<TSharedPtr<FStreamableHandle>> LoadHandles;
+	LoadHandles.Add(UAssetManager::GetStreamableManager().RequestAsyncLoad(PCGTextureVisualizationConstants::PlaneMeshPath, nullptr));
+	LoadHandles.Add(UAssetManager::GetStreamableManager().RequestAsyncLoad(PCGTextureVisualizationConstants::DebugMaterialPath, nullptr));
+
+	return LoadHandles;
+}
+
+FPCGSetupSceneFunc FPCGTexture2DArrayDataVisualization::GetViewportSetupFunc(const UPCGSettingsInterface* SettingsInterface, const UPCGData* Data) const
+{
+	return [this, WeakData=TWeakObjectPtr<const UPCGTexture2DArrayData>(Cast<UPCGTexture2DArrayData>(Data))](FPCGSceneSetupParams& InOutParams)
+	{
+		check(InOutParams.Scene);
+		check(InOutParams.EditorViewportClient);
+		check(InOutParams.EditorViewportClient->Viewport);
+
+		if (InOutParams.Resources.Num() != 2)
+		{
+			return;
+		}
+
+		if (!WeakData.IsValid())
+		{
+			UE_LOG(LogPCG, Error, TEXT("Failed to setup data viewport, the data was lost or invalid."));
+			return;
+		}
+
+		TObjectPtr<UTexture> Texture = nullptr;
+		bool bIsTextureArray = false;
+		uint16 SliceIndex = 0; // @todo_pcg: Need some custom widget to choose the desired slice index. Default to 0 for now.
+
+		if (WeakData->GetResourceType() == EPCGTextureResourceType::TextureObject)
+		{
+			Texture = WeakData->GetTexture();
+			bIsTextureArray = true;
+		}
+		else if (WeakData->GetResourceType() == EPCGTextureResourceType::ExportedTexture)
+		{
+			FPCGVisualizationTexture2DParams VisParams;
+			VisParams.TextureRHI = WeakData->GetRHI();
+			VisParams.Resolution = WeakData->GetResolution();
+			VisParams.SliceIndex = SliceIndex;
+
+			Texture = UPCGVisualizationTexture2D::Create(VisParams);
+			InOutParams.ManagedResources.Add(Texture);
+		}
+		else
+		{
+			UE_LOG(LogPCG, Error, TEXT("Texture2DArray data uses an unsupported resource type for data viewport visualization."));
+			return;
+		}
+
+		UMaterialInterface* DebugMaterial = Cast<UMaterialInterface>(InOutParams.Resources[1]);
+		UMaterialInstanceDynamic* MaterialInstance = UMaterialInstanceDynamic::Create(DebugMaterial, GetTransientPackage());
+
+		if (bIsTextureArray)
+		{
+			MaterialInstance->SetTextureParameterValue(FName(TEXT("DebugTextureArray")), Texture);
+			MaterialInstance->SetScalarParameterValue(FName(TEXT("SliceIndex")), SliceIndex);
+			MaterialInstance->SetScalarParameterValue(FName(TEXT("UseTextureArray")), 1.0f);
+		}
+		else
+		{
+			MaterialInstance->SetTextureParameterValue(FName(TEXT("DebugTexture")), Texture);
+		}
+
+		InOutParams.ManagedResources.Add(MaterialInstance);
+
+		TObjectPtr<UStaticMeshComponent> MeshComponent = NewObject<UStaticMeshComponent>(GetTransientPackage(), NAME_None, RF_Transient);
+		MeshComponent->SetStaticMesh(Cast<UStaticMesh>(InOutParams.Resources[0]));
+		MeshComponent->OverrideMaterials.Add(MaterialInstance);
+		InOutParams.ManagedResources.Add(MeshComponent);
+
+		if (GEditor->PreviewPlatform.GetEffectivePreviewFeatureLevel() <= ERHIFeatureLevel::ES3_1)
+		{
+			MeshComponent->SetMobility(EComponentMobility::Static);
+		}
+
+		static const FTransform MeshTransform(FRotator(0.0, -90.0, 0.0), FVector::ZeroVector, FVector::OneVector);
+		InOutParams.Scene->AddComponent(MeshComponent, MeshTransform);
+
+		// @todo_pcg: These settings should all be exposed through InOutParams. Textures should probably have their own preview scene profile that gets selected automatically.
+		InOutParams.Scene->SetFloorVisibility(false);
+		InOutParams.Scene->SetEnvironmentVisibility(false);
+		InOutParams.EditorViewportClient->SetViewportType(ELevelViewportType::LVT_OrthoTop);
+		InOutParams.EditorViewportClient->SetViewMode(EViewModeIndex::VMI_Unlit);
+
+		const FIntPoint ViewportSize = InOutParams.EditorViewportClient->Viewport->GetSizeXY();
+
+		if (ViewportSize.X > 0 && ViewportSize.Y > 0)
+		{
+			// Bounds will be updated already by SetStaticMesh() call.
+			InOutParams.FocusBounds = MeshComponent->Bounds;
+
+			// Fit the orthographic zoom to the mesh. 0.8f is chosen arbitrarily to add some padding around the mesh.
+			const float MeshUnitsPerPixel = FMath::Max(InOutParams.FocusBounds.BoxExtent.X / ViewportSize.X, InOutParams.FocusBounds.BoxExtent.Y / ViewportSize.Y) * 2.0f;
+			InOutParams.FocusOrthoZoom = FMath::Clamp(MeshUnitsPerPixel * DEFAULT_ORTHOZOOM * 0.8f, MIN_ORTHOZOOM, MAX_ORTHOZOOM);
+		}
+	};
+}
diff --git a/Engine/Plugins/PCG/Source/PCGEditor/Private/DataVisualizations/PCGVisualizationTexture2D.cpp b/Engine/Plugins/PCG/Source/PCGEditor/Private/DataVisualizations/PCGVisualizationTexture2D.cpp
--- a/Engine/Plugins/PCG/Source/PCGEditor/Private/DataVisualizations/PCGVisualizationTexture2D.cpp
+++ b/Engine/Plugins/PCG/Source/PCGEditor/Private/DataVisualizations/PCGVisualizationTexture2D.cpp
@@ -2,8 +2,6 @@
 
 #include "DataVisualizations/PCGVisualizationTexture2D.h"
 
-#include "Data/PCGTextureData.h"
-
 #include "RHIStaticStates.h"
 #include "RenderCaptureInterface.h"
 #include "TextureResource.h"
@@ -72,10 +70,9 @@ void FPCGVisualizationTexture2DResource::InitRHI(FRHICommandListBase& RHICmdList
 
 	FRHICommandListImmediate& RHICmdList = RHICmdListBase.GetAsImmediate();
 
-	const UPCGBaseTextureData* TextureData = Owner->GetTextureData();
-	FTextureRHIRef SourceTextureRHI = TextureData ? TextureData->GetTextureRHI() : nullptr;
+	FTextureRHIRef SourceTextureRHI = Owner->GetTextureRHI();
 
-	if (TextureData && SourceTextureRHI.IsValid())
+	if (SourceTextureRHI.IsValid())
 	{
 		RenderCaptureInterface::FScopedCapture RenderCapture(PCGVisTexture2DHelpers::GTriggerGPUCaptureDispatches > 0, &RHICmdList, TEXT("FPCGVisualizationTextureResource::InitRHI::CopySource"));
 		PCGVisTexture2DHelpers::GTriggerGPUCaptureDispatches = FMath::Max(PCGVisTexture2DHelpers::GTriggerGPUCaptureDispatches - 1, 0);
@@ -97,7 +94,7 @@ void FPCGVisualizationTexture2DResource::InitRHI(FRHICommandListBase& RHICmdList
 
 		// Copy the source texture.
 		FRHICopyTextureInfo CopyInfo;
-		CopyInfo.SourceSliceIndex = TextureData->GetTextureSlice();
+		CopyInfo.SourceSliceIndex = Owner->GetSliceIndex();
 		RHICmdList.CopyTexture(SourceTextureRHI, TextureRHI, CopyInfo);
 
 		RHICmdList.Transition(FRHITransitionInfo(TextureRHI, ERHIAccess::CopyDest, ERHIAccess::SRVMask));
@@ -131,26 +128,38 @@ FTextureResource* UPCGVisualizationTexture2D::CreateResource()
 
 float UPCGVisualizationTexture2D::GetSurfaceWidth() const
 {
-	return TextureData.Get() ? TextureData.Get()->GetTextureSize().X : 1u;
+	return Resolution.X;
 }
 
 float UPCGVisualizationTexture2D::GetSurfaceHeight() const
 {
-	return TextureData.Get() ? TextureData.Get()->GetTextureSize().Y : 1u;
+	return Resolution.Y;
+}
+
+FTextureRHIRef UPCGVisualizationTexture2D::GetTextureRHI() const
+{
+	return TextureRHI;
+}
+
+uint16 UPCGVisualizationTexture2D::GetSliceIndex() const
+{
+	return SliceIndex;
 }
 
-void UPCGVisualizationTexture2D::Init(TWeakObjectPtr<const UPCGBaseTextureData> InTextureData)
+void UPCGVisualizationTexture2D::Init(const FPCGVisualizationTexture2DParams& InParams)
 {
-	TextureData = InTextureData;
+	TextureRHI = InParams.TextureRHI;
+	Resolution = InParams.Resolution;
+	SliceIndex = InParams.SliceIndex;
 
 	UpdateResource();
 }
 
-UPCGVisualizationTexture2D* UPCGVisualizationTexture2D::Create(TWeakObjectPtr<const UPCGBaseTextureData> InTextureData)
+UPCGVisualizationTexture2D* UPCGVisualizationTexture2D::Create(const FPCGVisualizationTexture2DParams& InParams)
 {
 	UPCGVisualizationTexture2D* NewTexture = NewObject<UPCGVisualizationTexture2D>(GetTransientPackage(), NAME_None, RF_Transient);
 
-	NewTexture->Init(InTextureData);
+	NewTexture->Init(InParams);
 
 	return NewTexture;
 }
diff --git a/Engine/Plugins/PCG/Source/PCGEditor/Private/DataVisualizations/PCGVisualizationTexture2D.h b/Engine/Plugins/PCG/Source/PCGEditor/Private/DataVisualizations/PCGVisualizationTexture2D.h
--- a/Engine/Plugins/PCG/Source/PCGEditor/Private/DataVisualizations/PCGVisualizationTexture2D.h
+++ b/Engine/Plugins/PCG/Source/PCGEditor/Private/DataVisualizations/PCGVisualizationTexture2D.h
@@ -4,12 +4,25 @@
 
 #include "Engine/Texture.h"
 
+#include "RHIFwd.h"
+
 #include "PCGVisualizationTexture2D.generated.h"
 
-class UPCGBaseTextureData;
+namespace PCGTextureVisualizationConstants
+{
+	static const FSoftObjectPath PlaneMeshPath(TEXT("/Engine/BasicShapes/Plane.Plane"));
+	static const FSoftObjectPath DebugMaterialPath(TEXT("Material'/PCG/DebugObjects/PCG_DebugMaterialTexture.PCG_DebugMaterialTexture'"));
+}
+
+struct FPCGVisualizationTexture2DParams
+{
+	FTextureRHIRef TextureRHI = nullptr;
+	FIntPoint Resolution = FIntPoint::ZeroValue;
+	uint16 SliceIndex = 0;
+};
 
 /**
- * UTexture implementation for visualizing UPCGBaseTextureData in the data viewport.
+ * UTexture implementation for visualizing PCG texture data in the data viewport.
  * This is necessary because some texture data do not have their own UTexture, so one needs
  * to be created in order to override the texture parameter for the debug vis material.
  */
@@ -29,12 +42,16 @@ public:
 	virtual uint32 GetSurfaceArraySize() const override { return 0; }
 	//~ End UTexture Interface.
 	
-	void Init(TWeakObjectPtr<const UPCGBaseTextureData> InTextureData);
-	const UPCGBaseTextureData* GetTextureData() const { return TextureData.Get(); }
+	void Init(const FPCGVisualizationTexture2DParams& InParams);
 
 	/** Creates and initializes a new UPCGVisualizationTexture2D from the texture data. */
-	static UPCGVisualizationTexture2D* Create(TWeakObjectPtr<const UPCGBaseTextureData> InTextureData);
+	static UPCGVisualizationTexture2D* Create(const FPCGVisualizationTexture2DParams& InParams);
+
+	FTextureRHIRef GetTextureRHI() const;
+	uint16 GetSliceIndex() const;
 
 protected:
-	TWeakObjectPtr<const UPCGBaseTextureData> TextureData;
+	FTextureRHIRef TextureRHI = nullptr;
+	FIntPoint Resolution = FIntPoint::ZeroValue;
+	uint16 SliceIndex;
 };
diff --git a/Engine/Plugins/PCG/Source/PCGEditor/Private/PCGEditorModule.cpp b/Engine/Plugins/PCG/Source/PCGEditor/Private/PCGEditorModule.cpp
--- a/Engine/Plugins/PCG/Source/PCGEditor/Private/PCGEditorModule.cpp
+++ b/Engine/Plugins/PCG/Source/PCGEditor/Private/PCGEditorModule.cpp
@@ -9,6 +9,7 @@
 #include "PCGParamData.h"
 #include "PCGSubsystem.h"
 #include "PCGVolumeFactory.h"
+#include "Data/PCGTexture2DArrayData.h"
 #include "Data/PCGCollisionShapeData.h"
 #include "Data/PCGCollisionWrapperData.h"
 #include "Data/PCGDynamicMeshData.h"
@@ -45,6 +46,7 @@
 #include "DataVisualizations/PCGSpatialDataVisualization.h"
 #include "DataVisualizations/PCGSplineDataVisualization.h"
 #include "DataVisualizations/PCGStaticMeshDataVisualization.h"
+#include "DataVisualizations/PCGTexture2DArrayDataVisualization.h"
 #include "DataVisualizations/PCGVolumeDataVisualization.h"
 #include "Details/EnumSelectorDetails.h"
 #include "Details/PCGAttributePropertySelectorDetails.h"
@@ -864,6 +866,7 @@ void FPCGEditorModule::RegisterPCGDataVisualizations()
 	DataVisRegistry.InternalRegistry.Add(UPCGLandscapeData::StaticClass(), MakeUnique<const IPCGLandscapeDataVisualization>());
 	DataVisRegistry.InternalRegistry.Add(UPCGStaticMeshResourceData::StaticClass(), MakeUnique<const IPCGStaticMeshDataVisualization>());
 	DataVisRegistry.InternalRegistry.Add(UPCGBaseTextureData::StaticClass(), MakeUnique<const FPCGBaseTextureDataVisualization>());
+	DataVisRegistry.InternalRegistry.Add(UPCGTexture2DArrayData::StaticClass(), MakeUnique<const FPCGTexture2DArrayDataVisualization>());
 	DataVisRegistry.InternalRegistry.Add(UPCGPolygon2DData::StaticClass(), MakeUnique<const FPCGPolygon2DDataVisualization>());
 	DataVisRegistry.InternalRegistry.Add(UPCGPolygon2DInteriorSurfaceData::StaticClass(), MakeUnique<const FPCGPolygon2DDataVisualization>());
 	DataVisRegistry.InternalRegistry.Add(UPCGDataViewData::StaticClass(), MakeUnique<const FPCGDataViewVisualization>());
diff --git a/Engine/Plugins/PCG/Source/PCGEditor/Public/DataVisualizations/PCGTexture2DArrayDataVisualization.h b/Engine/Plugins/PCG/Source/PCGEditor/Public/DataVisualizations/PCGTexture2DArrayDataVisualization.h
new file mode 100644
--- /dev/null
+++ b/Engine/Plugins/PCG/Source/PCGEditor/Public/DataVisualizations/PCGTexture2DArrayDataVisualization.h
@@ -0,0 +1,17 @@
+// Copyright Epic Games, Inc. All Rights Reserved.
+
+#pragma once
+
+#include "PCGSpatialDataVisualization.h"
+#include "PCGEditorModule.h"
+
+class UPCGData;
+
+class FPCGTexture2DArrayDataVisualization : public IPCGSpatialDataVisualization
+{
+public:
+	// ~Begin IPCGDataVisualization interface
+	PCGEDITOR_API virtual TArray<TSharedPtr<FStreamableHandle>> LoadRequiredResources(const UPCGData* Data) const override;
+	PCGEDITOR_API virtual FPCGSetupSceneFunc GetViewportSetupFunc(const UPCGSettingsInterface* SettingsInterface, const UPCGData* Data) const override;
+	// ~End IPCGDataVisualization interface
+};
