[PCG] Fix missing condition after texture array addition which can trigger ensure downstream.

#rnx
#Synced-CL 49338130

diff --git a/Engine/Plugins/PCG/Source/PCG/Private/Compute/PCGComputeCommon.cpp b/Engine/Plugins/PCG/Source/PCG/Private/Compute/PCGComputeCommon.cpp
--- a/Engine/Plugins/PCG/Source/PCG/Private/Compute/PCGComputeCommon.cpp
+++ b/Engine/Plugins/PCG/Source/PCG/Private/Compute/PCGComputeCommon.cpp
@@ -105,7 +105,7 @@ namespace PCGComputeHelpers
 
 	int32 GetMaxOutputDataCount(const FPCGDataTypeIdentifier& InDataType)
 	{
-		if (InDataType.IsChildOf(EPCGDataType::BaseTexture) || InDataType.IsChildOf(FPCGDataTypeInfoRawBuffer::AsId()))
+		if (InDataType.IsChildOf(EPCGDataType::BaseTexture) || InDataType.IsChildOf(FPCGDataTypeInfoRawBuffer::AsId()) || InDataType.IsChildOf(FPCGDataTypeInfoTexture2DArray::AsId()))
 		{
 			// These data types can only emit a single data.
 			return 1;
