[PCG] Re-establish last selector after collapse to avoid leaving the collapse attribute as [REDACTED_HANDLE]

#rnx
#Synced-CL 49276028

diff --git a/Engine/Plugins/PCG/Source/PCG/Private/Elements/PCGCollapsePoints.cpp b/Engine/Plugins/PCG/Source/PCG/Private/Elements/PCGCollapsePoints.cpp
--- a/Engine/Plugins/PCG/Source/PCG/Private/Elements/PCGCollapsePoints.cpp
+++ b/Engine/Plugins/PCG/Source/PCG/Private/Elements/PCGCollapsePoints.cpp
@@ -529,6 +529,13 @@ bool FPCGCollapsePointsElement::ExecuteInternal(FPCGContext* InContext) const
 			CollapseState.OutputPointData->AllocateProperties(CollapseState.SourcePointData->GetAllocatedProperties());
 			CollapseState.OutputPointData->CopyUnallocatedPropertiesFrom(CollapseState.SourcePointData);
 
+			const bool bHasLastSelector = CollapseState.OutputPointData->HasCachedLastSelector();
+			FPCGAttributePropertyInputSelector LastSelector;
+			if (bHasLastSelector)
+			{
+				LastSelector = CollapseState.OutputPointData->GetCachedLastSelector();
+			}
+
 			// Build accessor list(s) + keys, based on settings
 			TArray<double> PartitionWeights;
 			TUniquePtr<IPCGAttributeAccessorKeys> OutputKeys = PCGAttributeAccessorHelpers::CreateKeys(CollapseState.OutputPointData, FPCGAttributePropertyOutputSelector());
@@ -626,6 +633,12 @@ bool FPCGCollapsePointsElement::ExecuteInternal(FPCGContext* InContext) const
 				++PartitionIndex;
 			}
 
+			// Set Last Selector if needed
+			if (bHasLastSelector)
+			{
+				CollapseState.OutputPointData->SetLastSelector(LastSelector);
+			}
+
 			return true;
 		}
 	});
