[PCG] Update spline sampler on interior to use Geometry Polygon intersection

* Existing code had a few edges cases that were not well handled.
* Added Unit Tests (and move existing ones)
* Also small API update on the Polygon2D to accept vertices by values for move semantics.

#rnx
#Synced-CL 49370341

diff --git a/Engine/Plugins/PCG/Source/PCG/Private/Data/PCGSplineInteriorSurfaceData.cpp b/Engine/Plugins/PCG/Source/PCG/Private/Data/PCGSplineInteriorSurfaceData.cpp
--- a/Engine/Plugins/PCG/Source/PCG/Private/Data/PCGSplineInteriorSurfaceData.cpp
+++ b/Engine/Plugins/PCG/Source/PCG/Private/Data/PCGSplineInteriorSurfaceData.cpp
@@ -139,7 +139,7 @@ UPCGSpatialData* UPCGSplineInteriorSurfaceData::CopyInternal(FPCGContext* Contex
 	NewData->SplineStruct = SplineStruct;
 	NewData->CachedBounds = CachedBounds;
 	NewData->CachedSplinePoints = CachedSplinePoints;
-	NewData->CachedSplinePoints2D = CachedSplinePoints2D;
+	NewData->CachedPolygon2D = CachedPolygon2D;
 #if WITH_EDITOR
 	NewData->bNeedsToCache = false;
 #endif
@@ -181,10 +181,21 @@ void UPCGSplineInteriorSurfaceData::CacheData(FPCGContext* Context)
 
 	// Cache the points which describe the polygon of our spline.
 	const TConstPCGValueRange<FTransform> TransformRange = PointData->GetConstTransformValueRange();
+	
+	TArray<FVector2D> SplinePoints2D;
+	SplinePoints2D.Reserve(TransformRange.Num());
+	CachedSplinePoints.Reserve(TransformRange.Num());
+	
 	for (const FTransform& PointTransform : TransformRange)
 	{
 		CachedSplinePoints.Add(PointTransform.GetLocation());
-		CachedSplinePoints2D.Add(FVector2D(PointTransform.GetLocation()));
+		SplinePoints2D.Add(FVector2D(PointTransform.GetLocation()));
+	}
+
+	CachedPolygon2D = UE::Geometry::FPolygon2d(MoveTemp(SplinePoints2D));
+	if (CachedPolygon2D.IsClockwise())
+	{
+		CachedPolygon2D.Reverse();
 	}
 }
 
@@ -214,10 +225,7 @@ bool UPCGSplineInteriorSurfaceData::PointInsidePolygon(const FTransform& InTrans
 
 	const FVector PointLocation = InTransform.GetLocation();
 
-	// Maximum distance a ray needs to travel to guarantee exiting the polygon at its widest point from the sample location.
-	const FVector::FReal MaxDistance = CachedBounds.Max.X - PointLocation.X + UE_KINDA_SMALL_NUMBER;
-
 	// Test sample location against spline interior.
 	// TODO: This could also consume a transform + bounds to accept rotated bounds that overlap the polygon.
-	return PCGSplineSamplerHelpers::PointInsidePolygon2D(CachedSplinePoints2D, FVector2D(PointLocation), MaxDistance);
+	return PCGSplineSamplerHelpers::PointInsidePolygon2D(CachedPolygon2D, FVector2D(PointLocation));
 }
diff --git a/Engine/Plugins/PCG/Source/PCG/Private/Elements/PCGSplineSampler.cpp b/Engine/Plugins/PCG/Source/PCG/Private/Elements/PCGSplineSampler.cpp
--- a/Engine/Plugins/PCG/Source/PCG/Private/Elements/PCGSplineSampler.cpp
+++ b/Engine/Plugins/PCG/Source/PCG/Private/Elements/PCGSplineSampler.cpp
@@ -24,283 +24,47 @@
 
 namespace PCGSplineSamplerHelpers
 {
-	/**
-	* Robust 2D line segment intersection. Returns true if the segments intersect and stores the intersection in OutIntersectionPoint.
-	* If the segments are collinear, one of the contained endpoints will be chosen as the intersection.
-	*/
-	bool SegmentIntersection2D(const FVector2D& SegmentStartA, const FVector2D& SegmentEndA, const FVector2D& SegmentStartB, const FVector2D& SegmentEndB, FVector2D& OutIntersectionPoint)
-	{
-		const FVector2D VectorA = SegmentEndA - SegmentStartA;
-		const FVector2D VectorB = SegmentEndB - SegmentStartB;
-
-		const FVector::FReal Determinant = VectorA.X * VectorB.Y - VectorB.X * VectorA.Y;
-
-		// Determinant is zero means the segments are parallel. Check if they are also collinear by projecting (0, 0) onto both lines.
-		if (FMath::IsNearlyZero(Determinant))
-		{
-			const FVector::FReal SquareMagnitudeA = VectorA.SquaredLength();
-			if (FMath::IsNearlyZero(SquareMagnitudeA))
-			{
-				const FVector2D& MinB = SegmentStartB.X < SegmentEndB.X ? SegmentStartB : SegmentEndB;
-				const FVector2D& MaxB = SegmentStartB.X > SegmentEndB.X ? SegmentStartB : SegmentEndB;
-
-				if (SegmentStartA.X > MinB.X && SegmentStartA.X < MaxB.X)
-				{
-					OutIntersectionPoint = SegmentStartA;
-					return true;
-				}
-
-				return false;
-			}
-
-			const FVector::FReal SquareMagnitudeB = VectorB.SquaredLength();
-			if (FMath::IsNearlyZero(SquareMagnitudeB))
-			{
-				const FVector2D& MinA = SegmentStartA.X < SegmentEndA.X ? SegmentStartA : SegmentEndA;
-				const FVector2D& MaxA = SegmentStartA.X > SegmentEndA.X ? SegmentStartA : SegmentEndA;
-
-				if (SegmentStartB.X > MinA.X && SegmentStartB.X < MaxA.X)
-				{
-					OutIntersectionPoint = SegmentStartB;
-					return true;
-				}
-
-				return false;
-			}
-
-			// Taking the ray from SegmentStart to (0, 0) for both segments, and projecting the ray onto its
-			// respective line segment should yield two coincident points if the segments are collinear
-			const FVector2D AToZero = -SegmentStartA; // (0, 0) - SegmentStartA
-			const FVector2D ProjectionOnA = VectorA * VectorA.Dot(AToZero) / SquareMagnitudeA;
-			const FVector2D ProjectedPointOnA = SegmentStartA + ProjectionOnA;
-
-			const FVector2D BToZero = -SegmentStartB; // (0, 0) - SegmentStartB
-			const FVector2D ProjectionOnB = VectorB * VectorB.Dot(BToZero) / SquareMagnitudeB;
-			const FVector2D ProjectedPointOnB = SegmentStartB + ProjectionOnB;
-			
-			// If the projected points are not equal, then we are not collinear
-			if (!ProjectedPointOnA.Equals(ProjectedPointOnB))
-			{
-				return false;
-			}
-			
-			const FVector2D& MinA = SegmentStartA.X < SegmentEndA.X ? SegmentStartA : SegmentEndA;
-			const FVector2D& MaxA = SegmentStartA.X > SegmentEndA.X ? SegmentStartA : SegmentEndA;
-			const FVector2D& MinB = SegmentStartB.X < SegmentEndB.X ? SegmentStartB : SegmentEndB;
-			const FVector2D& MaxB = SegmentStartB.X > SegmentEndB.X ? SegmentStartB : SegmentEndB;
-			
-			if (MinA.X > MinB.X && MinA.X < MaxB.X)
-			{
-				OutIntersectionPoint = MinA;
-				return true;
-			}
-			else if (MaxA.X > MinB.X && MaxA.X < MaxB.X)
-			{
-				OutIntersectionPoint = MaxA;
-				return true;
-			}
-			if (MinB.X > MinA.X && MinB.X < MaxA.X)
-			{
-				OutIntersectionPoint = MinB;
-				return true;
-			}
-			else if (MaxB.X > MinA.X && MaxB.X < MaxA.X)
-			{
-				OutIntersectionPoint = MaxB;
-				return true;
-			}
-			
-			return false;
-		}
-
-		const FVector::FReal DeltaX = SegmentStartA.X - SegmentStartB.X;
-		const FVector::FReal DeltaY = SegmentStartA.Y - SegmentStartB.Y;
-
-		const FVector::FReal S = (VectorA.X * DeltaY - VectorA.Y * DeltaX) / Determinant;
-		const FVector::FReal T = (VectorB.X * DeltaY - VectorB.Y * DeltaX) / Determinant;
-
-		if (S >= 0 && S <= 1 && T >= 0 && T <= 1)
-		{
-			OutIntersectionPoint = SegmentStartA + T * VectorA;
-			return true;
-		}
-
-		return false;
-	}
-
 	/** Intersects a line segment with all line segments of a polygon. May return duplicates if the segment intersects a vertex of the polygon. */
-	void SegmentPolygonIntersection2D(const FVector2D& SegmentStart, const FVector2D& SegmentEnd, const TArray<FVector2D>& PolygonPoints, TArray<FVector2D>& OutIntersectionPoints)
+	TArray<FVector2D> SegmentPolygonIntersection2D(const FVector2D& SegmentStart, const FVector2D& SegmentEnd, const UE::Geometry::FPolygon2d& InPolygon)
 	{
-		const int32 PointCount = PolygonPoints.Num();
+		const int32 PointCount = InPolygon.VertexCount();
 
 		if (PointCount < 3)
 		{
-			return;
+			return {};
 		}
 
-		// Checks if a duplicate point should be counted as a valid intersection and adds it to the intersection list if so. Collinear points are also considered duplicates.
-		auto TestDuplicatePoint = [PointCount, &PolygonPoints](const FVector2D& IntersectionPoint, int32 Index) -> bool
-		{
-			const int32 BackwardIndex = Index - 1;
-			const int32 ForwardIndex = Index + 1;
-
-			// Find DeltaY for the nearest backward point
-			FVector::FReal DeltaY1 = 0.f;
-			for (int32 Offset = 0; Offset < PointCount; ++Offset)
-			{
-				DeltaY1 = PolygonPoints[(BackwardIndex - Offset + PointCount) % PointCount].Y - IntersectionPoint.Y;
-
-				// Choose the first non-zero delta
-				if (!FMath::IsNearlyZero(DeltaY1))
-				{
-					break;
-				}
-			}
-
-			// Find DeltaY for the nearest forward point
-			FVector::FReal DeltaY2 = 0.f;
-			for (int32 Offset = 0; Offset < PointCount; ++Offset)
-			{
-				DeltaY2 = PolygonPoints[(ForwardIndex + Offset) % PointCount].Y - IntersectionPoint.Y;
-
-				// Choose the first non-zero delta
-				if (!FMath::IsNearlyZero(DeltaY2))
-				{
-					break;
-				}
-			}
-
-			// Allow double counting when adjacent points lie on opposite sides of our ray
-			return (DeltaY1 > 0 && DeltaY2 > 0) || (DeltaY1 < 0 && DeltaY2 < 0);
-		};
-
-		auto AreSegmentsCollinear = [](const FVector2D& SegmentStartA, const FVector2D& SegmentEndA, const FVector2D& SegmentStartB, const FVector2D& SegmentEndB) -> bool
-		{
-			const FVector2D VectorA = SegmentEndA - SegmentStartA;
-			const FVector2D VectorB = SegmentEndB - SegmentStartB;
-
-			const FVector::FReal SquareMagnitudeA = VectorA.SquaredLength();
-			const FVector::FReal SquareMagnitudeB = VectorB.SquaredLength();
-
-			if (!FMath::IsNearlyZero(FVector2D::CrossProduct(VectorA, VectorB)) || FMath::IsNearlyZero(SquareMagnitudeA) || FMath::IsNearlyZero(SquareMagnitudeB))
-			{
-				return false;
-			}
-
-			// Taking the ray from SegmentStart to (0, 0) for both segments, and projecting the ray onto its
-			// respective line segment should yield two coincident points if the segments are collinear
-			const FVector2D AToZero = -SegmentStartA; // (0, 0) - SegmentStartA
-			const FVector2D ProjectionOnA = VectorA * VectorA.Dot(AToZero) / SquareMagnitudeA;
-			const FVector2D ProjectedPointOnA = SegmentStartA + ProjectionOnA;
-
-			const FVector2D BToZero = -SegmentStartB; // (0, 0) - SegmentStartB
-			const FVector2D ProjectionOnB = VectorB * VectorB.Dot(BToZero) / SquareMagnitudeB;
-			const FVector2D ProjectedPointOnB = SegmentStartB + ProjectionOnB;
-
-			// If the projected points are not equal, then we are not collinear
-			return ProjectedPointOnA.Equals(ProjectedPointOnB);
-		};
-
-		FVector2D IntersectionPoint;
-		FVector2D PreviousIntersectionPoint;
-		bool bLastSegmentWasCollinear = false;
+		TArray<FVector2D> IntersectionPoints;
 
-		for (int32 PointIndex = 0; PointIndex < PointCount - 1; ++PointIndex)
+		// Taken from TPolygon2::FindIntersections
+		const UE::Geometry::FSegment2d InSegment{SegmentStart, SegmentEnd};
+		for (UE::Geometry::FSegment2d PolygonSegment : InPolygon.Segments())
 		{
-			const FVector2D& PolySegmentStart = PolygonPoints[PointIndex];
-			const FVector2D& PolySegmentEnd = PolygonPoints[PointIndex + 1];
-
-			// Discard zero length segments
-			if (PolySegmentStart.Equals(PolySegmentEnd))
-			{
-				continue;
-			}
-
-			const bool bSegmentIsCollinear = AreSegmentsCollinear(PolySegmentEnd, PolySegmentStart, SegmentStart, SegmentEnd);
-
-			// Discard collinear segments and test if duplicate points should be included
-			if (!bSegmentIsCollinear && SegmentIntersection2D(SegmentStart, SegmentEnd, PolySegmentStart, PolySegmentEnd, IntersectionPoint))
+			if (PolygonSegment.Intersects(InSegment))
 			{
-				if ((OutIntersectionPoints.Num() > 0 && IntersectionPoint.Equals(PreviousIntersectionPoint)) || bLastSegmentWasCollinear)
+				UE::Geometry::FIntrSegment2Segment2d Intersection(PolygonSegment, InSegment);
+				if (Intersection.Find())
 				{
-					if (TestDuplicatePoint(IntersectionPoint, PointIndex))
+					IntersectionPoints.Add(Intersection.Point0);
+					if (Intersection.Quantity == 2)
 					{
-						OutIntersectionPoints.Add(IntersectionPoint);
+						IntersectionPoints.Add(Intersection.Point1);
 					}
 				}
-				else
-				{
-					OutIntersectionPoints.Add(IntersectionPoint);
-				}
-
-				PreviousIntersectionPoint = IntersectionPoint;
 			}
-
-			bLastSegmentWasCollinear = bSegmentIsCollinear;
 		}
 
-		// Special case to wrap the last segment
-		{
-			const FVector2D& PolySegmentStart = PolygonPoints.Last();
-			const FVector2D& PolySegmentEnd = PolygonPoints[0];
-
-			// Discard zero length segments
-			if (PolySegmentStart.Equals(PolySegmentEnd))
-			{
-				return;
-			}
-
-			const bool bSegmentIsCollinear = AreSegmentsCollinear(PolySegmentEnd, PolySegmentStart, SegmentStart, SegmentEnd);
-
-			if (bSegmentIsCollinear && OutIntersectionPoints.Num() > 0 && OutIntersectionPoints[0].Equals(PolygonPoints[0]))
-			{
-				// If the last segment is collinear, then we will double count the segment unless we cull the first polygon point
-				OutIntersectionPoints.RemoveAt(0);
-			}
-			else if (!bSegmentIsCollinear && SegmentIntersection2D(SegmentStart, SegmentEnd, PolySegmentStart, PolySegmentEnd, IntersectionPoint))
-			{
-				if (OutIntersectionPoints.Num() > 0 || bLastSegmentWasCollinear)
-				{
-					if (IntersectionPoint.Equals(PolygonPoints[0]))
-					{
-						if (TestDuplicatePoint(IntersectionPoint, PointCount))
-						{
-							OutIntersectionPoints.Add(IntersectionPoint);
-						}
-					}
-					else if (IntersectionPoint.Equals(PolygonPoints.Last()))
-					{
-						if (TestDuplicatePoint(IntersectionPoint, PointCount - 1))
-						{
-							OutIntersectionPoints.Add(IntersectionPoint);
-						}
-					}
-					else
-					{
-						OutIntersectionPoints.Add(IntersectionPoint);
-					}
-				}
-				else
-				{
-					OutIntersectionPoints.Add(IntersectionPoint);
-				}
-			}
-		}
+		return IntersectionPoints;
 	}
 
-	/** Tests if a point lies inside the given polygon by casting a ray to MaxDistance and counting the intersections */
-	bool PointInsidePolygon2D(const TArray<FVector2D>& PolygonPoints, const FVector2D& Point, FVector::FReal MaxDistance) 
+	bool PointInsidePolygon2D(const UE::Geometry::FPolygon2d& InPolygon, const FVector2D& Point) 
 	{
-		if (PolygonPoints.Num() < 3)
+		if (InPolygon.VertexCount() < 3)
 		{
 			return false;
 		}
 
-		TArray<FVector2D> IntersectionPoints;
-		SegmentPolygonIntersection2D(Point, Point + FVector2D(MaxDistance, 0), PolygonPoints, IntersectionPoints);
-
-		// Odd number of intersections means we are inside the polygon
-		return (IntersectionPoints.Num() % 2) == 1;
+		return InPolygon.Contains(Point);
 	}
 
 	/** Projects a point in space onto the approximated surface defined by a closed spline. */
@@ -1048,6 +812,8 @@ namespace PCGSplineSamplerHelpers
  
 	void SampleInteriorData(FPCGContext* Context, const UPCGPolyLineData* LineData, const UPCGSpatialData* InBoundingShape, const UPCGSpatialData* InProjectionTarget, const FPCGProjectionParams& InProjectionParams, const FPCGSplineSamplerParams& Params, UPCGBasePointData* OutPointData)
 	{
+		using namespace UE::Geometry;
+
 		check(LineData && OutPointData);
 
 		const FPCGSplineStruct* Spline = nullptr;
@@ -1130,6 +896,10 @@ namespace PCGSplineSamplerHelpers
 				SplineSamplePoints.Add(Spline->GetLocationAtDistanceAlongSpline(Length, ESplineCoordinateSpace::Local));
 			}
 		}
+		else
+		{
+			return;
+		}
 
 		// Flat polygon representation of our spline points
 		TArray<FVector2D> SplineSamplePoints2D;
@@ -1143,6 +913,12 @@ namespace PCGSplineSamplerHelpers
 			MaxPoint = FVector::Max(MaxPoint, Point);
 		}
 
+		FPolygon2d PolygonFromSpline(MoveTemp(SplineSamplePoints2D));
+		if (PolygonFromSpline.IsClockwise())
+		{
+			PolygonFromSpline.Reverse();
+		}
+
 		TArray<TTuple<FVector2D, FVector2D>> MedialAxisEdges;
 		FVector2D Centroid;
 
@@ -1152,6 +928,9 @@ namespace PCGSplineSamplerHelpers
 			if (NumSegments >= 4)
 			{
 				TArray<FVector> PolygonPoints; // Top-down 2D projection polygon of the spline points
+				const FPolygon2d* PolygonPoints2DPtr = nullptr;
+				FPolygon2d PolygonPoint2D;
+
 				if (Params.bTreatSplineAsPolyline)
 				{
 					// If we already computed the polygon, use a copy instead of generating it again
@@ -1161,26 +940,33 @@ namespace PCGSplineSamplerHelpers
 						FVector& PolygonPoint = PolygonPoints.Add_GetRef(SplinePoint);
 						PolygonPoint.Z = MinPoint.Z;
 					}
+
+					PolygonPoints2DPtr = &PolygonFromSpline;
 				}
 				else
 				{
+					TArray<FVector2D> TempPolygonPoints2D;
 					PolygonPoints.Reserve(NumSegments);
+					TempPolygonPoints2D.Reserve(NumSegments);
 
 					for (int32 SegmentIndex = 0; SegmentIndex < NumSegments; ++SegmentIndex)
 					{
 						FVector& PolygonPoint = PolygonPoints.Add_GetRef(LineData->GetLocationAtDistance(SegmentIndex, 0, /*bWorldSpace=*/false));
 						PolygonPoint.Z = MinPoint.Z;
+						TempPolygonPoints2D.Add(FVector2D(PolygonPoint));
 					}
-				}
 
-				TArray<FVector2D> PolygonPoints2D;
-				PolygonPoints2D.Reserve(PolygonPoints.Num());
+					PolygonPoint2D = FPolygon2d{MoveTemp(TempPolygonPoints2D)};
+					if (PolygonPoint2D.IsClockwise())
+					{
+						PolygonPoint2D.Reverse();
+					}
 
-				for (const FVector& Point : PolygonPoints)
-				{
-					PolygonPoints2D.Add(FVector2D(Point));
+					PolygonPoints2DPtr = &PolygonPoint2D;
 				}
 
+				check(PolygonPoints2DPtr)
+
 				TArray<TTuple<FVector, FVector>> VoronoiEdges;
 				TArray<int32> CellMember;
 
@@ -1194,10 +980,11 @@ namespace PCGSplineSamplerHelpers
 				{
 					// Discard any edges which intersect the polygon
 					bool bDiscard = false;
-					for (int32 PointIndex = 0; PointIndex < PolygonPoints.Num(); ++PointIndex)
+					for (const FSegment2d& Segment : PolygonPoints2DPtr->Segments())
 					{
-						FVector2D IntersectionPoint;
-						if (PCGSplineSamplerHelpers::SegmentIntersection2D(PolygonPoints2D[PointIndex], PolygonPoints2D[(PointIndex + 1) % PolygonPoints2D.Num()], FVector2D(Edge.Get<0>()), FVector2D(Edge.Get<1>()), IntersectionPoint))
+						const FSegment2d Edge2D{FVector2D(Edge.Get<0>()), FVector2D(Edge.Get<1>())};
+
+						if (Segment.Intersects(Edge2D))
 						{
 							bDiscard = true;
 							break;
@@ -1210,7 +997,7 @@ namespace PCGSplineSamplerHelpers
 					}
 
 					// If either of the points lies within the polygon, the segment must lie within the polygon
-					if (PCGSplineSamplerHelpers::PointInsidePolygon2D(PolygonPoints2D, FVector2D(Edge.Get<0>()), MaxDimension))
+					if (PCGSplineSamplerHelpers::PointInsidePolygon2D(*PolygonPoints2DPtr, FVector2D(Edge.Get<0>())))
 					{
 						MedialAxisEdges.Add(Edge);
 					}
@@ -1269,14 +1056,12 @@ namespace PCGSplineSamplerHelpers
 			const FVector::FReal LocalMaxY = bIsLastIteration ? (MaxY + UE_KINDA_SMALL_NUMBER) : (MinY + EndIterationIndex * Params.InteriorSampleSpacing);
 
 			// Point sampling
-			for(FVector::FReal Y = LocalMinY; Y < LocalMaxY; Y += Params.InteriorSampleSpacing)
+			for (FVector::FReal Y = LocalMinY; Y < LocalMaxY; Y += Params.InteriorSampleSpacing)
 			{
 				const FVector2D RayMin(MinPoint.X - RayPadding, Y);
 				const FVector2D RayMax(MaxPoint.X + RayPadding, Y);
 
-				// Get the intersections along this ray, sorted by Y value
-				TArray<FVector2D> Intersections;
-				PCGSplineSamplerHelpers::SegmentPolygonIntersection2D(RayMin, RayMax, SplineSamplePoints2D, Intersections);
+				TArray<FVector2D> Intersections = SegmentPolygonIntersection2D(RayMin, RayMax, PolygonFromSpline);
 
 				if (Intersections.Num() % 2 != 0)
 				{
diff --git a/Engine/Plugins/PCG/Source/PCG/Private/Tests/Elements/PCGSplineSamplerTest.cpp b/Engine/Plugins/PCG/Source/PCG/Private/Tests/Elements/PCGSplineSamplerTest.cpp
deleted file mode 100644
diff --git a/Engine/Plugins/PCG/Source/PCG/Public/Data/PCGSplineInteriorSurfaceData.h b/Engine/Plugins/PCG/Source/PCG/Public/Data/PCGSplineInteriorSurfaceData.h
--- a/Engine/Plugins/PCG/Source/PCG/Public/Data/PCGSplineInteriorSurfaceData.h
+++ b/Engine/Plugins/PCG/Source/PCG/Public/Data/PCGSplineInteriorSurfaceData.h
@@ -5,6 +5,8 @@
 #include "Data/PCGSurfaceData.h"
 #include "Data/PCGSplineStruct.h"
 
+#include "Polygon2.h"
+
 #include "PCGSplineInteriorSurfaceData.generated.h"
 
 class UPCGSpatialData;
@@ -71,10 +73,13 @@ protected:
 	UPROPERTY(Transient)
 	TArray<FVector> CachedSplinePoints;
 
-	/** Cached list of 2D points to describe the polygon given by the spline. */
+	UE_DEPRECATED(5.8, "Not filled anymore, use CachedPolygon2D")
 	UPROPERTY(Transient)
 	TArray<FVector2D> CachedSplinePoints2D;
 
+	/** Cached polygon given by the spline. */
+	UE::Geometry::FPolygon2d CachedPolygon2D;
+
 #if WITH_EDITORONLY_DATA
 	/** Flag that indicates when we need to recompute our cached information. Preferable to serializing data that can be recreated on demand. */
 	UPROPERTY(Transient)
diff --git a/Engine/Plugins/PCG/Source/PCG/Public/Elements/PCGSplineSampler.h b/Engine/Plugins/PCG/Source/PCG/Public/Elements/PCGSplineSampler.h
--- a/Engine/Plugins/PCG/Source/PCG/Public/Elements/PCGSplineSampler.h
+++ b/Engine/Plugins/PCG/Source/PCG/Public/Elements/PCGSplineSampler.h
@@ -4,6 +4,7 @@
 
 #include "PCGSettings.h"
 
+#include "Polygon2.h"
 #include "Curves/CurveFloat.h"
 
 #include "PCGSplineSampler.generated.h"
@@ -250,8 +251,10 @@ namespace PCGSplineSamplerHelpers
 
 	const UPCGPolyLineData* GetPolyLineData(const UPCGSpatialData* InSpatialData);
 
-	/** Tests if a point lies inside the given polygon by casting a ray to MaxDistance and counting the intersections. */
-	bool PointInsidePolygon2D(const TArray<FVector2D>& PolygonPoints, const FVector2D& Point, FVector::FReal MaxDistance);
+	/**
+	 * Tests if a point lies inside the given polygon. Polygon MUST be CounterClockwise otherwise the result will be incorrect.
+	 */
+	bool PointInsidePolygon2D(const UE::Geometry::FPolygon2d& PolygonPoints, const FVector2D& Point);
 
 	/** Projects a point in space onto the approximated surface defined by a closed spline. Returns the height of the point after projection. */
 	FVector::FReal ProjectOntoSplineInteriorSurface(const TArray<FVector>& SplinePoints, const FVector& PointToProject);
diff --git a/Engine/Plugins/PCG/Tests/Private/Elements/PCGSplineSamplerTest.cpp b/Engine/Plugins/PCG/Tests/Private/Elements/PCGSplineSamplerTest.cpp
new file mode 100644
--- /dev/null
+++ b/Engine/Plugins/PCG/Tests/Private/Elements/PCGSplineSamplerTest.cpp
@@ -0,0 +1,111 @@
+// Copyright Epic Games, Inc. All Rights Reserved.
+
+#include "CoreMinimal.h"
+#include "TestHarness.h"
+
+#include "Elements/PCGSplineSampler.h"
+
+#include <catch2/catch_test_macros.hpp>
+#include <catch2/generators/catch_generators.hpp>
+
+
+TEST_CASE("PCG::SplineSampler", "[PCG][SplineSampler]")
+{	
+	/* Clock-wise torch shape, to test edge cases
+	 * |\/ \/|
+	 * |_   _|
+	 *   |_|
+	 */
+
+	TArray<FVector2D> PolygonPoints =
+	{
+		FVector2D(0, 0),
+		FVector2D(0, 1),
+		FVector2D(0, 2),
+		FVector2D(1, 1),
+		FVector2D(1.5, 2),
+		FVector2D(1.5, 2), // Some duplicate points
+		FVector2D(2, 1),
+		FVector2D(3, 2),
+		FVector2D(3, 1),
+		FVector2D(3, 0),
+		FVector2D(3, 0),
+		FVector2D(2.5, 0),
+		FVector2D(2, 0),
+		FVector2D(2, -1),
+		FVector2D(2, -1),
+		FVector2D(1, -1),
+		FVector2D(1, 0),
+		FVector2D(0.5, 0),
+	};
+
+	UE::Geometry::FPolygon2d Polygon{MoveTemp(PolygonPoints)};
+	if (Polygon.IsClockwise())
+	{
+		Polygon.Reverse();
+	}
+
+	auto [What, Point, bInsidePolygon] = GENERATE(Catch::Generators::table<FString, FVector2D, bool>(
+		{
+			// Interior point lies inside polygon
+			{"Interior point", {0.5, 0.5}, true},
+			// Exterior point lies outside polygon
+			{"Exterior point", {10.5, 0.5}, false},
+
+			// Y = -1
+			// (1-e, -1) lies outside polygon
+			{"(1-e, -1)", {1 - UE_KINDA_SMALL_NUMBER, -1}, false},
+			// (2+e, -1) lies outside polygon
+			{"(2+e, -1)", {2 + UE_KINDA_SMALL_NUMBER, -1}, false},
+
+			// Y = 0
+			// (0-e, 0) lies outside polygon
+			{"(0-e, 0)", {0 - UE_KINDA_SMALL_NUMBER, 0}, false},
+			// (1.5, 0) lies inside polygon
+			{"(1.5, 0)", {1.5, 0}, true},
+			// (3+e, 0) lies outside polygon
+			{"(3+e, 0)", {3 + UE_KINDA_SMALL_NUMBER, 0}, false},
+
+			// Y = 1
+			// Test around (0, 1)
+			// (0-e, 1) lies outside polygon
+			{"(0-e, 1)", {0 - UE_KINDA_SMALL_NUMBER, 1}, false},
+			// (0+e, 1) lies inside polygon
+			{"(0+e, 1)", { 0 + UE_KINDA_SMALL_NUMBER, 1}, true},
+			// Test around (1, 1)
+			// (1-e, 1) lies inside polygon
+			{"(1-e, 1)", {1 - UE_KINDA_SMALL_NUMBER, 1}, true},
+			// (1+e, 1) lies inside polygon
+			{"(1+e, 1)", { 1 + UE_KINDA_SMALL_NUMBER, 1}, true},
+			// Test around (2, 1)
+			// (2-e, 1) lies inside polygon
+			{"(2-e, 1)", {2 - UE_KINDA_SMALL_NUMBER, 1}, true},
+			// (2+e, 1) lies inside polygon
+			{"(2+e, 1)", { 2 + UE_KINDA_SMALL_NUMBER, 1}, true},
+			// Test around (3, 1)
+			// (3-e, 1) lies inside polygon
+			{"(3-e, 1)", {3 - UE_KINDA_SMALL_NUMBER, 1}, true},
+			// (3+e, 1) lies outside polygon
+			{"(3+e, 1)", { 3 + UE_KINDA_SMALL_NUMBER, 1}, false},
+
+			// Y = 2
+			// Test around (0, 2)
+			// (0-e, 2) lies outside polygon
+			{"(0-e, 2)", {0 - UE_KINDA_SMALL_NUMBER, 2}, false},
+			// (0+e, 2) lies outside polygon
+			{"(0+e, 2)", { 0 + UE_KINDA_SMALL_NUMBER, 2}, false},
+			// Test around (1.5, 2)
+			// (1.5, 2) lies outside polygon
+			{"(1.5, 2)", {1.5, 2}, false},
+			// Test around (3, 2)
+			// (3-e, 2) lies outside polygon
+			{"(3-e, 2)", {3 - UE_KINDA_SMALL_NUMBER, 2}, false},
+			// (3+e, 2) lies outside polygon
+			{"(3+e, 2)", { 3 + UE_KINDA_SMALL_NUMBER, 2}, false},
+		}));
+
+	DYNAMIC_SECTION(*FString::Printf(TEXT("%s lies %s polygon"), *What, bInsidePolygon ? TEXT("inside") : TEXT("outside")))
+	{
+		REQUIRE(PCGSplineSamplerHelpers::PointInsidePolygon2D(Polygon, Point) == bInsidePolygon);
+	}
+}
diff --git a/Engine/Source/Runtime/GeometryCore/Public/Polygon2.h b/Engine/Source/Runtime/GeometryCore/Public/Polygon2.h
--- a/Engine/Source/Runtime/GeometryCore/Public/Polygon2.h
+++ b/Engine/Source/Runtime/GeometryCore/Public/Polygon2.h
@@ -42,7 +42,7 @@ public:
 	/**
 	 * Construct polygon with given list of vertices
 	 */
-	TPolygon2(const TArray<TVector2<T>>& VertexList) : Vertices(VertexList)
+	TPolygon2(TArray<TVector2<T>> VertexList) : Vertices(MoveTemp(VertexList))
 	{
 	}
 
