[Backout] - CL48603520
#rnx
Original CL Desc

This change makes a grass instancing kernel 50% more expensive.

-----------------------------------------------------------------
[PCG] Add explicit bounds checks to buffer loads/stores in data collection data interface.

[RN] [PCG] Fix cases of undefined behaviour and GPU crashes by adding protection against out of bounds reads/writes from Custom HLSL nodes.
#Synced-CL Unknown

#Synced-CL 49460467

diff --git a/Engine/Plugins/PCG/Shaders/Private/PCGDataCollectionDataInterface.ush b/Engine/Plugins/PCG/Shaders/Private/PCGDataCollectionDataInterface.ush
--- a/Engine/Plugins/PCG/Shaders/Private/PCGDataCollectionDataInterface.ush
+++ b/Engine/Plugins/PCG/Shaders/Private/PCGDataCollectionDataInterface.ush
@@ -22,8 +22,6 @@
 // @todo_pcg: Support having a non-RW byte address buffer for input data.
 RWByteAddressBuffer {DataInterfaceName}_DataCollectionBuffer;
 
-uint {DataInterfaceName}_BufferSizeBytes;
-
 // Multiplier that can be configured on output pins.
 uint {DataInterfaceName}_ElementCountMultiplier;
 
@@ -47,16 +45,6 @@ struct FPCGAttributeHeader_{DataInterfaceName}
 	bool bAllocated;
 };
 
-bool IsAddressValid_{DataInterfaceName}(uint InAddressBytes)
-{
-	return InAddressBytes < {DataInterfaceName}_BufferSizeBytes;
-}
-
-bool IsAddressValid_{DataInterfaceName}(uint InAddressBytes, uint InElementSizeBytes)
-{
-	return (InAddressBytes + InElementSizeBytes) <= {DataInterfaceName}_BufferSizeBytes;
-}
-
 // #################### SHADER PARAM GETTERS ##########################
 
 // @todo_pcg: This parameter is broader than this data interface. Needs to be shared in some way (e.g. with TextureDataInterface).
@@ -75,24 +63,12 @@ uint GetNumData_{DataInterfaceName}()
 
 uint GetNumElements_{DataInterfaceName}(uint InDataIndex)
 {
-	if (InDataIndex >= GetNumData_{DataInterfaceName}())
-	{
-		return 0u;
-	}
-
 	const uint ReadAddress =
 		PCG_DATA_COLLECTION_HEADER_SIZE_BYTES +
 		InDataIndex * PCG_DATA_HEADER_SIZE_BYTES +
 		/*TypeId*/4 + /*Attribute Count*/4;
 
-	if (IsAddressValid_{DataInterfaceName}(ReadAddress))
-	{
-		return {DataInterfaceName}_DataCollectionBuffer.Load(ReadAddress);
-	}
-	else
-	{
-		return 0u;
-	}
+	return {DataInterfaceName}_DataCollectionBuffer.Load(ReadAddress);
 }
 
 bool GetThreadData_{DataInterfaceName}(uint InThreadIndex, out uint OutDataIndex, out uint OutElementIndex)
@@ -136,23 +112,12 @@ uint GetNumElements_{DataInterfaceName}()
 
 uint LoadBufferInternal_{DataInterfaceName}(uint InAddress)
 {
-	// Full branch, not ternary, to avoid executing the load if address invalid.
-	if (IsAddressValid_{DataInterfaceName}(InAddress))
-	{
-		return {DataInterfaceName}_DataCollectionBuffer.Load(InAddress);
-	}
-	else
-	{
-		return 0u;
-	}
+	return {DataInterfaceName}_DataCollectionBuffer.Load(InAddress);
 }
 
 void StoreBufferInternal_{DataInterfaceName}(uint InAddress, uint InValue)
 {
-	if (IsAddressValid_{DataInterfaceName}(InAddress))
-	{
-		{DataInterfaceName}_DataCollectionBuffer.Store(InAddress, InValue);
-	}
+	{DataInterfaceName}_DataCollectionBuffer.Store(InAddress, InValue);
 }
 
 int GetRemappedAttributeIdInternal_{DataInterfaceName}(int InGraphAttributeId)
@@ -171,11 +136,6 @@ int GetRemappedAttributeIdInternal_{DataInterfaceName}(int InGraphAttributeId)
 
 FPCGAttributeHeader_{DataInterfaceName} GetAttributeHeaderInternal_{DataInterfaceName}(uint InDataIndex, int InAttributeId)
 {
-	if (InDataIndex >= GetNumData_{DataInterfaceName}())
-	{
-		return (FPCGAttributeHeader_{DataInterfaceName})0;
-	}
-
 	const int AttributeId = GetRemappedAttributeIdInternal_{DataInterfaceName}(InAttributeId);
 
 	const uint ReadAddress =
@@ -184,11 +144,6 @@ FPCGAttributeHeader_{DataInterfaceName} GetAttributeHeaderInternal_{DataInterfac
 		/*TypeId*/4 + /*Attribute Count*/4 + /*Element Count*/4 +
 		AttributeId * PCG_ATTRIBUTE_HEADER_SIZE_BYTES;
 
-	if (!IsAddressValid_{DataInterfaceName}(ReadAddress))
-	{
-		return (FPCGAttributeHeader_{DataInterfaceName})0;
-	}
-
 	const uint2 PackedAttributeHeader = {DataInterfaceName}_DataCollectionBuffer.Load2(ReadAddress);
 
 	FPCGAttributeHeader_{DataInterfaceName} AttributeHeader;
@@ -218,13 +173,13 @@ uint GetFirstValidAttributeElementInternal_{DataInterfaceName}()
 	{
 		const uint ElementAddress = GetElementAddressInternal_{DataInterfaceName}(/*DataIndex=*/0, /*ElementIndex=*/0, AttributeId);
 
-		if (ElementAddress > 0u)
+		if (ElementAddress > 0)
 		{
 			return ElementAddress;
 		}
 	}
 
-	return 0u;
+	return 0;
 }
 
 uint GetDataNumAttributesInternal_{DataInterfaceName}(uint InDataIndex)
@@ -234,14 +189,7 @@ uint GetDataNumAttributesInternal_{DataInterfaceName}(uint InDataIndex)
 		InDataIndex * PCG_DATA_HEADER_SIZE_BYTES +
 		/*TypeId*/4;
 
-	if (IsAddressValid_{DataInterfaceName}(ReadAddress))
-	{
-		return {DataInterfaceName}_DataCollectionBuffer.Load(ReadAddress);
-	}
-	else
-	{
-		return 0u;
-	}
+	return {DataInterfaceName}_DataCollectionBuffer.Load(ReadAddress);
 }
 
 void SetAsExecutedInternal_{DataInterfaceName}()
@@ -257,7 +205,7 @@ uint GetUint_{DataInterfaceName}(uint InDataIndex, uint InElementIndex, int InAt
 {
 	const uint ElementAddress = GetElementAddressInternal_{DataInterfaceName}(InDataIndex, InElementIndex, InAttributeId);
 
-	if (ElementAddress > 0 && IsAddressValid_{DataInterfaceName}(ElementAddress))
+	if (ElementAddress > 0)
 	{
 		return {DataInterfaceName}_DataCollectionBuffer.Load(ElementAddress);
 	}
@@ -286,8 +234,7 @@ float2 GetFloat2_{DataInterfaceName}(uint InDataIndex, uint InElementIndex, int
 {
 	const uint ElementAddress = GetElementAddressInternal_{DataInterfaceName}(InDataIndex, InElementIndex, InAttributeId);
 
-	// 0 is invalid element address, and pen-ultimate word must lie within buffer.
-	if (ElementAddress > 0 && IsAddressValid_{DataInterfaceName}(ElementAddress, /*InElementSizeBytes=*/8u))
+	if (ElementAddress > 0)
 	{
 		return asfloat({DataInterfaceName}_DataCollectionBuffer.Load2(ElementAddress));
 	}
@@ -301,8 +248,7 @@ float3 GetFloat3_{DataInterfaceName}(uint InDataIndex, uint InElementIndex, int
 {
 	const uint ElementAddress = GetElementAddressInternal_{DataInterfaceName}(InDataIndex, InElementIndex, InAttributeId);
 
-	// 0 is invalid element address, and pen-ultimate word must lie within buffer.
-	if (ElementAddress > 0 && IsAddressValid_{DataInterfaceName}(ElementAddress, /*InElementSizeBytes=*/12u))
+	if (ElementAddress > 0)
 	{
 		return asfloat({DataInterfaceName}_DataCollectionBuffer.Load3(ElementAddress));
 	}
@@ -316,8 +262,7 @@ float4 GetFloat4_{DataInterfaceName}(uint InDataIndex, uint InElementIndex, int
 {
 	const uint ElementAddress = GetElementAddressInternal_{DataInterfaceName}(InDataIndex, InElementIndex, InAttributeId);
 
-	// 0 is invalid element address, and pen-ultimate word must lie within buffer.
-	if (ElementAddress > 0 && IsAddressValid_{DataInterfaceName}(ElementAddress, /*InElementSizeBytes=*/16u))
+	if (ElementAddress > 0)
 	{
 		return asfloat({DataInterfaceName}_DataCollectionBuffer.Load4(ElementAddress));
 	}
@@ -341,8 +286,7 @@ float4x4 GetTransform_{DataInterfaceName}(uint InDataIndex, uint InElementIndex,
 {
 	const uint ElementAddress = GetElementAddressInternal_{DataInterfaceName}(InDataIndex, InElementIndex, InAttributeId);
 
-	// 0 is invalid element address, and pen-ultimate word must lie within buffer.
-	if (ElementAddress > 0 && IsAddressValid_{DataInterfaceName}(ElementAddress, /*InElementSizeBytes=*/64u))
+	if (ElementAddress > 0)
 	{
 		return float4x4(
 			asfloat({DataInterfaceName}_DataCollectionBuffer.Load4(ElementAddress)),
@@ -382,7 +326,7 @@ void SetBool_{DataInterfaceName}(uint InDataIndex, uint InElementIndex, int InAt
 {
 	const uint ElementAddress = GetElementAddressInternal_{DataInterfaceName}(InDataIndex, InElementIndex, InAttributeId);
 
-	if (ElementAddress > 0 && IsAddressValid_{DataInterfaceName}(ElementAddress))
+	if (ElementAddress > 0)
 	{
 		{DataInterfaceName}_DataCollectionBuffer.Store(ElementAddress, InValue ? 1u : 0u);
 	}
@@ -392,7 +336,7 @@ void SetInt_{DataInterfaceName}(uint InDataIndex, uint InElementIndex, int InAtt
 {
 	const uint ElementAddress = GetElementAddressInternal_{DataInterfaceName}(InDataIndex, InElementIndex, InAttributeId);
 
-	if (ElementAddress > 0 && IsAddressValid_{DataInterfaceName}(ElementAddress))
+	if (ElementAddress > 0)
 	{
 		{DataInterfaceName}_DataCollectionBuffer.Store(ElementAddress, asuint(InValue));
 	}
@@ -402,7 +346,7 @@ void SetUint_{DataInterfaceName}(uint InDataIndex, uint InElementIndex, int InAt
 {
 	const uint ElementAddress = GetElementAddressInternal_{DataInterfaceName}(InDataIndex, InElementIndex, InAttributeId);
 
-	if (ElementAddress > 0 && IsAddressValid_{DataInterfaceName}(ElementAddress))
+	if (ElementAddress > 0)
 	{
 		{DataInterfaceName}_DataCollectionBuffer.Store(ElementAddress, InValue);
 	}
@@ -412,7 +356,7 @@ void SetFloat_{DataInterfaceName}(uint InDataIndex, uint InElementIndex, int InA
 {
 	const uint ElementAddress = GetElementAddressInternal_{DataInterfaceName}(InDataIndex, InElementIndex, InAttributeId);
 
-	if (ElementAddress > 0 && IsAddressValid_{DataInterfaceName}(ElementAddress))
+	if (ElementAddress > 0)
 	{
 		{DataInterfaceName}_DataCollectionBuffer.Store(ElementAddress, asuint(InValue));
 	}
@@ -422,7 +366,7 @@ void SetFloat2_{DataInterfaceName}(uint InDataIndex, uint InElementIndex, int In
 {
 	const uint ElementAddress = GetElementAddressInternal_{DataInterfaceName}(InDataIndex, InElementIndex, InAttributeId);
 
-	if (ElementAddress > 0 && IsAddressValid_{DataInterfaceName}(ElementAddress, /*InElementSizeBytes=*/8u))
+	if (ElementAddress > 0)
 	{
 		{DataInterfaceName}_DataCollectionBuffer.Store2(ElementAddress, asuint(InValue));
 	}
@@ -432,7 +376,7 @@ void SetFloat3_{DataInterfaceName}(uint InDataIndex, uint InElementIndex, int In
 {
 	const uint ElementAddress = GetElementAddressInternal_{DataInterfaceName}(InDataIndex, InElementIndex, InAttributeId);
 
-	if (ElementAddress > 0 && IsAddressValid_{DataInterfaceName}(ElementAddress, /*InElementSizeBytes=*/12u))
+	if (ElementAddress > 0)
 	{
 		{DataInterfaceName}_DataCollectionBuffer.Store3(ElementAddress, asuint(InValue));
 	}
@@ -442,7 +386,7 @@ void SetFloat4_{DataInterfaceName}(uint InDataIndex, uint InElementIndex, int In
 {
 	const uint ElementAddress = GetElementAddressInternal_{DataInterfaceName}(InDataIndex, InElementIndex, InAttributeId);
 
-	if (ElementAddress > 0 && IsAddressValid_{DataInterfaceName}(ElementAddress, /*InElementSizeBytes=*/16u))
+	if (ElementAddress > 0)
 	{
 		{DataInterfaceName}_DataCollectionBuffer.Store4(ElementAddress, asuint(InValue));
 	}
@@ -450,19 +394,29 @@ void SetFloat4_{DataInterfaceName}(uint InDataIndex, uint InElementIndex, int In
 
 void SetRotator_{DataInterfaceName}(uint InDataIndex, uint InElementIndex, int InAttributeId, float3 InValue)
 {
-	SetFloat3_{DataInterfaceName}(InDataIndex, InElementIndex, InAttributeId, InValue);
+	const uint ElementAddress = GetElementAddressInternal_{DataInterfaceName}(InDataIndex, InElementIndex, InAttributeId);
+
+	if (ElementAddress > 0)
+	{
+		{DataInterfaceName}_DataCollectionBuffer.Store3(ElementAddress, asuint(InValue));
+	}
 }
 
 void SetQuat_{DataInterfaceName}(uint InDataIndex, uint InElementIndex, int InAttributeId, float4 InValue)
 {
-	SetFloat4_{DataInterfaceName}(InDataIndex, InElementIndex, InAttributeId, InValue);
+	const uint ElementAddress = GetElementAddressInternal_{DataInterfaceName}(InDataIndex, InElementIndex, InAttributeId);
+
+	if (ElementAddress > 0)
+	{
+		{DataInterfaceName}_DataCollectionBuffer.Store4(ElementAddress, asuint(InValue));
+	}
 }
 
 void SetTransform_{DataInterfaceName}(uint InDataIndex, uint InElementIndex, int InAttributeId, float4x4 InValue)
 {
 	const uint ElementAddress = GetElementAddressInternal_{DataInterfaceName}(InDataIndex, InElementIndex, InAttributeId);
 
-	if (ElementAddress > 0 && IsAddressValid_{DataInterfaceName}(ElementAddress, /*InElementSizeBytes=*/64u))
+	if (ElementAddress > 0)
 	{
 		{DataInterfaceName}_DataCollectionBuffer.Store4(ElementAddress, asuint(InValue[0]));
 		{DataInterfaceName}_DataCollectionBuffer.Store4(ElementAddress + 4 * 4, asuint(InValue[1]));
@@ -489,7 +443,7 @@ uint GetFirstAttributeAsUint_{DataInterfaceName}()
 {
 	const uint AttributeAddress = GetFirstValidAttributeElementInternal_{DataInterfaceName}();
 
-	if (AttributeAddress > 0 && IsAddressValid_{DataInterfaceName}(AttributeAddress))
+	if (AttributeAddress > 0)
 	{
 		return {DataInterfaceName}_DataCollectionBuffer.Load(AttributeAddress);
 	}
@@ -518,7 +472,7 @@ float2 GetFirstAttributeAsFloat2_{DataInterfaceName}()
 {
 	const uint AttributeAddress = GetFirstValidAttributeElementInternal_{DataInterfaceName}();
 
-	if (AttributeAddress > 0 && IsAddressValid_{DataInterfaceName}(AttributeAddress, /*InElementSizeBytes=*/8u))
+	if (AttributeAddress > 0)
 	{
 		return asfloat({DataInterfaceName}_DataCollectionBuffer.Load2(AttributeAddress));
 	}
@@ -532,7 +486,7 @@ float3 GetFirstAttributeAsFloat3_{DataInterfaceName}()
 {
 	const uint AttributeAddress = GetFirstValidAttributeElementInternal_{DataInterfaceName}();
 
-	if (AttributeAddress > 0 && IsAddressValid_{DataInterfaceName}(AttributeAddress, /*InElementSizeBytes=*/12u))
+	if (AttributeAddress > 0)
 	{
 		return asfloat({DataInterfaceName}_DataCollectionBuffer.Load3(AttributeAddress));
 	}
@@ -546,7 +500,7 @@ float4 GetFirstAttributeAsFloat4_{DataInterfaceName}()
 {
 	const uint AttributeAddress = GetFirstValidAttributeElementInternal_{DataInterfaceName}();
 
-	if (AttributeAddress > 0 && IsAddressValid_{DataInterfaceName}(AttributeAddress, /*InElementSizeBytes=*/16u))
+	if (AttributeAddress > 0)
 	{
 		return asfloat({DataInterfaceName}_DataCollectionBuffer.Load4(AttributeAddress));
 	}
@@ -569,7 +523,7 @@ float4 GetFirstAttributeAsQuat_{DataInterfaceName}()
 float4x4 GetFirstAttributeAsTransform_{DataInterfaceName}()
 {
 	const uint AttributeAddress = GetFirstValidAttributeElementInternal_{DataInterfaceName}();
-	if (AttributeAddress == 0 || !IsAddressValid_{DataInterfaceName}(AttributeAddress, /*InElementSizeBytes=*/64u))
+	if (AttributeAddress == 0)
 	{
 		return (float4x4)0;
 	}
@@ -601,7 +555,7 @@ int AtomicAddInt_{DataInterfaceName}(uint InDataIndex, uint InElementIndex, int
 {
 	const uint ElementAddress = GetElementAddressInternal_{DataInterfaceName}(InDataIndex, InElementIndex, InAttributeId);
 
-	if (ElementAddress > 0 && IsAddressValid_{DataInterfaceName}(ElementAddress))
+	if (ElementAddress > 0)
 	{
 		uint OriginalValue;
 		{DataInterfaceName}_DataCollectionBuffer.InterlockedAdd(ElementAddress, (uint)InValueToAdd, OriginalValue);
diff --git a/Engine/Plugins/PCG/Source/PCG/Private/Compute/DataInterfaces/PCGDataCollectionDataInterface.cpp b/Engine/Plugins/PCG/Source/PCG/Private/Compute/DataInterfaces/PCGDataCollectionDataInterface.cpp
--- a/Engine/Plugins/PCG/Source/PCG/Private/Compute/DataInterfaces/PCGDataCollectionDataInterface.cpp
+++ b/Engine/Plugins/PCG/Source/PCG/Private/Compute/DataInterfaces/PCGDataCollectionDataInterface.cpp
@@ -467,7 +467,6 @@ void UPCGDataCollectionDataInterface::GetSupportedInputs(TArray<FShaderFunctionD
 
 BEGIN_SHADER_PARAMETER_STRUCT(FPCGDataCollectionDataInterfaceParameters, )
 	SHADER_PARAMETER_RDG_BUFFER_UAV(RWByteAddressBuffer, DataCollectionBuffer)
-	SHADER_PARAMETER(uint32, BufferSizeBytes)
 	SHADER_PARAMETER(uint32, ElementCountMultiplier)
 	SHADER_PARAMETER_RDG_BUFFER_SRV(StructuredBuffer<int32>, GraphToBufferAttributeId)
 	SHADER_PARAMETER(int32, FirstRemappedAttributeId)
@@ -726,7 +725,6 @@ void FPCGDataCollectionDataProviderProxy::GatherDispatchData(FDispatchData const
 	{
 		FParameters& Parameters = ParameterArray[InvocationIndex];
 		Parameters.DataCollectionBuffer = BufferUAV;
-		Parameters.BufferSizeBytes = static_cast<uint32>(SizeBytes);
 		Parameters.ElementCountMultiplier = ElementCountMultiplier;
 		Parameters.GraphToBufferAttributeId = AttributeIdRemapSRV;
 		Parameters.FirstRemappedAttributeId = FirstRemappedAttributeId;
