[PCG] Split some CVars between Editor/Runtime

#Synced-CL 49558812

diff --git a/Engine/Plugins/PCG/Source/PCG/Private/Graph/PCGGraphExecutor.cpp b/Engine/Plugins/PCG/Source/PCG/Private/Graph/PCGGraphExecutor.cpp
--- a/Engine/Plugins/PCG/Source/PCG/Private/Graph/PCGGraphExecutor.cpp
+++ b/Engine/Plugins/PCG/Source/PCG/Private/Graph/PCGGraphExecutor.cpp
@@ -45,28 +45,28 @@
 
 #include UE_INLINE_GENERATED_CPP_BY_NAME(PCGGraphExecutor)
 
-static TAutoConsoleVariable<int32> CVarMaxNumTasks(
-	TEXT("pcg.MaxConcurrentTasks"),
-	4096,
-	TEXT("Maximum number of concurrent tasks for PCG processing"));
-
-static TAutoConsoleVariable<float> CVarMaxPercentageOfThreadsToUse(
-	TEXT("pcg.MaxPercentageOfThreadsToUse"),
-	0.9f,
-	TEXT("Maximum percentage of number of threads for concurrent PCG processing"));
-
-static TAutoConsoleVariable<float> CVarMaxPercentageOfExecutingThreads(
-	TEXT("pcg.MaxPercentageOfExecutingThreads"),
-	0.3f,
-	TEXT("Maximum percentage of number of threads for processing scheduled pcg elements (percentage of pcg.MaxPercentageOfThreadsToUse)"));
-
-static TAutoConsoleVariable<int32> CVarMaxExecutingThreads(
-	TEXT("pcg.MaxExecutingThreads"),
-	8,
-	TEXT("Hard cap on number of threads for processing scheduled pcg elements calculated from pcg.MaxPercentageOfExecutingThreads"));
-
 namespace PCGGraphExecutor
 {
+	TAutoConsoleVariable<int32> CVarRuntimeMaxNumProcessingTasks(
+		TEXT("pcg.Runtime.MaxProcessingTasks"),
+		4096,
+		TEXT("Runtime: Maximum number of concurrent tasks for PCG Async processing"));
+
+	TAutoConsoleVariable<float> CVarRuntimeMaxPercentageOfThreadsToUse(
+		TEXT("pcg.Runtime.MaxPercentageOfThreadsToUse"),
+		0.9f,
+		TEXT("Runtime: Maximum percentage of number of threads for all concurrent PCG processing/execution"));
+
+	TAutoConsoleVariable<float> CVarRuntimeMaxPercentageOfExecutingThreads(
+		TEXT("pcg.Runtime.MaxPercentageOfExecutingThreads"),
+		0.3f,
+		TEXT("Runtime: Maximum percentage of number of threads for processing scheduled pcg elements (percentage of pcg.Runtime.MaxPercentageOfThreadsToUse)"));
+
+	TAutoConsoleVariable<int32> CVarRuntimeMaxExecutingThreads(
+		TEXT("pcg.Runtime.MaxExecutingThreads"),
+		8,
+		TEXT("Runtime: Hard cap on number of threads for processing scheduled pcg elements calculated from pcg.Runtime.MaxPercentageOfExecutingThreads"));
+
 	TAutoConsoleVariable<float> CVarTimePerFrame(
 		TEXT("pcg.FrameTime"),
 		5.0f,
@@ -78,6 +78,26 @@ namespace PCGGraphExecutor
 		TEXT("Controls whether the graph can dispatch multiple tasks at the same time"));
 
 #if WITH_EDITOR
+	TAutoConsoleVariable<int32> CVarEditorMaxNumProcessingTasks(
+		TEXT("pcg.Editor.MaxProcessingTasks"),
+		4096,
+		TEXT("Editor: Maximum number of concurrent tasks for PCG Async processing"));
+
+	TAutoConsoleVariable<float> CVarEditorMaxPercentageOfThreadsToUse(
+		TEXT("pcg.Editor.MaxPercentageOfThreadsToUse"),
+		0.9f,
+		TEXT("Editor: Maximum percentage of number of threads for all concurrent PCG processing/execution"));
+
+	TAutoConsoleVariable<float> CVarEditorMaxPercentageOfExecutingThreads(
+		TEXT("pcg.Editor.MaxPercentageOfExecutingThreads"),
+		0.3f,
+		TEXT("Editor: Maximum percentage of number of threads for processing scheduled pcg elements (percentage of pcg.Editor.MaxPercentageOfThreadsToUse)"));
+
+	TAutoConsoleVariable<int32> CVarEditorMaxExecutingThreads(
+		TEXT("pcg.Editor.MaxExecutingThreads"),
+		8,
+		TEXT("Editor: Hard cap on number of threads for processing scheduled pcg elements calculated from pcg.Editor.MaxPercentageOfExecutingThreads"));
+
 	TAutoConsoleVariable<float> CVarEditorTimePerFrame(
 		TEXT("pcg.EditorFrameTime"),
 		15.0f,
@@ -94,6 +114,134 @@ namespace PCGGraphExecutor
 		TEXT("Pauses all execution of PCG while user transactions are active"));
 #endif
 
+	TAutoConsoleVariable<int32> CVarMaxNumTasks_DEPRECATED(
+		TEXT("pcg.MaxConcurrentTasks"),
+		4096,
+		TEXT("DEPRECATED(5.8): use pcg.Editor.MaxProcessingTasks or pcg.Runtime.MaxProcessingTasks instead."),
+		FConsoleVariableDelegate::CreateLambda([](IConsoleVariable* InVariable)
+		{
+			// Called when user sets cvar from console, or when this cvar is set when launching standalone build. Pushes value to new cvars.
+			UE_LOG(LogPCG, Warning, TEXT("pcg.MaxConcurrentTasks is deprecated in 5.8. use pcg.Editor.MaxProcessingTasks or pcg.Runtime.MaxProcessingTasks instead."));
+			check(InVariable);
+			int32 NewValue = 0;
+			InVariable->GetValue(NewValue);
+
+			CVarRuntimeMaxNumProcessingTasks->Set(NewValue, ECVF_SetByCode);
+#if WITH_EDITOR
+			CVarEditorMaxNumProcessingTasks->Set(NewValue, ECVF_SetByCode);
+#endif
+		}));
+
+	TAutoConsoleVariable<float> CVarMaxPercentageOfThreadsToUse_DEPRECATED(
+		TEXT("pcg.MaxPercentageOfThreadsToUse"),
+		0.9f,
+		TEXT("DEPRECATED(5.8) : use pcg.Editor.MaxPercentageOfThreadsToUse or pcg.Runtime.MaxPercentageOfThreadsToUse instead."),
+		FConsoleVariableDelegate::CreateLambda([](IConsoleVariable* InVariable)
+		{
+			// Called when user sets cvar from console, or when this cvar is set when launching standalone build. Pushes value to new cvars.
+			UE_LOG(LogPCG, Warning, TEXT("pcg.MaxPercentageOfThreadsToUse is deprecated in 5.8. use pcg.Editor.MaxPercentageOfThreadsToUse or pcg.Runtime.MaxPercentageOfThreadsToUse instead."));
+			check(InVariable);
+			float NewValue = 0.f;
+			InVariable->GetValue(NewValue);
+
+			CVarRuntimeMaxPercentageOfThreadsToUse->Set(NewValue, ECVF_SetByCode);
+#if WITH_EDITOR
+			CVarEditorMaxPercentageOfThreadsToUse->Set(NewValue, ECVF_SetByCode);
+#endif
+			}));
+
+	TAutoConsoleVariable<float> CVarMaxPercentageOfExecutingThreads_DEPRECATED(
+		TEXT("pcg.MaxPercentageOfExecutingThreads"),
+		0.3f,
+		TEXT("DEPRECATED(5.8) : use pcg.Editor.MaxPercentageOfExecutingThreads or pcg.Runtime.MaxPercentageOfExecutingThreads instead."),
+		FConsoleVariableDelegate::CreateLambda([](IConsoleVariable* InVariable)
+		{
+			// Called when user sets cvar from console, or when this cvar is set when launching standalone build. Pushes value to new cvars.
+			UE_LOG(LogPCG, Warning, TEXT("pcg.MaxPercentageOfExecutingThreads is deprecated in 5.8. use pcg.Editor.MaxPercentageOfExecutingThreads or pcg.Runtime.MaxPercentageOfExecutingThreads instead."));
+			check(InVariable);
+			float NewValue = 0.f;
+			InVariable->GetValue(NewValue);
+
+			CVarRuntimeMaxPercentageOfExecutingThreads->Set(NewValue, ECVF_SetByCode);
+#if WITH_EDITOR
+			CVarEditorMaxPercentageOfExecutingThreads->Set(NewValue, ECVF_SetByCode);
+#endif
+		}));
+
+#if WITH_EDITOR
+	// Cvar deprecation. One-shot migration: runs after ini loading / user input via a cvar sink. Only required with editor, in standalone
+	// it appears the pcg.Cache.Enabled change delegate fires which handles driving the new cvars.
+	static void MigrateCVarsIfNeeded()
+	{
+		static bool bDone = false;
+		if (bDone)
+		{
+			return;
+		}
+		bDone = true;
+
+		auto SetFromIni = [](uint32 InFlags, uint32 InSetBy)
+		{
+			// Treat any of these as "user set in ini" (covers typical places users put CVars)
+			return InSetBy == ECVF_SetBySystemSettingsIni || // [SystemSettings] sections in *Engine.ini
+				InSetBy == ECVF_SetByConsoleVariablesIni || // ConsoleVariables.ini
+				InSetBy == ECVF_SetByDeviceProfile || // DeviceProfiles.ini
+				(InFlags & ECVF_CreatedFromIni) != 0; // Created by ini before registration
+		};
+
+		{
+			const uint32 MaxNumTasks_Flags = CVarMaxNumTasks_DEPRECATED->GetFlags();
+			const uint32 MaxNumTasks_SetBy = (MaxNumTasks_Flags & ECVF_SetByMask);
+
+			if (SetFromIni(MaxNumTasks_Flags, MaxNumTasks_SetBy))
+			{
+				const int32 Value = CVarMaxNumTasks_DEPRECATED->GetInt();
+
+				// Copy with the same SetBy priority so behavior matches the user's source of truth.
+				CVarRuntimeMaxNumProcessingTasks->Set(Value, (EConsoleVariableFlags)MaxNumTasks_SetBy);
+				CVarEditorMaxNumProcessingTasks->Set(Value, (EConsoleVariableFlags)MaxNumTasks_SetBy);
+
+				UE_LOG(LogPCG, Warning, TEXT("Cvar 'pcg.MaxConcurrentTasks' was set from ini but is deprecated in 5.8. Migrated value to pcg.Runtime.MaxProcessingTasks=%d, pcg.Editor.MaxProcessingTasks=%d (SetBy=0x%x)."), Value, Value, MaxNumTasks_SetBy);
+			}
+		}
+
+		{
+			const uint32 MaxPercentage_Flags = CVarMaxPercentageOfThreadsToUse_DEPRECATED->GetFlags();
+			const uint32 MaxPercentage_SetBy = (MaxPercentage_Flags & ECVF_SetByMask);
+
+			if (SetFromIni(MaxPercentage_Flags, MaxPercentage_SetBy))
+			{
+				const float Value = CVarMaxPercentageOfThreadsToUse_DEPRECATED->GetFloat();
+
+				// Copy with the same SetBy priority so behavior matches the user's source of truth.
+				CVarRuntimeMaxPercentageOfThreadsToUse->Set(Value, (EConsoleVariableFlags)MaxPercentage_SetBy);
+				CVarEditorMaxPercentageOfThreadsToUse->Set(Value, (EConsoleVariableFlags)MaxPercentage_SetBy);
+
+				UE_LOG(LogPCG, Warning, TEXT("Cvar 'pcg.MaxPercentageOfThreadsToUse' was set from ini but is deprecated in 5.8. Migrated value to pcg.Runtime.MaxPercentageOfThreadsToUse=%f, pcg.Editor.MaxPercentageOfThreadsToUse=%f (SetBy=0x%x)."), Value, Value, MaxPercentage_SetBy);
+			}
+		}
+
+		{
+			const uint32 MaxExecutePercentage_Flags = CVarMaxPercentageOfExecutingThreads_DEPRECATED->GetFlags();
+			const uint32 MaxExecutePercentage_SetBy = (MaxExecutePercentage_Flags & ECVF_SetByMask);
+
+			if (SetFromIni(MaxExecutePercentage_Flags, MaxExecutePercentage_SetBy))
+			{
+				const float Value = CVarMaxPercentageOfExecutingThreads_DEPRECATED->GetFloat();
+
+				// Copy with the same SetBy priority so behavior matches the user's source of truth.
+				CVarRuntimeMaxPercentageOfExecutingThreads->Set(Value, (EConsoleVariableFlags)MaxExecutePercentage_SetBy);
+				CVarEditorMaxPercentageOfExecutingThreads->Set(Value, (EConsoleVariableFlags)MaxExecutePercentage_SetBy);
+
+				UE_LOG(LogPCG, Warning, TEXT("Cvar 'pcg.MaxPercentageOfExecutingThreads' was set from ini but is deprecated in 5.8. Migrated value to pcg.Runtime.MaxPercentageOfExecutingThreads=%f, pcg.Editor.MaxPercentageOfExecutingThreads=%f (SetBy=0x%x)."), Value, Value, MaxExecutePercentage_SetBy);
+			}
+		}
+	}
+
+	// Register a sink so this runs at the right time (after ini load / on first cvar change tick).
+	static FAutoConsoleVariableSink GCVarMigrationSink(FConsoleCommandDelegate::CreateStatic(&MigrateCVarsIfNeeded));
+#endif // WITH_EDITOR
+
 	TAutoConsoleVariable<bool> CVarDynamicTaskCulling(
 		TEXT("pcg.Graph.DynamicTaskCulling"),
 		true,
@@ -170,6 +318,19 @@ namespace PCGGraphExecutor
 			}
 		}
 	}
+
+	template<class T>
+	T GetCVarValueForWorld(const UWorld* World, const TAutoConsoleVariable<T>& RuntimeValue, const TAutoConsoleVariable<T>& EditorValue)
+	{
+#if WITH_EDITOR
+		if (!World || !World->IsGameWorld())
+		{
+			return EditorValue.GetValueOnAnyThread();
+		}
+#endif
+
+		return RuntimeValue.GetValueOnAnyThread();
+	}
 }
 
 FPCGGraphExecutor::FPCGGraphExecutor()
@@ -400,7 +561,7 @@ TSet<IPCGGraphExecutionSource*> FPCGGraphExecutor::Cancel(TFunctionRef<bool(TWea
 			}
 		}
 	}
-		
+
 	{
 		// Visit ready tasks
 		PCG::TScopeLock ScopeLock(LiveTasksLock);
@@ -542,11 +703,11 @@ TSet<IPCGGraphExecutionSource*> FPCGGraphExecutor::Cancel(TFunctionRef<bool(TWea
 				// Wait for Task (if any) to complete
 				UE::Tasks::TTask<bool> TaskHandle = CancelledActiveTask.Get<UE::Tasks::TTask<bool>>();
 				TaskHandle.Wait();
-				
-				TSharedPtr<FPCGGraphActiveTask> ActiveTask = CancelledActiveTask.Get<TSharedPtr<FPCGGraphActiveTask>>();			
+
+				TSharedPtr<FPCGGraphActiveTask> ActiveTask = CancelledActiveTask.Get<TSharedPtr<FPCGGraphActiveTask>>();
 				check(ActiveTask->Element);
 				ActiveTask->Element->Abort(ActiveTask->Context.Get());
-												
+
 				bStableCancellationSet &= !CancelNextTasks(ActiveTask->NodeId, CancelledExecutionSources);
 			}
 		}
@@ -1507,16 +1668,49 @@ void FPCGGraphExecutor::DebugDisplayTask(TWeakPtr<FPCGGraphExecutor::FGameThread
 }
 #endif // WITH_EDITOR
 
-double FPCGGraphExecutor::GetTickBudgetInSeconds()
+double FPCGGraphExecutor::GetTickBudgetInSeconds() const
 {
 #if WITH_EDITOR
-	if (GEditor && !GEditor->IsPlaySessionInProgress())
-	{
-		return PCGGraphExecutor::CVarEditorTimePerFrame.GetValueOnAnyThread() / 1000.0;
-	}
+	return PCGGraphExecutor::GetCVarValueForWorld(World, PCGGraphExecutor::CVarTimePerFrame, PCGGraphExecutor::CVarEditorTimePerFrame) / 1000.0;
+#else
+	return PCGGraphExecutor::CVarTimePerFrame.GetValueOnAnyThread() / 1000.0;
 #endif
+}
 
-	return PCGGraphExecutor::CVarTimePerFrame.GetValueOnAnyThread() / 1000.0;
+int32 FPCGGraphExecutor::GetMaxNumProcessingTasks() const
+{
+#if WITH_EDITOR
+	return PCGGraphExecutor::GetCVarValueForWorld(World, PCGGraphExecutor::CVarRuntimeMaxNumProcessingTasks, PCGGraphExecutor::CVarEditorMaxNumProcessingTasks);
+#else
+	return PCGGraphExecutor::CVarRuntimeMaxNumProcessingTasks.GetValueOnAnyThread();
+#endif
+}
+
+float FPCGGraphExecutor::GetMaxPercentageOfThreadsToUse() const
+{
+#if WITH_EDITOR
+	return PCGGraphExecutor::GetCVarValueForWorld(World, PCGGraphExecutor::CVarRuntimeMaxPercentageOfThreadsToUse, PCGGraphExecutor::CVarEditorMaxPercentageOfThreadsToUse);
+#else
+	return PCGGraphExecutor::CVarRuntimeMaxPercentageOfThreadsToUse.GetValueOnAnyThread();
+#endif
+}
+
+float FPCGGraphExecutor::GetMaxPercentageOfExecutingThreads() const
+{
+#if WITH_EDITOR
+	return PCGGraphExecutor::GetCVarValueForWorld(World, PCGGraphExecutor::CVarRuntimeMaxPercentageOfExecutingThreads, PCGGraphExecutor::CVarEditorMaxPercentageOfExecutingThreads);
+#else
+	return PCGGraphExecutor::CVarRuntimeMaxPercentageOfExecutingThreads.GetValueOnAnyThread();
+#endif
+}
+
+int32 FPCGGraphExecutor::GetMaxExecutingThreads() const
+{
+#if WITH_EDITOR
+	return PCGGraphExecutor::GetCVarValueForWorld(World, PCGGraphExecutor::CVarRuntimeMaxExecutingThreads, PCGGraphExecutor::CVarEditorMaxExecutingThreads);
+#else
+	return PCGGraphExecutor::CVarRuntimeMaxExecutingThreads.GetValueOnAnyThread();
+#endif
 }
 
 bool FPCGGraphExecutor::ExecuteScheduling(double EndTime, TSharedPtr<FPCGGraphActiveTask>* OutMainThreadTask, bool bForceCheckPausedTasks)
@@ -1534,10 +1728,10 @@ bool FPCGGraphExecutor::ExecuteScheduling(double EndTime, TSharedPtr<FPCGGraphAc
 	const bool bIsInGameThread = IsInGameThread();
 	TRACE_CPUPROFILER_EVENT_SCOPE(FPCGGraphExecutor::ExecuteScheduling);
 	
-	const float MaxPercentageOfThreadsToUse = FMath::Clamp(CVarMaxPercentageOfThreadsToUse.GetValueOnAnyThread(), 0.0f, 1.0f);
+	const float MaxPercentageOfThreadsToUse = FMath::Clamp(GetMaxPercentageOfThreadsToUse(), 0.0f, 1.0f);
 
 	// Clamp to 0.8 to never allow starving of PCGAsync Tasks
-	const float MaxPercentageOfExecutingThreads = FMath::Clamp(CVarMaxPercentageOfExecutingThreads.GetValueOnAnyThread(), 0.0f, 0.8f);
+	const float MaxPercentageOfExecutingThreads = FMath::Clamp(GetMaxPercentageOfExecutingThreads(), 0.0f, 0.8f);
 
 	// Number of threads to use across PCG
 	const int32 MaxNumThreads = FApp::ShouldUseThreadingForPerformance() ? FMath::Max(0, (int32)(LowLevelTasks::FScheduler::Get().GetNumWorkers() * MaxPercentageOfThreadsToUse)) : 0;
@@ -1546,10 +1740,10 @@ bool FPCGGraphExecutor::ExecuteScheduling(double EndTime, TSharedPtr<FPCGGraphAc
 	const int32 MaxExecutingThreadsPriorHardCap = PCGGraphExecutor::CVarGraphMultithreading.GetValueOnAnyThread() ? FMath::Max(0, (int32)(MaxNumThreads * MaxPercentageOfExecutingThreads)) : 0;
 
 	// Final number of threads to use for PCG Element execution (outside of main thread tick)
-	const int32 MaxExecutingThreads = FMath::Min(MaxExecutingThreadsPriorHardCap, FMath::Max(0, CVarMaxExecutingThreads.GetValueOnAnyThread()));
+	const int32 MaxExecutingThreads = FMath::Min(MaxExecutingThreadsPriorHardCap, FMath::Max(0, GetMaxExecutingThreads()));
 
 	// Number of threads to use for PCG Async operations 
-	const int32 MaxPCGAsyncThreads = FMath::Min(MaxNumThreads - MaxExecutingThreads, CVarMaxNumTasks.GetValueOnAnyThread());
+	const int32 MaxPCGAsyncThreads = FMath::Min(MaxNumThreads - MaxExecutingThreads, GetMaxNumProcessingTasks());
 
 	FCachedResults CachedResults;
 	{
diff --git a/Engine/Plugins/PCG/Source/PCG/Private/Graph/PCGGraphExecutor.h b/Engine/Plugins/PCG/Source/PCG/Private/Graph/PCGGraphExecutor.h
--- a/Engine/Plugins/PCG/Source/PCG/Private/Graph/PCGGraphExecutor.h
+++ b/Engine/Plugins/PCG/Source/PCG/Private/Graph/PCGGraphExecutor.h
@@ -19,7 +19,6 @@
 
 #if WITH_EDITOR
 #include "Editor/IPCGEditorProgressNotification.h"
-#include "WorldPartition/WorldPartitionHandle.h" // Needed for FWorldPartitionReference
 #endif
 
 #include "PCGGraphExecutor.generated.h"
@@ -186,8 +185,12 @@ private:
 	void SetLocalSpaceBounds(FPCGTaskId InGraphExecutionTaskId, const FBox& InBounds) { PerExecutionCache.SetExecutionCacheBounds(InGraphExecutionTaskId, EPCGPerExecutionCacheDataType::LocalSpaceBounds, InBounds); }
 	void SetOriginalLocalSpaceBounds(FPCGTaskId InGraphExecutionTaskId, const FBox& InBounds) { PerExecutionCache.SetExecutionCacheBounds(InGraphExecutionTaskId, EPCGPerExecutionCacheDataType::OriginalLocalSpaceBounds, InBounds); }
 
-	static double GetTickBudgetInSeconds();
-		
+	double GetTickBudgetInSeconds() const;
+	int32 GetMaxNumProcessingTasks() const;
+	float GetMaxPercentageOfThreadsToUse() const;
+	float GetMaxPercentageOfExecutingThreads() const;
+	int32 GetMaxExecutingThreads() const;
+
 	struct FCachedResult
 	{
 		FPCGTaskId TaskId = InvalidPCGTaskId;
diff --git a/Engine/Plugins/PCG/Source/PCG/Private/Subsystems/IPCGBaseSubsystem.cpp b/Engine/Plugins/PCG/Source/PCG/Private/Subsystems/IPCGBaseSubsystem.cpp
--- a/Engine/Plugins/PCG/Source/PCG/Private/Subsystems/IPCGBaseSubsystem.cpp
+++ b/Engine/Plugins/PCG/Source/PCG/Private/Subsystems/IPCGBaseSubsystem.cpp
@@ -74,7 +74,13 @@ void IPCGBaseSubsystem::CancelGeneration(IPCGGraphExecutionSource* Source)
 
 double IPCGBaseSubsystem::GetTickEndTime() const
 {
-	return FPlatformTime::Seconds() + FPCGGraphExecutor::GetTickBudgetInSeconds();
+	if (GraphExecutor)
+	{
+		return FPlatformTime::Seconds() + GraphExecutor->GetTickBudgetInSeconds();
+	}
+
+	// No Graph Executor no budget
+	return FPlatformTime::Seconds();
 }
 
 double IPCGBaseSubsystem::Tick()
