[PCG] Establish new base class for 2D textures (single and arrays).

Hierarchy design:

-Texture Base (future work, if use case arises)
----Texture 2D Base
--------Single Texture 2D Base
------------Texture 2D
------------Render Target 2D
--------Texture 2D Array
------------(Can flesh out hierarchy here for texture array and render target array if/when needed in future)
----Texture 3D (future work, when use case arises)

Replaces (abstract) class UPCGBaseTextureData with something more specific to this part of the data hierarchy: UPCGSingleSliceTextureBaseData. Some members of this class will likely be hoisted up over time, for example if/when texture array data supports CPU readback, but some logic will likely always diverge between array and non array textures (like we need to decide what SamplePoint returns for texture arrays).

#Synced-CL 49460467

diff --git a/Engine/Plugins/PCG/Config/DefaultPCG.ini b/Engine/Plugins/PCG/Config/DefaultPCG.ini
--- a/Engine/Plugins/PCG/Config/DefaultPCG.ini
+++ b/Engine/Plugins/PCG/Config/DefaultPCG.ini
@@ -104,9 +104,10 @@
 
 +StructRedirects=(OldName="PCGPointFilterThresholdSettings",NewName="/Script/PCG.PCGAttributeFilterThresholdSettings")
 +StructRedirects=(OldName="PCGActorPropertyOverrideDescription",NewName="/Script/PCG.PCGObjectPropertyOverrideDescription")
-+StructRedirects=(OldName="PCGDataTypeInfoBaseTexture",NewName="/Script/PCG.PCGDataTypeInfoBaseTexture2D")
 +StructRedirects=(OldName="PCGDataTypeInfoTexture",NewName="/Script/PCG.PCGDataTypeInfoTexture2D")
 +StructRedirects=(OldName="PCGDataTypeInfoRenderTarget",NewName="/Script/PCG.PCGDataTypeInfoRenderTarget2D")
++StructRedirects=(OldName="PCGDataTypeInfoBaseTexture",NewName="/Script/PCG.PCGDataTypeInfoTexture2DSingleBase")
++StructRedirects=(OldName="PCGDataTypeInfoBaseTexture2D",NewName="/Script/PCG.PCGDataTypeInfoTexture2DSingleBase")
 
 +EnumRedirects=(OldName="EPCGDensityNoiseMode",NewName="EPCGAttributeNoiseMode")
 +EnumRedirects=(OldName="EPCGPointFilterOperator",NewName="EPCGAttributeFilterOperator")
diff --git a/Engine/Plugins/PCG/Source/PCG/Private/Compute/DataInterfaces/PCGTextureDataInterface.cpp b/Engine/Plugins/PCG/Source/PCG/Private/Compute/DataInterfaces/PCGTextureDataInterface.cpp
--- a/Engine/Plugins/PCG/Source/PCG/Private/Compute/DataInterfaces/PCGTextureDataInterface.cpp
+++ b/Engine/Plugins/PCG/Source/PCG/Private/Compute/DataInterfaces/PCGTextureDataInterface.cpp
@@ -277,7 +277,7 @@ void UPCGTextureDataProvider::BuildInfosFromDataCollection(UPCGDataBinding* InBi
 	for (int I = 0; I < InputTaggedData.Num(); ++I)
 	{
 		const FPCGTaggedData& TaggedData = InputTaggedData[I];
-		const UPCGBaseTextureData* TextureData = Cast<UPCGBaseTextureData>(TaggedData.Data);
+		const UPCGTexture2DSingleBaseData* TextureData = Cast<UPCGTexture2DSingleBaseData>(TaggedData.Data);
 
 		if (!ensure(TextureData))
 		{
diff --git a/Engine/Plugins/PCG/Source/PCG/Private/Compute/DataInterfaces/PCGTextureDataInterface.h b/Engine/Plugins/PCG/Source/PCG/Private/Compute/DataInterfaces/PCGTextureDataInterface.h
--- a/Engine/Plugins/PCG/Source/PCG/Private/Compute/DataInterfaces/PCGTextureDataInterface.h
+++ b/Engine/Plugins/PCG/Source/PCG/Private/Compute/DataInterfaces/PCGTextureDataInterface.h
@@ -13,7 +13,7 @@
 #include "PCGTextureDataInterface.generated.h"
 
 class FPCGTextureDataInterfaceParameters;
-class UPCGBaseTextureData;
+class UPCGTexture2DSingleBaseData;
 
 namespace PCGTextureDataInterfaceConstants
 {
diff --git a/Engine/Plugins/PCG/Source/PCG/Private/Compute/DataInterfaces/PCGTextureInfo.cpp b/Engine/Plugins/PCG/Source/PCG/Private/Compute/DataInterfaces/PCGTextureInfo.cpp
--- a/Engine/Plugins/PCG/Source/PCG/Private/Compute/DataInterfaces/PCGTextureInfo.cpp
+++ b/Engine/Plugins/PCG/Source/PCG/Private/Compute/DataInterfaces/PCGTextureInfo.cpp
@@ -8,7 +8,7 @@
 
 #include "RHIDefinitions.h"
 
-FPCGTextureBindingInfo::FPCGTextureBindingInfo(const UPCGBaseTextureData* InTextureData)
+FPCGTextureBindingInfo::FPCGTextureBindingInfo(const UPCGTexture2DSingleBaseData* InTextureData)
 {
 	check(InTextureData);
 
diff --git a/Engine/Plugins/PCG/Source/PCG/Private/Compute/DataInterfaces/PCGTextureInfo.h b/Engine/Plugins/PCG/Source/PCG/Private/Compute/DataInterfaces/PCGTextureInfo.h
--- a/Engine/Plugins/PCG/Source/PCG/Private/Compute/DataInterfaces/PCGTextureInfo.h
+++ b/Engine/Plugins/PCG/Source/PCG/Private/Compute/DataInterfaces/PCGTextureInfo.h
@@ -4,14 +4,14 @@
 
 #include "Data/PCGTextureData.h"
 
-class UPCGBaseTextureData;
+class UPCGTexture2DSingleBaseData;
 class UPCGTexture2DArrayData;
 struct FPCGDataDesc;
 
 struct FPCGTextureBindingInfo
 {
 	FPCGTextureBindingInfo() = default;
-	explicit FPCGTextureBindingInfo(const UPCGBaseTextureData* InTextureData);
+	explicit FPCGTextureBindingInfo(const UPCGTexture2DSingleBaseData* InTextureData);
 	explicit FPCGTextureBindingInfo(const UPCGTexture2DArrayData* InTextureData);
 	FPCGTextureBindingInfo(const FPCGDataDesc& InDataDesc);
 	bool operator==(const FPCGTextureBindingInfo& Other) const;
diff --git a/Engine/Plugins/PCG/Source/PCG/Private/Compute/Elements/PCGCustomHLSL.cpp b/Engine/Plugins/PCG/Source/PCG/Private/Compute/Elements/PCGCustomHLSL.cpp
--- a/Engine/Plugins/PCG/Source/PCG/Private/Compute/Elements/PCGCustomHLSL.cpp
+++ b/Engine/Plugins/PCG/Source/PCG/Private/Compute/Elements/PCGCustomHLSL.cpp
@@ -1101,7 +1101,7 @@ void UPCGCustomHLSLSettings::UpdatePinSettings()
 		}
 
 		// Enforce single data on pin types that don't support a dynamic number of data.
-		if (Properties.AllowedTypes.IsChildOf(FPCGDataTypeInfoBaseTexture2D::AsId()))
+		if (Properties.AllowedTypes.IsChildOf(FPCGDataTypeInfoTexture2DSingleBase::AsId()))
 		{
 			Properties.bAllowMultipleData = false;
 			Properties.bAllowEditMultipleData = false;
diff --git a/Engine/Plugins/PCG/Source/PCG/Private/Compute/PCGComputeCommon.cpp b/Engine/Plugins/PCG/Source/PCG/Private/Compute/PCGComputeCommon.cpp
--- a/Engine/Plugins/PCG/Source/PCG/Private/Compute/PCGComputeCommon.cpp
+++ b/Engine/Plugins/PCG/Source/PCG/Private/Compute/PCGComputeCommon.cpp
@@ -63,9 +63,9 @@ namespace PCGComputeHelpers
 		{
 			ElementCount = Proxy->GetElementCount();
 		}
-		else if (const UPCGBaseTextureData* TextureData = Cast<UPCGBaseTextureData>(InData))
+		else if (const UPCGTexture2DSingleBaseData* TextureData = Cast<UPCGTexture2DSingleBaseData>(InData))
 		{
-			const FIntPoint TextureSize = TextureData->GetTextureSize();
+			const FIntPoint TextureSize = TextureData->GetResolution();
 			ElementCount.X = TextureSize.X;
 			ElementCount.Y = TextureSize.Y;
 		}
diff --git a/Engine/Plugins/PCG/Source/PCG/Private/Compute/PCGComputeGraph.cpp b/Engine/Plugins/PCG/Source/PCG/Private/Compute/PCGComputeGraph.cpp
--- a/Engine/Plugins/PCG/Source/PCG/Private/Compute/PCGComputeGraph.cpp
+++ b/Engine/Plugins/PCG/Source/PCG/Private/Compute/PCGComputeGraph.cpp
@@ -221,7 +221,7 @@ void UPCGComputeGraph::DebugLogDataDescriptions(const UPCGDataBinding* InBinding
 			UE_LOG(LogPCG, Display, TEXT("\t\t\tType: %s"), *Type.ToString());
 			UE_LOG(LogPCG, Display, TEXT("\t\t\tNum Elements: %s (%s)"), *DataDesc.GetElementCount().ToString(), *UEnum::GetValueAsString(DataDesc.GetElementDimension()));
 
-			if (Type.IsChildOf(FPCGDataTypeInfoBaseTexture2D::AsId()) || Type.IsChildOf(FPCGDataTypeInfoTexture2DArray::AsId()))
+			if (Type.IsChildOf(FPCGDataTypeInfoTexture2DSingleBase::AsId()) || Type.IsChildOf(FPCGDataTypeInfoTexture2DArray::AsId()))
 			{
 				UE_LOG(LogPCG, Display, TEXT("\t\t\tFormat: %s"), *UEnum::GetValueAsString(DataDesc.GetRenderTargetFormat()));
 				UE_LOG(LogPCG, Display, TEXT("\t\t\tTransform: %s"), *DataDesc.GetTextureTransform().ToString());
diff --git a/Engine/Plugins/PCG/Source/PCG/Private/Compute/PCGDataDescription.cpp b/Engine/Plugins/PCG/Source/PCG/Private/Compute/PCGDataDescription.cpp
--- a/Engine/Plugins/PCG/Source/PCG/Private/Compute/PCGDataDescription.cpp
+++ b/Engine/Plugins/PCG/Source/PCG/Private/Compute/PCGDataDescription.cpp
@@ -261,18 +261,11 @@ FPCGDataDesc::FPCGDataDesc(const FPCGTaggedData& InTaggedData, const UPCGDataBin
 	ElementDimension = PCGComputeHelpers::GetElementDimension(InTaggedData.Data);
 
 	// @todo_pcg: Implement a way for a UPCGData to generate its data description.
-	if (const UPCGTexture2DArrayData* Texture2DArrayData = Cast<UPCGTexture2DArrayData>(InTaggedData.Data))
+	if (const UPCGTexture2DBaseData* TextureData = Cast<UPCGTexture2DBaseData>(InTaggedData.Data))
 	{
-		SetTextureArraySize(Texture2DArrayData->GetArraySize());
-		SetRenderTargetFormat(PCGComputeHelpers::GetPCGRenderTargetFormatFromPixelFormat(Texture2DArrayData->GetFormat()));
-		SetTextureFilter(Texture2DArrayData->GetFilter());
-		SetTextureTransform(Texture2DArrayData->GetTransform());
-	}
-
-	if (const UPCGBaseTextureData* TextureData = Cast<UPCGBaseTextureData>(InTaggedData.Data))
-	{
-		SetTextureFilter(TextureData->Filter);
-		SetRenderTargetFormat(PCGComputeHelpers::GetPCGRenderTargetFormatFromPixelFormat(TextureData->GetTextureFormat()));
+		SetTextureArraySize(TextureData->GetArraySize());
+		SetRenderTargetFormat(PCGComputeHelpers::GetPCGRenderTargetFormatFromPixelFormat(TextureData->GetFormat()));
+		SetTextureFilter(TextureData->GetFilter());
 		SetTextureTransform(TextureData->GetTransform());
 	}
 
diff --git a/Engine/Plugins/PCG/Source/PCG/Private/Data/PCGTexture2DArrayData.cpp b/Engine/Plugins/PCG/Source/PCG/Private/Data/PCGTexture2DArrayData.cpp
--- a/Engine/Plugins/PCG/Source/PCG/Private/Data/PCGTexture2DArrayData.cpp
+++ b/Engine/Plugins/PCG/Source/PCG/Private/Data/PCGTexture2DArrayData.cpp
@@ -43,7 +43,7 @@ bool UPCGTexture2DArrayData::Initialize(UTexture* InTexture, const FPCGTexture2D
 
 bool UPCGTexture2DArrayData::Initialize(TRefCountPtr<IPooledRenderTarget> InExportedTextureHandle, const FPCGTexture2DArrayDataInitParams& InParams)
 {
-	ExportedTextureHandle = InExportedTextureHandle;
+	TextureHandle = InExportedTextureHandle;
 	ResourceType = EPCGTextureResourceType::ExportedTexture;
 
 	return InitializeInternal(InParams);
@@ -58,7 +58,7 @@ FTextureRHIRef UPCGTexture2DArrayData::GetRHI() const
 	}
 	else
 	{
-		return ExportedTextureHandle ? ExportedTextureHandle->GetRHI() : nullptr;
+		return TextureHandle ? TextureHandle->GetRHI() : nullptr;
 	}
 }
 
@@ -96,7 +96,7 @@ void UPCGTexture2DArrayData::CopyTexture2DArrayData(UPCGTexture2DArrayData* NewT
 	CopyBaseSurfaceData(NewTextureData);
 
 	NewTextureData->Texture = Texture;
-	NewTextureData->ExportedTextureHandle = ExportedTextureHandle;
+	NewTextureData->TextureHandle = TextureHandle;
 	NewTextureData->ResourceType = ResourceType;
 	NewTextureData->Resolution = Resolution;
 	NewTextureData->ArraySize = ArraySize;
diff --git a/Engine/Plugins/PCG/Source/PCG/Private/Data/PCGTexture2DBaseData.cpp b/Engine/Plugins/PCG/Source/PCG/Private/Data/PCGTexture2DBaseData.cpp
new file mode 100644
--- /dev/null
+++ b/Engine/Plugins/PCG/Source/PCG/Private/Data/PCGTexture2DBaseData.cpp
@@ -0,0 +1,25 @@
+// Copyright Epic Games, Inc. All Rights Reserved.
+
+#include "Data/PCGTexture2DBaseData.h"
+
+#include "PCGContext.h"
+#include "Data/PCGBasePointData.h"
+#include "Data/PCGPointArrayData.h"
+#include "Data/PCGPointData.h"
+
+#include "TextureResource.h"
+#include "Engine/Texture.h"
+
+#include UE_INLINE_GENERATED_CPP_BY_NAME(PCGTexture2DBaseData)
+
+PCG_DEFINE_TYPE_INFO(FPCGDataTypeInfoTexture2DBase, UPCGTexture2DBaseData)
+
+void UPCGTexture2DBaseData::CopyTexture2DBaseData(UPCGTexture2DBaseData* NewTextureData) const
+{
+	CopyBaseSurfaceData(NewTextureData);
+
+	NewTextureData->Format = Format;
+	NewTextureData->Filter = Filter;
+	NewTextureData->TextureHandle = TextureHandle;
+	NewTextureData->Bounds = Bounds;
+}
diff --git a/Engine/Plugins/PCG/Source/PCG/Private/Data/PCGTextureData.cpp b/Engine/Plugins/PCG/Source/PCG/Private/Data/PCGTextureData.cpp
--- a/Engine/Plugins/PCG/Source/PCG/Private/Data/PCGTextureData.cpp
+++ b/Engine/Plugins/PCG/Source/PCG/Private/Data/PCGTextureData.cpp
@@ -18,7 +18,7 @@
 
 #include UE_INLINE_GENERATED_CPP_BY_NAME(PCGTextureData)
 
-PCG_DEFINE_TYPE_INFO(FPCGDataTypeInfoBaseTexture2D, UPCGBaseTextureData)
+PCG_DEFINE_TYPE_INFO(FPCGDataTypeInfoTexture2DSingleBase, UPCGTexture2DSingleBaseData)
 PCG_DEFINE_TYPE_INFO(FPCGDataTypeInfoTexture2D, UPCGTextureData)
 
 namespace PCGTextureSamplingHelpers
@@ -122,7 +122,7 @@ namespace PCGTextureSamplingHelpers
 	template<typename ValueType>
 	bool Sample(const FVector2D& InPosition,
 		const FBox2D& InSurface,
-		const UPCGBaseTextureData* InTextureData,
+		const UPCGTexture2DSingleBaseData* InTextureData,
 		int32 Width,
 		int32 Height,
 		ValueType& SampledValue,
@@ -193,7 +193,7 @@ namespace PCGTextureSamplingHelpers
 	}
 }
 
-void UPCGBaseTextureData::PostLoad()
+void UPCGTexture2DSingleBaseData::PostLoad()
 {
 	Super::PostLoad();
 
@@ -208,28 +208,18 @@ void UPCGBaseTextureData::PostLoad()
 }
 
 PRAGMA_DISABLE_DEPRECATION_WARNINGS
-EPCGTextureDensityFunction UPCGBaseTextureData::GetDensityFunctionEquivalent() const
+EPCGTextureDensityFunction UPCGTexture2DSingleBaseData::GetDensityFunctionEquivalent() const
 {
 	return bUseDensitySourceChannel ? EPCGTextureDensityFunction::Multiply : EPCGTextureDensityFunction::Ignore;
 }
 
-void UPCGBaseTextureData::SetDensityFunctionEquivalent(EPCGTextureDensityFunction InDensityFunction)
+void UPCGTexture2DSingleBaseData::SetDensityFunctionEquivalent(EPCGTextureDensityFunction InDensityFunction)
 {
 	bUseDensitySourceChannel = (InDensityFunction != EPCGTextureDensityFunction::Ignore);
 }
 PRAGMA_ENABLE_DEPRECATION_WARNINGS
 
-FBox UPCGBaseTextureData::GetBounds() const
-{
-	return Bounds;
-}
-
-FBox UPCGBaseTextureData::GetStrictBounds() const
-{
-	return Bounds;
-}
-
-bool UPCGBaseTextureData::SamplePoint(const FTransform& InTransform, const FBox& InBounds, FPCGPoint& OutPoint, UPCGMetadata* OutMetadata) const
+bool UPCGTexture2DSingleBaseData::SamplePoint(const FTransform& InTransform, const FBox& InBounds, FPCGPoint& OutPoint, UPCGMetadata* OutMetadata) const
 {
 	// TODO: add metadata support
 	// TODO: add sampling along the bounds
@@ -280,19 +270,19 @@ bool UPCGBaseTextureData::SamplePoint(const FTransform& InTransform, const FBox&
 	}
 }
 
-const UPCGPointData* UPCGBaseTextureData::CreatePointData(FPCGContext* Context) const
+const UPCGPointData* UPCGTexture2DSingleBaseData::CreatePointData(FPCGContext* Context) const
 {
 	return CastChecked<UPCGPointData>(CreateBasePointData(Context, UPCGPointData::StaticClass()));
 }
 
-const UPCGPointArrayData* UPCGBaseTextureData::CreatePointArrayData(FPCGContext* Context, const FBox& InBounds) const
+const UPCGPointArrayData* UPCGTexture2DSingleBaseData::CreatePointArrayData(FPCGContext* Context, const FBox& InBounds) const
 {
 	return CastChecked<UPCGPointArrayData>(CreateBasePointData(Context, UPCGPointArrayData::StaticClass()));
 }
 
-const UPCGBasePointData* UPCGBaseTextureData::CreateBasePointData(FPCGContext* Context, TSubclassOf<UPCGBasePointData> PointDataClass) const
+const UPCGBasePointData* UPCGTexture2DSingleBaseData::CreateBasePointData(FPCGContext* Context, TSubclassOf<UPCGBasePointData> PointDataClass) const
 {
-	TRACE_CPUPROFILER_EVENT_SCOPE(UPCGBaseTextureData::CreatePointData);
+	TRACE_CPUPROFILER_EVENT_SCOPE(UPCGTexture2DSingleBaseData::CreatePointData);
 	FBox2D LocalSurfaceBounds(FVector2D(-1.0f, -1.0f), FVector2D(1.0f, 1.0f));
 
 	UPCGBasePointData* Data = FPCGContext::NewObject_AnyThread<UPCGBasePointData>(Context, GetTransientPackage(), PointDataClass);
@@ -397,12 +387,12 @@ const UPCGBasePointData* UPCGBaseTextureData::CreateBasePointData(FPCGContext* C
 	return Data;
 }
 
-bool UPCGBaseTextureData::IsValid() const
+bool UPCGTexture2DSingleBaseData::IsValid() const
 {
 	return Height > 0 && Width > 0 && (!ColorData.IsEmpty() || bSkipReadbackToCPU);
 }
 
-bool UPCGBaseTextureData::SamplePointLocal(const FVector2D& LocalPosition, FVector4& OutColor, float& OutDensity) const
+bool UPCGTexture2DSingleBaseData::SamplePointLocal(const FVector2D& LocalPosition, FVector4& OutColor, float& OutDensity) const
 {
 	if (!IsValid())
 	{
@@ -432,9 +422,9 @@ bool UPCGBaseTextureData::SamplePointLocal(const FVector2D& LocalPosition, FVect
 	return OutDensity > 0.0 || bKeepZeroDensityPoints;
 }
 
-void UPCGBaseTextureData::CopyBaseTextureData(UPCGBaseTextureData* NewTextureData) const
+void UPCGTexture2DSingleBaseData::CopyBaseTextureData(UPCGTexture2DSingleBaseData* NewTextureData) const
 {
-	CopyBaseSurfaceData(NewTextureData);
+	CopyTexture2DBaseData(NewTextureData);
 
 	NewTextureData->bUseDensitySourceChannel = bUseDensitySourceChannel;
 	NewTextureData->ColorChannel = ColorChannel;
@@ -446,7 +436,6 @@ void UPCGBaseTextureData::CopyBaseTextureData(UPCGBaseTextureData* NewTextureDat
 	NewTextureData->bUseTileBounds = bUseTileBounds;
 	NewTextureData->TileBounds = TileBounds;
 	NewTextureData->ColorData = ColorData;
-	NewTextureData->Bounds = Bounds;
 	NewTextureData->Height = Height;
 	NewTextureData->Width = Width;
 	NewTextureData->bSkipReadbackToCPU = bSkipReadbackToCPU;
@@ -708,7 +697,6 @@ UPCGSpatialData* UPCGTextureData::CopyInternal(FPCGContext* Context) const
 
 	NewTextureData->ResourceType = ResourceType;
 	NewTextureData->Texture = Texture;
-	NewTextureData->TextureHandle = TextureHandle;
 
 	return NewTextureData;
 }
diff --git a/Engine/Plugins/PCG/Source/PCG/Private/Data/Registry/PCGDataType.cpp b/Engine/Plugins/PCG/Source/PCG/Private/Data/Registry/PCGDataType.cpp
--- a/Engine/Plugins/PCG/Source/PCG/Private/Data/Registry/PCGDataType.cpp
+++ b/Engine/Plugins/PCG/Source/PCG/Private/Data/Registry/PCGDataType.cpp
@@ -156,7 +156,7 @@ FPCGDataTypeBaseId FPCGDataTypeBaseId::MakeFromLegacyType(EPCGDataType DataType)
 	case EPCGDataType::VirtualTexture:
 		return Construct<FPCGDataTypeInfoVirtualTexture>();
 	case EPCGDataType::BaseTexture:
-		return Construct<FPCGDataTypeInfoBaseTexture2D>();
+		return Construct<FPCGDataTypeInfoTexture2DSingleBase>();
 	case EPCGDataType::Surface:
 		return Construct<FPCGDataTypeInfoSurface>();
 	case EPCGDataType::Volume:
@@ -205,7 +205,7 @@ EPCGDataType FPCGDataTypeBaseId::AsLegacyType() const
 	PCG_DATA_TYPE_TEST(FPCGDataTypeInfoTexture2D)
 	PCG_DATA_TYPE_TEST(FPCGDataTypeInfoRenderTarget2D)
 	PCG_DATA_TYPE_TEST(FPCGDataTypeInfoVirtualTexture)
-	PCG_DATA_TYPE_TEST(FPCGDataTypeInfoBaseTexture2D)
+	PCG_DATA_TYPE_TEST(FPCGDataTypeInfoTexture2DSingleBase)
 	PCG_DATA_TYPE_TEST(FPCGDataTypeInfoSurface)
 	PCG_DATA_TYPE_TEST(FPCGDataTypeInfoVolume)
 	PCG_DATA_TYPE_TEST(FPCGDataTypeInfoPrimitive)
diff --git a/Engine/Plugins/PCG/Source/PCG/Private/Elements/PCGCopyToRenderTarget.cpp b/Engine/Plugins/PCG/Source/PCG/Private/Elements/PCGCopyToRenderTarget.cpp
--- a/Engine/Plugins/PCG/Source/PCG/Private/Elements/PCGCopyToRenderTarget.cpp
+++ b/Engine/Plugins/PCG/Source/PCG/Private/Elements/PCGCopyToRenderTarget.cpp
@@ -143,7 +143,7 @@ bool FPCGCopyToRenderTargetElement::ExecuteInternal(FPCGContext* InContext) cons
 
 		for (int Index = 0; Index < Inputs.Num(); ++Index)
 		{
-			const UPCGBaseTextureData* TextureData = Cast<UPCGBaseTextureData>(Inputs[Index].Data);
+			const UPCGTexture2DSingleBaseData* TextureData = Cast<UPCGTexture2DSingleBaseData>(Inputs[Index].Data);
 
 			if (!TextureData)
 			{
diff --git a/Engine/Plugins/PCG/Source/PCG/Private/Elements/PCGSampleTexture.cpp b/Engine/Plugins/PCG/Source/PCG/Private/Elements/PCGSampleTexture.cpp
--- a/Engine/Plugins/PCG/Source/PCG/Private/Elements/PCGSampleTexture.cpp
+++ b/Engine/Plugins/PCG/Source/PCG/Private/Elements/PCGSampleTexture.cpp
@@ -54,7 +54,7 @@ bool FPCGSampleTextureElement::ExecuteInternal(FPCGContext* Context) const
 		PCGE_LOG(Warning, GraphAndLog, LOCTEXT("InvalidNumberOfTextureData", "Only 1 texture input is allowed."));
 	}
 
-	const UPCGBaseTextureData* BaseTextureData = !BaseTextureInput.IsEmpty() ? Cast<UPCGBaseTextureData>(BaseTextureInput[0].Data) : nullptr;
+	const UPCGTexture2DSingleBaseData* BaseTextureData = !BaseTextureInput.IsEmpty() ? Cast<UPCGTexture2DSingleBaseData>(BaseTextureInput[0].Data) : nullptr;
 	if (!BaseTextureData)
 	{
 		return true;
@@ -105,7 +105,7 @@ bool FPCGSampleTextureElement::ExecuteInternal(FPCGContext* Context) const
 
 			EPCGPointNativeProperties PropertiesToAllocate = PointData->GetAllocatedProperties();
 
-			// This is based of the UPCGBaseTextureData::SamplePoint code, we should probably add a virtual to UPCGSpatialData to return the sampled properties to better futurproof this code
+			// This is based of the UPCGTexture2DSingleBaseData::SamplePoint code, we should probably add a virtual to UPCGSpatialData to return the sampled properties to better futurproof this code
 			PropertiesToAllocate |= (EPCGPointNativeProperties::Color | EPCGPointNativeProperties::Density);
 			if (Settings->TextureMappingMethod != EPCGTextureMappingMethod::UVCoordinates)
 			{
diff --git a/Engine/Plugins/PCG/Source/PCG/Private/Elements/PCGSaveTextureToAsset.cpp b/Engine/Plugins/PCG/Source/PCG/Private/Elements/PCGSaveTextureToAsset.cpp
--- a/Engine/Plugins/PCG/Source/PCG/Private/Elements/PCGSaveTextureToAsset.cpp
+++ b/Engine/Plugins/PCG/Source/PCG/Private/Elements/PCGSaveTextureToAsset.cpp
@@ -93,7 +93,7 @@ bool FPCGSaveTextureToAssetElement::ExecuteInternal(FPCGContext* InContext) cons
 			PCGLog::InputOutput::LogFirstInputOnlyWarning(PCGPinConstants::DefaultInputLabel, InContext);
 		}
 
-		Context->InputTextureData = Cast<const UPCGBaseTextureData>(Inputs[0].Data);
+		Context->InputTextureData = Cast<const UPCGTexture2DSingleBaseData>(Inputs[0].Data);
 
 		if (!Context->InputTextureData)
 		{
diff --git a/Engine/Plugins/PCG/Source/PCG/Private/Elements/PCGTextureSampler.cpp b/Engine/Plugins/PCG/Source/PCG/Private/Elements/PCGTextureSampler.cpp
--- a/Engine/Plugins/PCG/Source/PCG/Private/Elements/PCGTextureSampler.cpp
+++ b/Engine/Plugins/PCG/Source/PCG/Private/Elements/PCGTextureSampler.cpp
@@ -219,7 +219,7 @@ bool FPCGTextureSamplerElement::ExecuteInternal(FPCGContext* InContext) const
 	}
 
 	// The new texture data to add. Will be assigned the correct texture data based on the referenced texture.
-	UPCGBaseTextureData* BaseTextureData = nullptr;
+	UPCGTexture2DSingleBaseData* BaseTextureData = nullptr;
 
 	if (UTextureRenderTarget2D* RenderTarget = Cast<UTextureRenderTarget2D>(Texture))
 	{
diff --git a/Engine/Plugins/PCG/Source/PCG/Private/Tests/Registry/PCGDataTypeRegistryTests.cpp b/Engine/Plugins/PCG/Source/PCG/Private/Tests/Registry/PCGDataTypeRegistryTests.cpp
--- a/Engine/Plugins/PCG/Source/PCG/Private/Tests/Registry/PCGDataTypeRegistryTests.cpp
+++ b/Engine/Plugins/PCG/Source/PCG/Private/Tests/Registry/PCGDataTypeRegistryTests.cpp
@@ -63,7 +63,7 @@ bool FPCGByClassTest::RunTest(const FString& Parameters)
 		{TEXT("Polygon2D"), FPCGDataTypeIdentifier::Construct<UPCGPolygon2DData>()},
 		{TEXT("Landscape spline"), FPCGDataTypeIdentifier::Construct<UPCGLandscapeSplineData>()},
 		{TEXT("Landscape"), FPCGDataTypeIdentifier::Construct<UPCGLandscapeData>()},
-		{TEXT("Base texture"), FPCGDataTypeIdentifier::Construct<UPCGBaseTextureData>()},
+		{TEXT("Base texture"), FPCGDataTypeIdentifier::Construct<UPCGTexture2DSingleBaseData>()},
 		{TEXT("Texture"), FPCGDataTypeIdentifier::Construct<UPCGTextureData>()},
 		{TEXT("Proxy for GPU"), FPCGDataTypeIdentifier::Construct<UPCGProxyForGPUData>()},
 		{TEXT("Static mesh resource"), FPCGDataTypeIdentifier::Construct<UPCGStaticMeshResourceData>()},
diff --git a/Engine/Plugins/PCG/Source/PCG/Public/Data/PCGRenderTargetData.h b/Engine/Plugins/PCG/Source/PCG/Public/Data/PCGRenderTargetData.h
--- a/Engine/Plugins/PCG/Source/PCG/Public/Data/PCGRenderTargetData.h
+++ b/Engine/Plugins/PCG/Source/PCG/Public/Data/PCGRenderTargetData.h
@@ -11,8 +11,8 @@
 class UTexture;
 class UTextureRenderTarget2D;
 
-USTRUCT()
-struct FPCGDataTypeInfoRenderTarget2D : public FPCGDataTypeInfoBaseTexture2D
+USTRUCT(meta = (PCG_DataTypeDisplayName = "Render Target 2D"))
+struct FPCGDataTypeInfoRenderTarget2D : public FPCGDataTypeInfoTexture2DSingleBase
 {
 	GENERATED_BODY()
 
@@ -22,7 +22,7 @@ struct FPCGDataTypeInfoRenderTarget2D : public FPCGDataTypeInfoBaseTexture2D
 //TODO: It's possible that caching the result in this class is not as efficient as it could be
 // if we expect to sample in different ways (e.g. channel) in the same render target
 UCLASS(MinimalAPI, BlueprintType, ClassGroup = (Procedural))
-class UPCGRenderTargetData : public UPCGBaseTextureData
+class UPCGRenderTargetData : public UPCGTexture2DSingleBaseData
 {
 	GENERATED_BODY()
 
@@ -34,11 +34,11 @@ public:
 	virtual void ReleaseTransientResources(const TCHAR* InReason = nullptr) override;
 	// ~End UPCGData interface
 
-	//~Begin UPCGBaseTextureData interface
+	//~Begin UPCGTexture2DSingleBaseData interface
 	virtual UTexture* GetTexture() const override;
 	virtual FTextureRHIRef GetTextureRHI() const override;
 	virtual EPCGTextureResourceType GetTextureResourceType() const override { return EPCGTextureResourceType::TextureObject; }
-	//~End UPCGBaseTextureData interface
+	//~End UPCGTexture2DSingleBaseData interface
 
 	//~Begin UPCGSpatialData interface
 protected:
diff --git a/Engine/Plugins/PCG/Source/PCG/Public/Data/PCGTexture2DArrayData.h b/Engine/Plugins/PCG/Source/PCG/Public/Data/PCGTexture2DArrayData.h
--- a/Engine/Plugins/PCG/Source/PCG/Public/Data/PCGTexture2DArrayData.h
+++ b/Engine/Plugins/PCG/Source/PCG/Public/Data/PCGTexture2DArrayData.h
@@ -2,7 +2,7 @@
 
 #pragma once
 
-#include "Data/PCGSurfaceData.h"
+#include "Data/PCGTexture2DBaseData.h"
 
 #include "PCGCommon.h"
 
@@ -20,8 +20,8 @@ struct FPCGContext;
 struct FPCGPoint;
 struct IPooledRenderTarget;
 
-USTRUCT()
-struct FPCGDataTypeInfoTexture2DArray : public FPCGDataTypeInfoSurface
+USTRUCT(meta = (PCG_DataTypeDisplayName = "Texture 2D Array"))
+struct FPCGDataTypeInfoTexture2DArray : public FPCGDataTypeInfoTexture2DBase
 {
 	GENERATED_BODY()
 
@@ -35,20 +35,18 @@ struct FPCGTexture2DArrayDataInitParams
 };
 
 UCLASS(MinimalAPI)
-class UPCGTexture2DArrayData : public UPCGSurfaceData
+class UPCGTexture2DArrayData : public UPCGTexture2DBaseData
 {
 	GENERATED_BODY()
 
 public:
-	// ~Begin UPCGData interface
+	//~Begin UPCGData interface
 	PCG_ASSIGN_TYPE_INFO(FPCGDataTypeInfoTexture2DArray);
 	virtual bool IsCacheable() const { return Super::IsCacheable() && ResourceType != EPCGTextureResourceType::ExportedTexture; }
-	virtual void ReleaseTransientResources(const TCHAR* InReason = nullptr) override { ExportedTextureHandle = nullptr; }
-	// ~End UPCGData interface
+	virtual void ReleaseTransientResources(const TCHAR* InReason = nullptr) override { TextureHandle = nullptr; }
+	//~End UPCGData interface
 
 	//~Begin UPCGSpatialData interface
-	virtual FBox GetBounds() const override { return Bounds; }
-	virtual FBox GetStrictBounds() const override { return Bounds; }
 	virtual bool SamplePoint(const FTransform& InTransform, const FBox& InBounds, FPCGPoint& OutPoint, UPCGMetadata* OutMetadata) const override { return false; }
 
 protected:
@@ -66,10 +64,8 @@ public:
 	UTexture* GetTexture() const { return Texture.Get(); }
 	PCG_API FTextureRHIRef GetRHI() const;
 	EPCGTextureResourceType GetResourceType() const { return ResourceType; }
-	TRefCountPtr<IPooledRenderTarget> GetRefCountedTexture() const { return ExportedTextureHandle; }
-	FIntPoint GetResolution() const { return Resolution; }
-	uint16 GetArraySize() const { return ArraySize; }
-	EPixelFormat GetFormat() const { return Format; }
+	FIntPoint GetResolution() const override { return Resolution; }
+	virtual uint16 GetArraySize() const override { return ArraySize; }
 	EPCGTextureFilter GetFilter() const { return Filter; }
 
 protected:
@@ -81,9 +77,6 @@ protected:
 	UPROPERTY()
 	TWeakObjectPtr<UTexture> Texture = nullptr;
 
-	/** If initialized from an exported texture this holds a reference to the resource. */
-	TRefCountPtr<IPooledRenderTarget> ExportedTextureHandle = nullptr;
-
 	/** The type of underlying resource that this texture data represents. */
 	UPROPERTY()
 	EPCGTextureResourceType ResourceType = EPCGTextureResourceType::TextureObject;
@@ -92,16 +85,5 @@ protected:
 	FIntPoint Resolution = FIntPoint::ZeroValue;
 
 	UPROPERTY()
-	uint16 ArraySize = 0;
-
-	/** Format of the underlying texture resource. */
-	UPROPERTY()
-	TEnumAsByte<EPixelFormat> Format = EPixelFormat::PF_Unknown;
-
-	/** Method used to determine the value for a sample based on the value of nearby texels. */
-	UPROPERTY()
-	EPCGTextureFilter Filter = EPCGTextureFilter::Bilinear;
-
-	UPROPERTY()
-	FBox Bounds = FBox(EForceInit::ForceInit);
+	uint16 ArraySize = 1;
 };
diff --git a/Engine/Plugins/PCG/Source/PCG/Public/Data/PCGTexture2DBaseData.h b/Engine/Plugins/PCG/Source/PCG/Public/Data/PCGTexture2DBaseData.h
new file mode 100644
--- /dev/null
+++ b/Engine/Plugins/PCG/Source/PCG/Public/Data/PCGTexture2DBaseData.h
@@ -0,0 +1,63 @@
+// Copyright Epic Games, Inc. All Rights Reserved.
+
+#pragma once
+
+#include "Data/PCGSurfaceData.h"
+
+#include "PCGCommon.h"
+
+#include "PixelFormat.h"
+#include "RendererInterface.h"
+#include "RHI.h"
+
+#include "PCGTexture2DBaseData.generated.h"
+
+/** Base type of 2D textures, 2D texture arrays, 2d render targets. */
+USTRUCT(meta = (PCG_DataTypeDisplayName = "Texture 2D Base"))
+struct FPCGDataTypeInfoTexture2DBase : public FPCGDataTypeInfoSurface
+{
+	GENERATED_BODY()
+
+	PCG_DECLARE_TYPE_INFO()
+};
+
+/** Base class of 2D textures, 2D texture arrays, 2d render targets. */
+UCLASS(MinimalAPI, Abstract)
+class UPCGTexture2DBaseData : public UPCGSurfaceData
+{
+	GENERATED_BODY()
+
+public:
+	// ~Begin UPCGData interface
+	PCG_ASSIGN_TYPE_INFO(FPCGDataTypeInfoTexture2DBase)
+	// ~End UPCGData interface
+
+	//~ Begin UPCGSpatialData interface
+	virtual FBox GetBounds() const override { return Bounds; }
+	virtual FBox GetStrictBounds() const override { return Bounds; }
+	//~ End UPCGSpatialData interface
+
+	virtual EPixelFormat GetFormat() const { return Format; }
+	virtual EPCGTextureFilter GetFilter() const { return Filter; }
+	virtual uint16 GetArraySize() const { return 1; }
+	virtual FIntPoint GetResolution() const PURE_VIRTUAL(UPCGTexture2DBaseData::GetResolution, return FIntPoint::ZeroValue;);
+	virtual TRefCountPtr<IPooledRenderTarget> GetRefCountedTexture() const { return TextureHandle; }
+
+protected:
+	PCG_API void CopyTexture2DBaseData(UPCGTexture2DBaseData* NewTextureData) const;
+
+protected:
+	UPROPERTY(BlueprintReadOnly, VisibleAnywhere, Category = SpatialData)
+	TEnumAsByte<EPixelFormat> Format = EPixelFormat::PF_Unknown;
+
+public:
+	UPROPERTY(BlueprintReadOnly, VisibleAnywhere, Category = SpatialData)
+	EPCGTextureFilter Filter = EPCGTextureFilter::Bilinear;
+
+protected:
+	/** If initialized from an exported texture this holds a reference to the resource. */
+	TRefCountPtr<IPooledRenderTarget> TextureHandle = nullptr;
+
+	UPROPERTY(BlueprintReadOnly, VisibleAnywhere, Category = SpatialData)
+	FBox Bounds = FBox(EForceInit::ForceInit);
+};
diff --git a/Engine/Plugins/PCG/Source/PCG/Public/Data/PCGTextureData.h b/Engine/Plugins/PCG/Source/PCG/Public/Data/PCGTextureData.h
--- a/Engine/Plugins/PCG/Source/PCG/Public/Data/PCGTextureData.h
+++ b/Engine/Plugins/PCG/Source/PCG/Public/Data/PCGTextureData.h
@@ -2,7 +2,7 @@
 
 #pragma once
 
-#include "Data/PCGSurfaceData.h"
+#include "Data/PCGTexture2DBaseData.h"
 
 #include "PCGCommon.h"
 
@@ -51,40 +51,30 @@ namespace PCGTextureSamplingHelpers
 }
 
 /** Base type of 2D textures/render targets. */
-USTRUCT()
-struct FPCGDataTypeInfoBaseTexture2D : public FPCGDataTypeInfoSurface
+USTRUCT(meta = (PCG_DataTypeDisplayName = "Single Texture 2D Base"))
+struct FPCGDataTypeInfoTexture2DSingleBase : public FPCGDataTypeInfoTexture2DBase
 {
 	GENERATED_BODY()
 
 	PCG_DECLARE_LEGACY_TYPE_INFO(EPCGDataType::BaseTexture)
 };
 
-USTRUCT()
-struct FPCGDataTypeInfoTexture2D : public FPCGDataTypeInfoBaseTexture2D
-{
-	GENERATED_BODY()
-
-	PCG_DECLARE_LEGACY_TYPE_INFO(EPCGDataType::Texture)
-};
-
 /** Base class for a 2D texture or render target. */
 UCLASS(MinimalAPI, Abstract)
-class UPCGBaseTextureData : public UPCGSurfaceData
+class UPCGTexture2DSingleBaseData : public UPCGTexture2DBaseData
 {
 	GENERATED_BODY()
 
 public:
-	// ~Being UObject interface
+	//~Being UObject interface
 	UE_API virtual void PostLoad() override;
-	// ~End UObject interface
+	//~End UObject interface
 
-	// ~Begin UPCGData interface
-	PCG_ASSIGN_DEFAULT_TYPE_INFO(FPCGDataTypeInfoBaseTexture2D)
-	// ~End UPCGData interface
+	//~Begin UPCGData interface
+	PCG_ASSIGN_DEFAULT_TYPE_INFO(FPCGDataTypeInfoTexture2DSingleBase)
+	//~End UPCGData interface
 
 	//~Begin UPCGSpatialData interface
-	UE_API virtual FBox GetBounds() const override;
-	UE_API virtual FBox GetStrictBounds() const override;
 	UE_API virtual bool SamplePoint(const FTransform& Transform, const FBox& Bounds, FPCGPoint& OutPoint, UPCGMetadata* OutMetadata) const override;
 	//~End UPCGSpatialData interface
 
@@ -93,18 +83,19 @@ public:
 	UE_API virtual const UPCGPointArrayData* CreatePointArrayData(FPCGContext* Context, const FBox& InBounds) const override;
 	//~End UPCGSpatialDataWithPointCache interface
 
+	//~Begin UPCGTexture2DBaseData interface
+	virtual FIntPoint GetResolution() const override { return FIntPoint(Width, Height); }
+	//~End UPCGTexture2DBaseData interface
+
 	/** Sample using a local space 'UV' position. */
 	UE_API bool SamplePointLocal(const FVector2D& LocalPosition, FVector4& OutColor, float& OutDensity) const;
 
 	UE_API virtual bool IsValid() const;
 
-	virtual UTexture* GetTexture() const PURE_VIRTUAL(UPCGBaseTextureData::GetTexture, return nullptr;)
-	virtual FTextureRHIRef GetTextureRHI() const PURE_VIRTUAL(UPCGBaseTextureData::GetTextureResource, return nullptr;);
-	virtual EPCGTextureResourceType GetTextureResourceType() const PURE_VIRTUAL(UPCGBaseTextureData::GetTextureResourceType, return EPCGTextureResourceType::Invalid;);
-	virtual TRefCountPtr<IPooledRenderTarget> GetRefCountedTexture() const { return nullptr; }
+	virtual UTexture* GetTexture() const PURE_VIRTUAL(UPCGTexture2DSingleBaseData::GetTexture, return nullptr;)
+	virtual FTextureRHIRef GetTextureRHI() const PURE_VIRTUAL(UPCGTexture2DSingleBaseData::GetTextureResource, return nullptr;);
+	virtual EPCGTextureResourceType GetTextureResourceType() const PURE_VIRTUAL(UPCGTexture2DSingleBaseData::GetTextureResourceType, return EPCGTextureResourceType::Invalid;);
 	virtual int GetTextureSlice() const { return 0; }
-	virtual FIntPoint GetTextureSize() const { return FIntPoint(Width, Height); }
-	EPixelFormat GetTextureFormat() const { return Format; }
 
 protected:
 	UE_API const UPCGBasePointData* CreateBasePointData(FPCGContext* Context, TSubclassOf<UPCGBasePointData> PointDataClass) const;
@@ -132,10 +123,6 @@ public:
 	UPROPERTY(BlueprintReadWrite, EditAnywhere, Category = Settings, meta = (DisplayName = "Density Source Channel", EditCondition = "bUseDensitySourceChannel"))
 	EPCGTextureColorChannel ColorChannel = EPCGTextureColorChannel::Alpha;
 
-	/** Method used to determine the value for a sample based on the value of nearby texels. */
-	UPROPERTY(BlueprintReadWrite, EditAnywhere, Category = Settings, meta = (PCG_Overridable))
-	EPCGTextureFilter Filter = EPCGTextureFilter::Bilinear;
-
 	/** The size of one texel in cm, used when calling ToPointData. */
 	UPROPERTY(BlueprintReadWrite, EditAnywhere, Category = Settings, meta = (UIMin = "1.0", ClampMin = "1.0"))
 	float TexelSize = 50.0f;
@@ -164,30 +151,38 @@ protected:
 	UPROPERTY()
 	TArray<FLinearColor> ColorData;
 
-	UPROPERTY(BlueprintReadOnly, VisibleAnywhere, Category = SpatialData)
-	FBox Bounds = FBox(EForceInit::ForceInit);
-
 	UPROPERTY(BlueprintReadOnly, VisibleAnywhere, Category = SpatialData)
 	int32 Height = 0;
 
 	UPROPERTY(BlueprintReadOnly, VisibleAnywhere, Category = SpatialData)
 	int32 Width = 0;
 
-	/** Format of the underlying texture resource. Note that for render targets this is distinct from the render target format. */
-	UPROPERTY(BlueprintReadOnly, VisibleAnywhere, Category = SpatialData)
-	TEnumAsByte<EPixelFormat> Format = EPixelFormat::PF_Unknown;
-
 	UPROPERTY()
 	bool bSkipReadbackToCPU = false;
 
 	/** Used to make sure errors are only logged once when trying to sample points from a data which hasn't been read back into a CPU buffer. */
 	mutable bool bEmittedNoReadbackDataError = false;
 
-	UE_API void CopyBaseTextureData(UPCGBaseTextureData* NewTextureData) const;
+	UE_API void CopyBaseTextureData(UPCGTexture2DSingleBaseData* NewTextureData) const;
+
+	// Deprecated section
+public:
+	UE_DEPRECATED(5.8, "Use GetResolution instead")
+	virtual FIntPoint GetTextureSize() const { return GetResolution(); }
+	UE_DEPRECATED(5.8, "Use GetFormat instead")
+	EPixelFormat GetTextureFormat() const { return Format; }
+};
+
+USTRUCT(meta = (PCG_DataTypeDisplayName = "Texture 2D"))
+struct FPCGDataTypeInfoTexture2D : public FPCGDataTypeInfoTexture2DSingleBase
+{
+	GENERATED_BODY()
+
+	PCG_DECLARE_LEGACY_TYPE_INFO(EPCGDataType::Texture)
 };
 
 UCLASS(MinimalAPI, BlueprintType, ClassGroup=(Procedural))
-class UPCGTextureData : public UPCGBaseTextureData // Texture2D
+class UPCGTextureData : public UPCGTexture2DSingleBaseData // Texture2D
 {
 	GENERATED_BODY()
 
@@ -215,12 +210,12 @@ public:
 	virtual void ReleaseTransientResources(const TCHAR* InReason = nullptr) override { TextureHandle = nullptr; }
 	//~End UPCGData interface
 
-	//~Begin UPCGBaseTextureData interface
+	//~Begin UPCGTexture2DSingleBaseData interface
 	virtual UTexture* GetTexture() const override { return Texture.IsValid() ? Texture.Get() : nullptr; }
 	virtual FTextureRHIRef GetTextureRHI() const override;
 	virtual TRefCountPtr<IPooledRenderTarget> GetRefCountedTexture() const override { return TextureHandle; }
 	virtual int GetTextureSlice() const override { return TextureIndex; }
-	//~End UPCGBaseTextureData interface
+	//~End UPCGTexture2DSingleBaseData interface
 	
 	//~Begin UPCGSpatialData interface
 protected:
@@ -270,9 +265,6 @@ protected:
 	UPROPERTY()
 	EPCGTextureResourceType ResourceType = EPCGTextureResourceType::TextureObject;
 
-	/** If initialized from an exported texture this holds a reference to the resource. */
-	TRefCountPtr<IPooledRenderTarget> TextureHandle = nullptr;
-
 	bool bUpdatedReadbackTextureResource = false;
 
 public:
@@ -284,4 +276,15 @@ public:
 	bool bReadbackFromGPUInitiated_DEPRECATED = false;
 };
 
+// Deprecated section
+UCLASS(MinimalAPI, Abstract, meta = (DeprecationMessage = "Deprecated 5.8, use UPCGTexture2DSingleBaseData instead"))
+class UPCGBaseTextureData : public UPCGTexture2DBaseData
+{
+	GENERATED_BODY()
+
+public:
+	UE_DEPRECATED(5.8, "Use UPCGTexture2DSingleBaseData instead")
+	UPCGBaseTextureData() {}
+};
+
 #undef UE_API
diff --git a/Engine/Plugins/PCG/Source/PCG/Public/Elements/PCGCopyToRenderTarget.h b/Engine/Plugins/PCG/Source/PCG/Public/Elements/PCGCopyToRenderTarget.h
--- a/Engine/Plugins/PCG/Source/PCG/Public/Elements/PCGCopyToRenderTarget.h
+++ b/Engine/Plugins/PCG/Source/PCG/Public/Elements/PCGCopyToRenderTarget.h
@@ -8,7 +8,7 @@
 
 #include "PCGCopyToRenderTarget.generated.h"
 
-class UPCGBaseTextureData;
+class UPCGTexture2DSingleBaseData;
 class UTextureRenderTarget2D;
 
 /** Copy texture data into render target assets. Input textures should be 1:1 with the render target assets and must match in extents and format. */
@@ -62,7 +62,7 @@ protected:
 public:
 	bool bEnqueuedCopyCommand = false;
 	TArray<FSoftObjectPath> RenderTargetsToLoad;
-	TArray<TPair<const UPCGBaseTextureData*, TObjectPtr<UTextureRenderTarget2D>>> TexturesToCopy;
+	TArray<TPair<const UPCGTexture2DSingleBaseData*, TObjectPtr<UTextureRenderTarget2D>>> TexturesToCopy;
 };
 
 class FPCGCopyToRenderTargetElement : public IPCGElementWithCustomContext<FPCGCopyToRenderTargetContext>
diff --git a/Engine/Plugins/PCG/Source/PCG/Public/Elements/PCGSaveTextureToAsset.h b/Engine/Plugins/PCG/Source/PCG/Public/Elements/PCGSaveTextureToAsset.h
--- a/Engine/Plugins/PCG/Source/PCG/Public/Elements/PCGSaveTextureToAsset.h
+++ b/Engine/Plugins/PCG/Source/PCG/Public/Elements/PCGSaveTextureToAsset.h
@@ -8,7 +8,7 @@
 
 #include "PCGSaveTextureToAsset.generated.h"
 
-class UPCGBaseTextureData;
+class UPCGTexture2DSingleBaseData;
 class UTexture2D;
 
 /** Save the input texture to a UTexture2D asset (format is always BGRA8). Outputs a soft object path attribute for the saved texture. Editor only. */
@@ -46,7 +46,7 @@ protected:
 	virtual void AddExtraStructReferencedObjects(FReferenceCollector& Collector) override;
 
 public:
-	TObjectPtr<const UPCGBaseTextureData> InputTextureData = nullptr;
+	TObjectPtr<const UPCGTexture2DSingleBaseData> InputTextureData = nullptr;
 	TObjectPtr<UTexture2D> ExportedTexture = nullptr;
 	uint8* RawReadbackData = nullptr;
 	int32 ReadbackWidth = INDEX_NONE;
diff --git a/Engine/Plugins/PCG/Source/PCGEditor/Private/DataVisualizations/PCGBaseTextureDataVisualization.cpp b/Engine/Plugins/PCG/Source/PCGEditor/Private/DataVisualizations/PCGBaseTextureDataVisualization.cpp
--- a/Engine/Plugins/PCG/Source/PCGEditor/Private/DataVisualizations/PCGBaseTextureDataVisualization.cpp
+++ b/Engine/Plugins/PCG/Source/PCGEditor/Private/DataVisualizations/PCGBaseTextureDataVisualization.cpp
@@ -24,7 +24,7 @@ TArray<TSharedPtr<FStreamableHandle>> FPCGBaseTextureDataVisualization::LoadRequ
 
 FPCGSetupSceneFunc FPCGBaseTextureDataVisualization::GetViewportSetupFunc(const UPCGSettingsInterface* SettingsInterface, const UPCGData* Data) const
 {
-	return [this, WeakData=TWeakObjectPtr<const UPCGBaseTextureData>(Cast<UPCGBaseTextureData>(Data))](FPCGSceneSetupParams& InOutParams)
+	return [this, WeakData=TWeakObjectPtr<const UPCGTexture2DSingleBaseData>(Cast<UPCGTexture2DSingleBaseData>(Data))](FPCGSceneSetupParams& InOutParams)
 	{
 		check(InOutParams.Scene);
 		check(InOutParams.EditorViewportClient);
@@ -53,7 +53,7 @@ FPCGSetupSceneFunc FPCGBaseTextureDataVisualization::GetViewportSetupFunc(const
 		{
 			FPCGVisualizationTexture2DParams VisParams;
 			VisParams.TextureRHI = WeakData->GetTextureRHI();
-			VisParams.Resolution = WeakData->GetTextureSize();
+			VisParams.Resolution = WeakData->GetResolution();
 			VisParams.SliceIndex = WeakData->GetTextureSlice();
 
 			Texture = UPCGVisualizationTexture2D::Create(VisParams);
diff --git a/Engine/Plugins/PCG/Source/PCGEditor/Private/PCGEditorModule.cpp b/Engine/Plugins/PCG/Source/PCGEditor/Private/PCGEditorModule.cpp
--- a/Engine/Plugins/PCG/Source/PCGEditor/Private/PCGEditorModule.cpp
+++ b/Engine/Plugins/PCG/Source/PCGEditor/Private/PCGEditorModule.cpp
@@ -865,7 +865,7 @@ void FPCGEditorModule::RegisterPCGDataVisualizations()
 	DataVisRegistry.InternalRegistry.Add(UPCGCollisionWrapperData::StaticClass(), MakeUnique<const FPCGCollisionWrapperDataVisualization>());
 	DataVisRegistry.InternalRegistry.Add(UPCGLandscapeData::StaticClass(), MakeUnique<const IPCGLandscapeDataVisualization>());
 	DataVisRegistry.InternalRegistry.Add(UPCGStaticMeshResourceData::StaticClass(), MakeUnique<const IPCGStaticMeshDataVisualization>());
-	DataVisRegistry.InternalRegistry.Add(UPCGBaseTextureData::StaticClass(), MakeUnique<const FPCGBaseTextureDataVisualization>());
+	DataVisRegistry.InternalRegistry.Add(UPCGTexture2DSingleBaseData::StaticClass(), MakeUnique<const FPCGBaseTextureDataVisualization>());
 	DataVisRegistry.InternalRegistry.Add(UPCGTexture2DArrayData::StaticClass(), MakeUnique<const FPCGTexture2DArrayDataVisualization>());
 	DataVisRegistry.InternalRegistry.Add(UPCGPolygon2DData::StaticClass(), MakeUnique<const FPCGPolygon2DDataVisualization>());
 	DataVisRegistry.InternalRegistry.Add(UPCGPolygon2DInteriorSurfaceData::StaticClass(), MakeUnique<const FPCGPolygon2DDataVisualization>());
@@ -919,7 +919,7 @@ void FPCGEditorModule::RegisterPinColorAndIcons()
 	InRegistry.RegisterPinColorFunction(FPCGDataTypeInfoPolygon2D::AsId(), [](const FPCGDataTypeIdentifier&) { return GetDefault<UPCGEditorSettings>()->Polygon2DDataPinColor; });
 	InRegistry.RegisterPinColorFunction(FPCGDataTypeInfoPolyline::AsId(), [](const FPCGDataTypeIdentifier&) { return GetDefault<UPCGEditorSettings>()->PolyLineDataPinColor; });
 	InRegistry.RegisterPinColorFunction(FPCGDataTypeInfoLandscape::AsId(), [](const FPCGDataTypeIdentifier&) { return GetDefault<UPCGEditorSettings>()->LandscapeDataPinColor; });
-	InRegistry.RegisterPinColorFunction(FPCGDataTypeInfoBaseTexture2D::AsId(), [](const FPCGDataTypeIdentifier&) { return GetDefault<UPCGEditorSettings>()->BaseTextureDataPinColor; });
+	InRegistry.RegisterPinColorFunction(FPCGDataTypeInfoTexture2DSingleBase::AsId(), [](const FPCGDataTypeIdentifier&) { return GetDefault<UPCGEditorSettings>()->BaseTextureDataPinColor; });
 	InRegistry.RegisterPinColorFunction(FPCGDataTypeInfoTexture2D::AsId(), [](const FPCGDataTypeIdentifier&) { return GetDefault<UPCGEditorSettings>()->TextureDataPinColor; });
 	InRegistry.RegisterPinColorFunction(FPCGDataTypeInfoRenderTarget2D::AsId(), [](const FPCGDataTypeIdentifier&) { return GetDefault<UPCGEditorSettings>()->RenderTargetDataPinColor; });
 	InRegistry.RegisterPinColorFunction(FPCGDataTypeInfoVirtualTexture::AsId(), [](const FPCGDataTypeIdentifier&) { return GetDefault<UPCGEditorSettings>()->VirtualTextureDataPinColor; });
@@ -1008,7 +1008,7 @@ void FPCGEditorModule::UnregisterPinColorAndIcons()
 	InRegistry.UnregisterPinColorFunction(FPCGDataTypeInfoVirtualTexture::AsId());
 	InRegistry.UnregisterPinColorFunction(FPCGDataTypeInfoRenderTarget2D::AsId());
 	InRegistry.UnregisterPinColorFunction(FPCGDataTypeInfoTexture2D::AsId());
-	InRegistry.UnregisterPinColorFunction(FPCGDataTypeInfoBaseTexture2D::AsId());
+	InRegistry.UnregisterPinColorFunction(FPCGDataTypeInfoTexture2DSingleBase::AsId());
 	InRegistry.UnregisterPinColorFunction(FPCGDataTypeInfoLandscape::AsId());
 	InRegistry.UnregisterPinColorFunction(FPCGDataTypeInfoPolyline::AsId());
 	InRegistry.UnregisterPinColorFunction(FPCGDataTypeInfoPolygon2D::AsId());
