[PCG] Add Attribute Set Generator and massage Custom HLSL logic toward being a bit more generic.

Minimizes usages of EPCGKernelType, which is a step away from removing it entirely and making the Custom HLSL node extensible.

[RN] [PCG] Added Attribute Set Generator mode to Custom HLSL node.
#Synced-CL 49460467

diff --git a/Engine/Plugins/PCG/Source/PCG/Private/Compute/Elements/PCGCustomHLSL.cpp b/Engine/Plugins/PCG/Source/PCG/Private/Compute/Elements/PCGCustomHLSL.cpp
--- a/Engine/Plugins/PCG/Source/PCG/Private/Compute/Elements/PCGCustomHLSL.cpp
+++ b/Engine/Plugins/PCG/Source/PCG/Private/Compute/Elements/PCGCustomHLSL.cpp
@@ -6,7 +6,7 @@
 #include "PCGComponent.h"
 #include "PCGEdge.h"
 #include "PCGModule.h"
-#include "PCGPoint.h"
+#include "PCGParamData.h"
 #include "PCGSubsystem.h"
 #include "Compute/PCGComputeSource.h"
 #include "Compute/PCGKernelHelpers.h"
@@ -256,6 +256,44 @@ const FPCGPinPropertiesGPU* UPCGCustomHLSLSettings::GetFirstOutputPinProperties(
 	return !OutputPins.IsEmpty() ? &OutputPins[0] : nullptr;
 }
 
+EPCGElementDimension UPCGCustomHLSLSettings::GetElementDimension() const
+{
+	switch (KernelType)
+	{
+	case EPCGKernelType::PointGenerator:
+	case EPCGKernelType::PointProcessor:
+	case EPCGKernelType::AttributeSetProcessor:
+	case EPCGKernelType::AttributeSetGenerator:
+	case EPCGKernelType::Custom:
+		return EPCGElementDimension::One;
+	case EPCGKernelType::TextureGenerator:
+	case EPCGKernelType::TextureProcessor:
+		return EPCGElementDimension::Two;
+	default:
+		checkNoEntry();
+		return EPCGElementDimension::One;
+	}
+}
+
+FPCGDataTypeIdentifier UPCGCustomHLSLSettings::GetElementType() const
+{
+	switch (KernelType)
+	{
+	case EPCGKernelType::PointGenerator:
+	case EPCGKernelType::PointProcessor:
+		return FPCGDataTypeInfoPoint::AsId();
+	case EPCGKernelType::TextureGenerator:
+	case EPCGKernelType::TextureProcessor:
+		return FPCGDataTypeInfoTexture2DSingleBase::AsId();
+	case EPCGKernelType::AttributeSetProcessor:
+	case EPCGKernelType::AttributeSetGenerator:
+		return FPCGDataTypeInfoParam::AsId();
+	default:
+		// Custom kernel has no primary type and returns None.
+		return EPCGDataType::None;
+	}
+}
+
 #if WITH_EDITOR
 FString UPCGCustomHLSLSettings::GetDeclarationsText() const
 {
@@ -375,15 +413,24 @@ void UPCGCustomHLSLSettings::UpdateInputDeclarations()
 		if (IsGeneratorKernel())
 		{
 			InputDeclarations += TEXT("/*** INPUT CONSTANTS ***/\n\n");
-		}
 
-		if (KernelType == EPCGKernelType::PointGenerator)
-		{
-			InputDeclarations += TEXT("const uint NumElements;\n\n");
-		}
-		else if (KernelType == EPCGKernelType::TextureGenerator)
-		{
-			InputDeclarations += TEXT("const uint2 NumElements;\n\n");
+			switch (GetElementDimension())
+			{
+			case EPCGElementDimension::One:
+				InputDeclarations += TEXT("const uint NumElements;\n\n");
+				break;
+			case EPCGElementDimension::Two:
+				InputDeclarations += TEXT("const uint2 NumElements;\n\n");
+				break;
+			case EPCGElementDimension::Three:
+				InputDeclarations += TEXT("const uint3 NumElements;\n\n");
+				break;
+			case EPCGElementDimension::Four:
+				InputDeclarations += TEXT("const uint4 NumElements;\n\n");
+				break;
+			default:
+				checkNoEntry();
+			}
 		}
 
 		InputDeclarations += TEXT("/*** INPUT PER-THREAD CONSTANTS ***/\n\n");
@@ -412,13 +459,26 @@ void UPCGCustomHLSLSettings::UpdateInputDeclarations()
 			}
 		}
 
-		if (IsPointKernel())
+		if (!IsCustomKernel())
 		{
-			InputDeclarations += TEXT("const uint ElementIndex;\n");
-		}
-		else if (IsTextureKernel())
-		{
-			InputDeclarations += TEXT("const uint2 ElementIndex;\n");
+			switch (GetElementDimension())
+			{
+			case EPCGElementDimension::One:
+				InputDeclarations += TEXT("const uint ElementIndex;\n");
+				break;
+			case EPCGElementDimension::Two:
+				InputDeclarations += TEXT("const uint2 ElementIndex;\n");
+				break;
+			case EPCGElementDimension::Three:
+				InputDeclarations += TEXT("const uint3 ElementIndex;\n");
+				break;
+			case EPCGElementDimension::Four:
+				InputDeclarations += TEXT("const uint4 ElementIndex;\n");
+				break;
+			default:
+				ensure(false);
+				break;
+			}
 		}
 
 		InputDeclarations += TEXT("\n");
@@ -998,21 +1058,9 @@ void UPCGCustomHLSLSettings::UpdatePinSettings()
 {
 	if (IsProcessorKernel() && InputPins.IsEmpty())
 	{
-		if (IsPointKernel())
-		{
-			InputPins.Emplace(PCGPinConstants::DefaultInputLabel, EPCGDataType::Point);
-		}
-		else if (IsTextureKernel())
-		{
-			InputPins.Emplace(PCGPinConstants::DefaultInputLabel, EPCGDataType::BaseTexture);
-		}
-		else if (KernelType == EPCGKernelType::AttributeSetProcessor)
+		if (FPCGDataTypeIdentifier ElementType = GetElementType(); ensure(ElementType != EPCGDataType::None))
 		{
-			InputPins.Emplace(PCGPinConstants::DefaultInputLabel, EPCGDataType::Param);
-		}
-		else
-		{
-			ensure(false);
+			InputPins.Emplace(PCGPinConstants::DefaultInputLabel, ElementType);
 		}
 	}
 
@@ -1028,24 +1076,16 @@ void UPCGCustomHLSLSettings::UpdatePinSettings()
 		}
 
 		// Allow kernel type to drive the first pin type.
-		if (PinIndex == 0)
+		if (PinIndex == 0 && IsProcessorKernel())
 		{
-			if (KernelType == EPCGKernelType::PointProcessor)
-			{
-				Properties.AllowedTypes = EPCGDataType::Point;
-			}
-			else if (KernelType == EPCGKernelType::TextureProcessor)
-			{
-				Properties.AllowedTypes = EPCGDataType::BaseTexture;
+			Properties.AllowedTypes = GetElementType();
 
-				// Texture kernels are single-data for now.
+			if (IsTextureKernel())
+			{
+				// Texture kernels are single-data for now due to pressure from resource binding limits and sampler counts.
 				Properties.bAllowMultipleData = false;
 				Properties.bAllowEditMultipleData = false;
 			}
-			else if (KernelType == EPCGKernelType::AttributeSetProcessor)
-			{
-				Properties.AllowedTypes = EPCGDataType::Param;
-			}
 		}
 
 		// Enforce single data on pin types that don't support a dynamic number of data.
@@ -1081,20 +1121,14 @@ void UPCGCustomHLSLSettings::UpdatePinSettings()
 		bool bInitModeDrivenByKernel = false;
 		if (PinIndex == 0)
 		{
-			if (KernelType == EPCGKernelType::PointProcessor || KernelType == EPCGKernelType::PointGenerator)
-			{
-				bInitModeDrivenByKernel = true;
-				Properties.AllowedTypes = EPCGDataType::Point;
-			}
-			else if (KernelType == EPCGKernelType::TextureProcessor || KernelType == EPCGKernelType::TextureGenerator)
+			if (!IsCustomKernel())
 			{
 				bInitModeDrivenByKernel = true;
-				Properties.AllowedTypes = EPCGDataType::BaseTexture;
-			}
-			else if (KernelType == EPCGKernelType::AttributeSetProcessor)
-			{
-				bInitModeDrivenByKernel = true;
-				Properties.AllowedTypes = EPCGDataType::Param;
+
+				if (FPCGDataTypeIdentifier ElementType = GetElementType(); ElementType != EPCGDataType::None)
+				{
+					Properties.AllowedTypes = ElementType;
+				}
 			}
 
 			Properties.bShowPropertiesGPU = !IsTextureKernel();
diff --git a/Engine/Plugins/PCG/Source/PCG/Private/Compute/Elements/PCGCustomHLSL.h b/Engine/Plugins/PCG/Source/PCG/Private/Compute/Elements/PCGCustomHLSL.h
--- a/Engine/Plugins/PCG/Source/PCG/Private/Compute/Elements/PCGCustomHLSL.h
+++ b/Engine/Plugins/PCG/Source/PCG/Private/Compute/Elements/PCGCustomHLSL.h
@@ -25,6 +25,7 @@ enum class EPCGKernelType : uint8
 	TextureGenerator UMETA(Tooltip = "Kernel executes for each texel in a fixed size texture, configurable on node."),
 	Custom UMETA(Tooltip = "Execution thread counts and output buffer sizes configurable on node. All data read/write indices must be manually bounds checked."),
 	AttributeSetProcessor UMETA(Tooltip = "Kernel executes on each element of the attribute sets in the first input pin."),
+	AttributeSetGenerator UMETA(Tooltip = "Kernel executes on an attribute set with a fixed number of elements, configurable on node."),
 };
 
 /** Total number of threads that will be dispatched for this kernel. */
@@ -139,21 +140,35 @@ protected:
 	bool IsThreadCountMultiplierInUse() const { return KernelType == EPCGKernelType::Custom && DispatchThreadCount != EPCGDispatchThreadCount::Fixed; }
 
 	bool IsProcessorKernel() const { return KernelType == EPCGKernelType::PointProcessor || KernelType == EPCGKernelType::TextureProcessor || KernelType == EPCGKernelType::AttributeSetProcessor; }
-	bool IsGeneratorKernel() const { return KernelType == EPCGKernelType::PointGenerator || KernelType == EPCGKernelType::TextureGenerator; }
-	bool IsTextureKernel() const { return KernelType == EPCGKernelType::TextureProcessor || KernelType == EPCGKernelType::TextureGenerator; }
-	bool IsPointKernel() const { return KernelType == EPCGKernelType::PointProcessor || KernelType == EPCGKernelType::PointGenerator; }
+	bool IsGeneratorKernel() const { return KernelType == EPCGKernelType::PointGenerator || KernelType == EPCGKernelType::TextureGenerator || KernelType == EPCGKernelType::AttributeSetGenerator; }
+
+	/** The dimensionality of the kernels processing. Custom kernel type is 1D. */
+	EPCGElementDimension GetElementDimension() const;
+
+	/** Primary element data type of the kernel, if there is one (returns None for Custom kernel type). */
+	FPCGDataTypeIdentifier GetElementType() const;
+
+	bool IsTextureKernel() const { return GetElementType() == EPCGDataType::BaseTexture; }
+	bool IsPointKernel() const { return GetElementType() == EPCGDataType::Point; }
+	bool IsAttributeSetKernel() const { return GetElementType() == EPCGDataType::Param; }
+	bool IsCustomKernel() const { return KernelType == EPCGKernelType::Custom; }
+
+	UFUNCTION()
+	bool DisplayNumElements1D() const { return IsGeneratorKernel() && GetElementDimension() == EPCGElementDimension::One; }
+	UFUNCTION()
+	bool DisplayNumElements2D() const { return IsGeneratorKernel() && GetElementDimension() == EPCGElementDimension::Two; }
 
 public:
 	UPROPERTY(EditAnywhere, Category = "Settings")
 	EPCGKernelType KernelType = EPCGKernelType::PointProcessor;
 
-	UPROPERTY(EditAnywhere, Category = "Settings", meta = (EditCondition = "KernelType == EPCGKernelType::PointGenerator", EditConditionHides, PCG_OverridableCPUAndGPUWithReadback))
+	UPROPERTY(EditAnywhere, Category = "Settings", meta = (EditCondition = "DisplayNumElements1D()", EditConditionHides, PCG_OverridableCPUAndGPUWithReadback))
 	int NumElements = 256;
 
-	UPROPERTY(EditAnywhere, Category = "Settings", meta = (EditCondition = "KernelType == EPCGKernelType::TextureGenerator", EditConditionHides, PCG_OverridableCPUAndGPUWithReadback))
+	UPROPERTY(EditAnywhere, Category = "Settings", meta = (EditCondition = "DisplayNumElements2D()", EditConditionHides, PCG_OverridableCPUAndGPUWithReadback))
 	int NumElementsX = 64;
 
-	UPROPERTY(EditAnywhere, Category = "Settings", meta = (EditCondition = "KernelType == EPCGKernelType::TextureGenerator", EditConditionHides, PCG_OverridableCPUAndGPUWithReadback))
+	UPROPERTY(EditAnywhere, Category = "Settings", meta = (EditCondition = "DisplayNumElements2D()", EditConditionHides, PCG_OverridableCPUAndGPUWithReadback))
 	int NumElementsY = 64;
 
 	UPROPERTY(EditAnywhere, Category = "Settings", meta = (EditCondition = "KernelType == EPCGKernelType::TextureGenerator", EditConditionHides, PCG_OverridableCPUAndGPUWithReadback))
diff --git a/Engine/Plugins/PCG/Source/PCG/Private/Compute/Elements/PCGCustomHLSLKernel.cpp b/Engine/Plugins/PCG/Source/PCG/Private/Compute/Elements/PCGCustomHLSLKernel.cpp
--- a/Engine/Plugins/PCG/Source/PCG/Private/Compute/Elements/PCGCustomHLSLKernel.cpp
+++ b/Engine/Plugins/PCG/Source/PCG/Private/Compute/Elements/PCGCustomHLSLKernel.cpp
@@ -257,7 +257,7 @@ TSharedPtr<const FPCGDataCollectionDesc> UPCGCustomHLSLKernel::ComputeOutputBind
 			}
 		}
 	}
-	else if (InOutputPinLabel == FirstOutputPinLabel && CustomHLSLSettings->KernelType == EPCGKernelType::PointGenerator)
+	else if (InOutputPinLabel == FirstOutputPinLabel && CustomHLSLSettings->IsGeneratorKernel() && CustomHLSLSettings->GetElementDimension() == EPCGElementDimension::One)
 	{
 		const FPCGKernelParams* KernelParams = InBinding->GetCachedKernelParams(this);
 
@@ -268,8 +268,8 @@ TSharedPtr<const FPCGDataCollectionDesc> UPCGCustomHLSLKernel::ComputeOutputBind
 
 		OutDataDesc = FPCGDataCollectionDesc::MakeShared();
 
-		// Generators always produce a single point data with known point count.
-		OutDataDesc->GetDataDescriptionsMutable().Emplace(EPCGDataType::Point, KernelParams->GetValueInt(GET_MEMBER_NAME_CHECKED(UPCGCustomHLSLSettings, NumElements)));
+		// Generators always produce a single data with known element count.
+		OutDataDesc->GetDataDescriptionsMutable().Emplace(CustomHLSLSettings->GetElementType(), KernelParams->GetValueInt(GET_MEMBER_NAME_CHECKED(UPCGCustomHLSLSettings, NumElements)));
 	}
 	else if (InOutputPinLabel == FirstOutputPinLabel && CustomHLSLSettings->KernelType == EPCGKernelType::TextureGenerator)
 	{
@@ -492,112 +492,108 @@ FString UPCGCustomHLSLKernel::GetCookedSource(FPCGGPUCompilationContext& InOutCo
 			{ InPin->Label.ToString() });
 	};
 
-	if (CustomHLSLSettings->KernelType == EPCGKernelType::PointProcessor)
+	// Expand code below and update this ensure when functionality comes online.
+	ensure(CustomHLSLSettings->GetElementDimension() == EPCGElementDimension::One || CustomHLSLSettings->GetElementDimension() == EPCGElementDimension::Two);
+
+	if (CustomHLSLSettings->IsProcessorKernel())
 	{
 		const FPCGPinProperties* InputPin = GetFirstInputPin();
 		const FPCGPinPropertiesGPU* OutputPin = GetFirstOutputPin();
 
 		if (InputPin && OutputPin)
 		{
-			KernelSpecificPreamble += TEXT("    uint ElementIndex; // Assumption - element index identical in input and output data.\n");
+			if (CustomHLSLSettings->GetElementDimension() == EPCGElementDimension::One)
+			{
+				KernelSpecificPreamble += TEXT("    uint ElementIndex; // Assumption - element index identical in input and output data for processor.\n");
+			}
+			else if (CustomHLSLSettings->GetElementDimension() == EPCGElementDimension::Two)
+			{
+				KernelSpecificPreamble += TEXT("    uint2 ElementIndex; // Assumption - element index identical in input and output data.\n");
+			}
 
 			AddThreadInfoForPin(InputPin);
 			AddThreadInfoForPin(OutputPin);
 
-			// If input point is invalid, mark output point as invalid and abort.
-			KernelSpecificPreamble += FString::Format(TEXT(
-				"    if ({0}_IsPointRemoved({0}_DataIndex, ElementIndex))\n"
-				"    {\n"
-				"        {1}_RemovePoint({1}_DataIndex, ElementIndex);\n"
-				"        return;\n"
-				"    }\n"),
-				{ InputPin->Label.ToString(), OutputPin->Label.ToString() });
-
-			// Automatically copy value of all attributes for this element.
-			// TODO pass in IDs of attributes that are actually present.
-			KernelSpecificPreamble += FString::Format(TEXT(
-				"\n"
-				"    // Point processor always initializes outputs by copying input data elements.\n"
-				"    PCG_COPY_ALL_ATTRIBUTES_TO_OUTPUT({1}, {0}, {1}_DataIndex, ElementIndex, {0}_DataIndex, ElementIndex);\n"),
-				{ InputPin->Label.ToString(), OutputPin->Label.ToString() });
+			if (CustomHLSLSettings->IsPointKernel())
+			{
+				// If input point is invalid, mark output point as invalid and abort.
+				KernelSpecificPreamble += FString::Format(TEXT(
+					"    if ({0}_IsPointRemoved({0}_DataIndex, ElementIndex))\n"
+					"    {\n"
+					"        {1}_RemovePoint({1}_DataIndex, ElementIndex);\n"
+					"        return;\n"
+					"    }\n"),
+					{ InputPin->Label.ToString(), OutputPin->Label.ToString() });
+			}
+
+			if (CustomHLSLSettings->GetElementDimension() == EPCGElementDimension::One)
+			{
+				// Automatically copy value of all attributes for this element.
+				// TODO pass in IDs of attributes that are actually present.
+				KernelSpecificPreamble += FString::Format(TEXT(
+					"\n"
+					"    // Processor always initializes outputs by copying input data elements.\n"
+					"    PCG_COPY_ALL_ATTRIBUTES_TO_OUTPUT({1}, {0}, {1}_DataIndex, ElementIndex, {0}_DataIndex, ElementIndex);\n"),
+					{ InputPin->Label.ToString(), OutputPin->Label.ToString() });
+			}
+			else if (CustomHLSLSettings->GetElementDimension() == EPCGElementDimension::Two)
+			{
+				// Automatically copy the data across.
+				KernelSpecificPreamble += FString::Format(TEXT(
+					"\n"
+					"    // Processor always initializes outputs by copying the inputs.\n"
+					"    {1}_Store({1}_DataIndex, ElementIndex, {0}_Load({0}_DataIndex, ElementIndex));\n"),
+					{ InputPin->Label.ToString(), OutputPin->Label.ToString() });
+			}
 		}
 	}
-	else if (CustomHLSLSettings->KernelType == EPCGKernelType::PointGenerator)
+	else if (CustomHLSLSettings->IsGeneratorKernel())
 	{
-		KernelSpecificPreamble += TEXT(
-			"    const uint NumElements = NumElements_GetOverridableValue();\n"
-			"    // NumPoints is deprecated.\n"
-			"    const uint NumPoints = NumElements;\n");
-
-		if (const FPCGPinPropertiesGPU* OutputPin = GetFirstOutputPin())
+		if (CustomHLSLSettings->GetElementDimension() == EPCGElementDimension::One)
 		{
-			KernelSpecificPreamble += TEXT("uint ElementIndex; // Assumption - element index identical in input and output data.\n");
-
-			AddThreadInfoForPin(OutputPin);
-
-			KernelSpecificPreamble += FString::Format(TEXT(
-				"\n"
-				"    // Initialize all values to defaults for output pin {0}\n"
-				"    {0}_InitializePoint({0}_DataIndex, ElementIndex);\n"),
-				{ OutputPin->Label.ToString() });
+			KernelSpecificPreamble += TEXT("    const uint NumElements = NumElements_GetOverridableValue();\n");
 		}
-	}
-	else if (CustomHLSLSettings->KernelType == EPCGKernelType::TextureProcessor)
-	{
-		const FPCGPinProperties* InputPin = GetFirstInputPin();
-		const FPCGPinPropertiesGPU* OutputPin = GetFirstOutputPin();
-
-		if (InputPin && OutputPin)
+		else if (CustomHLSLSettings->GetElementDimension() == EPCGElementDimension::Two)
 		{
-			KernelSpecificPreamble += TEXT("    uint2 ElementIndex; // Assumption - texel index identical in input and output data.\n");
-
-			AddThreadInfoForPin(InputPin);
-			AddThreadInfoForPin(OutputPin);
-
-			// Automatically copy the input texture across.
-			KernelSpecificPreamble += FString::Format(TEXT(
-				"\n"
-				"    // Texture processor always initializes outputs by copying the input texture.\n"
-				"    {1}_Store({1}_DataIndex, ElementIndex, {0}_Load({0}_DataIndex, ElementIndex));\n"),
-				{ InputPin->Label.ToString(), OutputPin->Label.ToString() });
+			KernelSpecificPreamble += TEXT("    const uint2 NumElements = uint2(NumElementsX_GetOverridableValue(), NumElementsY_GetOverridableValue());\n");
 		}
-	}
-	else if (CustomHLSLSettings->KernelType == EPCGKernelType::TextureGenerator)
-	{
-		KernelSpecificPreamble += TEXT("    const uint2 NumElements = uint2(NumElementsX_GetOverridableValue(), NumElementsY_GetOverridableValue());\n");
 
-		if (const FPCGPinPropertiesGPU* OutputPin = GetFirstOutputPin())
+		if (CustomHLSLSettings->IsPointKernel())
 		{
-			KernelSpecificPreamble += TEXT("    uint2 ElementIndex; // Assumption - texel index identical in input and output data.\n");
-
-			AddThreadInfoForPin(OutputPin);
-
-			KernelSpecificPreamble += FString::Format(TEXT(
-				"\n"
-				"    // Zero-initialize for output pin {0}\n"
-				"    {0}_Store({0}_DataIndex, ElementIndex, (float4)0.0f);\n"),
-				{ OutputPin->Label.ToString() });
+			KernelSpecificPreamble += TEXT(
+				"    // NumPoints is deprecated.\n"
+				"    const uint NumPoints = NumElements;\n");
 		}
-	}
-	else if (CustomHLSLSettings->KernelType == EPCGKernelType::AttributeSetProcessor)
-	{
-		const FPCGPinProperties* InputPin = GetFirstInputPin();
-		const FPCGPinPropertiesGPU* OutputPin = GetFirstOutputPin();
 
-		if (InputPin && OutputPin)
+		if (const FPCGPinPropertiesGPU* OutputPin = GetFirstOutputPin())
 		{
-			KernelSpecificPreamble += TEXT("    uint ElementIndex; // Assumption - element index identical in input and output data for processor.\n");
+			if (CustomHLSLSettings->GetElementDimension() == EPCGElementDimension::One)
+			{
+				KernelSpecificPreamble += TEXT("    uint ElementIndex; // Assumption - element index identical in input and output data.\n");
+			}
+			else if (CustomHLSLSettings->GetElementDimension() == EPCGElementDimension::Two)
+			{
+				KernelSpecificPreamble += TEXT("    uint2 ElementIndex; // Assumption - element index identical in input and output data.\n");
+			}
 
-			AddThreadInfoForPin(InputPin);
 			AddThreadInfoForPin(OutputPin);
 
-			// Automatically copy value of all attributes for this element.
-			// TODO pass in IDs of attributes that are actually present.
-			KernelSpecificPreamble += FString::Format(TEXT(
-				"\n"
-				"    // Attribute Set processor always initializes outputs by copying input data elements.\n"
-				"    PCG_COPY_ALL_ATTRIBUTES_TO_OUTPUT({1}, {0}, {1}_DataIndex, ElementIndex, {0}_DataIndex, ElementIndex);\n"),
-				{ InputPin->Label.ToString(), OutputPin->Label.ToString() });
+			if (CustomHLSLSettings->GetElementDimension() == EPCGElementDimension::One)
+			{
+				KernelSpecificPreamble += FString::Format(TEXT(
+					"\n"
+					"    // Initialize all values to defaults for output pin {0}\n"
+					"    {0}_InitializePoint({0}_DataIndex, ElementIndex);\n"),
+					{ OutputPin->Label.ToString() });
+			}
+			else if (CustomHLSLSettings->GetElementDimension() == EPCGElementDimension::Two)
+			{
+				KernelSpecificPreamble += FString::Format(TEXT(
+					"\n"
+					"    // Zero-initialize for output pin {0}\n"
+					"    {0}_Store({0}_DataIndex, ElementIndex, (float4)0.0f);\n"),
+					{ OutputPin->Label.ToString() });
+			}
 		}
 	}
 
@@ -661,18 +657,26 @@ int UPCGCustomHLSLKernel::ComputeThreadCount(const UPCGDataBinding* InBinding) c
 		return ThreadCount;
 	}
 
-	if (CustomHLSLSettings->KernelType == EPCGKernelType::PointGenerator)
+	if (CustomHLSLSettings->IsGeneratorKernel())
 	{
-		// Point generator has fixed thread count.
-		ThreadCount = KernelParams->GetValueInt(GET_MEMBER_NAME_CHECKED(UPCGCustomHLSLSettings, NumElements));
-	}
-	else if (CustomHLSLSettings->KernelType == EPCGKernelType::TextureGenerator)
-	{
-		const int NumElementsX = KernelParams->GetValueInt(GET_MEMBER_NAME_CHECKED(UPCGCustomHLSLSettings, NumElementsX));
-		const int NumElementsY = KernelParams->GetValueInt(GET_MEMBER_NAME_CHECKED(UPCGCustomHLSLSettings, NumElementsY));
+		// Generators have fixed thread count.
+		if (CustomHLSLSettings->GetElementDimension() == EPCGElementDimension::One)
+		{
+			ThreadCount = KernelParams->GetValueInt(GET_MEMBER_NAME_CHECKED(UPCGCustomHLSLSettings, NumElements));
+		}
+		else if (CustomHLSLSettings->GetElementDimension() == EPCGElementDimension::Two)
+		{
+			const int NumElementsX = KernelParams->GetValueInt(GET_MEMBER_NAME_CHECKED(UPCGCustomHLSLSettings, NumElementsX));
+			const int NumElementsY = KernelParams->GetValueInt(GET_MEMBER_NAME_CHECKED(UPCGCustomHLSLSettings, NumElementsY));
 
-		// Texture generator has fixed thread count.
-		ThreadCount = NumElementsX * NumElementsY;
+			// Texture generator has fixed thread count.
+			ThreadCount = NumElementsX * NumElementsY;
+		}
+		else
+		{
+			// Needs implementation.
+			checkNoEntry();
+		}
 	}
 	else if (CustomHLSLSettings->IsProcessorKernel())
 	{
@@ -1376,7 +1380,7 @@ bool UPCGCustomHLSLKernel::PerformStaticValidation()
 			return false;
 		}
 
-		const bool bPinIsDefinedByKernel = bIsFirstOutputPin && (CustomHLSLSettings->IsPointKernel() || CustomHLSLSettings->IsTextureKernel());
+		const bool bPinIsDefinedByKernel = bIsFirstOutputPin && CustomHLSLSettings->KernelType != EPCGKernelType::Custom;
 
 		if (bIsFirstOutputPin && CustomHLSLSettings->IsPointKernel())
 		{
@@ -1404,13 +1408,13 @@ bool UPCGCustomHLSLKernel::PerformStaticValidation()
 				return false;
 			}
 		}
-		else if (bIsFirstOutputPin && CustomHLSLSettings->KernelType == EPCGKernelType::AttributeSetProcessor)
+		else if (bIsFirstOutputPin && CustomHLSLSettings->IsAttributeSetKernel())
 		{
 			if (!(Properties.AllowedTypes & EPCGDataType::Param))
 			{
 #if PCG_KERNEL_LOGGING_ENABLED
 				AddStaticLogEntry(FText::Format(
-					LOCTEXT("InvalidNonAttributeSetPrimaryOutput", "'Attribute Set Processor' nodes require primary output pin to be of type 'Attribute Set', but found '{0}'."),
+					LOCTEXT("InvalidNonAttributeSetPrimaryOutput", "'Attribute Set Processor' and 'Attribute Set Generator' nodes require primary output pin to be of type 'Attribute Set', but found '{0}'."),
 					FText::FromString(PCGCustomHLSLKernel::GetDataTypeString(Properties.AllowedTypes))),
 					EPCGKernelLogVerbosity::Error);
 #endif
