[PCG] Misc Custom HLSL node fixes

* Errors in declarations
* Element count multiplier was double applied to thread count leading to confusing bugs and OOB data writes. It's already applied to the output data element count and should not be re-multiplied in to the thread count.

#rnx
#Synced-CL 49460467

diff --git a/Engine/Plugins/PCG/Source/PCG/Private/Compute/Elements/PCGCustomHLSL.cpp b/Engine/Plugins/PCG/Source/PCG/Private/Compute/Elements/PCGCustomHLSL.cpp
--- a/Engine/Plugins/PCG/Source/PCG/Private/Compute/Elements/PCGCustomHLSL.cpp
+++ b/Engine/Plugins/PCG/Source/PCG/Private/Compute/Elements/PCGCustomHLSL.cpp
@@ -656,8 +656,8 @@ void UPCGCustomHLSLSettings::UpdateInputDeclarations()
 			"float4 {0}_Sample(uint DataIndex, float2 TexCoords, uint InSliceIndex);\n"
 			"// Computes sample coordinates of the WorldPos relative to the texture data's transform.\n"
 			"float4 {0}_SampleWorldPos(uint DataIndex, float3 WorldPos, uint InSliceIndex);\n"
-			"float4 {0}_Load(uint DataIndex, uint ElementIndex, uint SliceIndex);\n"
-			"float4 {0}_Load(uint DataIndex, uint ElementIndex, uint SliceIndex, uint MipIndex);\n"
+			"float4 {0}_Load(uint DataIndex, uint3 ElementIndex);\n"
+			"float4 {0}_Load(uint DataIndex, uint3 ElementIndex, uint MipIndex);\n"
 			"float2 {0}_GetTexelSize(uint DataIndex);\n"
 			"float2 {0}_GetTexelSizeWorld(uint DataIndex);\n"
 			"float4x4 {0}_GetTransform(uint DataIndex);\n"
@@ -889,7 +889,7 @@ void UPCGCustomHLSLSettings::UpdateOutputDeclarations()
 			"float2 {0}_GetTexelSizeWorld(uint DataIndex);\n"
 			"float4x4 {0}_GetTransform(uint DataIndex);\n"
 			"uint {0}_GetNumMips(uint DataIndex);\n"
-			"void {0}_Store(uint DataIndex, uint2 ElementIndex, uint SliceIndex, float4 Value);\n"),
+			"void {0}_Store(uint DataIndex, uint3 ElementIndex, float4 Value);\n"),
 			{ bMultiPin ? PCGHLSLElement::PinDeclTemplateStr : TextureArrayDataPins[0] });
 
 		OutputDeclarations += TEXT("\n");
diff --git a/Engine/Plugins/PCG/Source/PCG/Private/Compute/Elements/PCGCustomHLSLKernel.cpp b/Engine/Plugins/PCG/Source/PCG/Private/Compute/Elements/PCGCustomHLSLKernel.cpp
--- a/Engine/Plugins/PCG/Source/PCG/Private/Compute/Elements/PCGCustomHLSLKernel.cpp
+++ b/Engine/Plugins/PCG/Source/PCG/Private/Compute/Elements/PCGCustomHLSLKernel.cpp
@@ -693,7 +693,7 @@ int UPCGCustomHLSLKernel::ComputeThreadCount(const UPCGDataBinding* InBinding) c
 			if (const FPCGPinPropertiesGPU* OutputPin = CustomHLSLSettings->OutputPins.IsEmpty() ? nullptr : &CustomHLSLSettings->OutputPins[0])
 			{
 				const TSharedPtr<const FPCGDataCollectionDesc> Desc = InBinding->GetCachedKernelPinDataDesc(this, OutputPin->Label, /*bIsInput=*/false);
-				ThreadCount = Desc ? OutputPin->GetElementCountMultiplier() * Desc->ComputeTotalElementCount() : 0;
+				ThreadCount = Desc ? Desc->ComputeTotalElementCount() : 0;
 			}
 		}
 		else if (CustomHLSLSettings->DispatchThreadCount == EPCGDispatchThreadCount::FromProductOfInputPins)
