[PCG] Support outputting texture arrays from Generate Landscape Textures node.

[RN] [PCG] Add support for outputting texture arrays from the Generate Landscape Textures and Generate Grass Maps nodes.
#Synced-CL 49460467

diff --git a/Engine/Plugins/PCG/Source/PCG/Private/Elements/PCGGenerateLandscapeTextures.cpp b/Engine/Plugins/PCG/Source/PCG/Private/Elements/PCGGenerateLandscapeTextures.cpp
--- a/Engine/Plugins/PCG/Source/PCG/Private/Elements/PCGGenerateLandscapeTextures.cpp
+++ b/Engine/Plugins/PCG/Source/PCG/Private/Elements/PCGGenerateLandscapeTextures.cpp
@@ -8,6 +8,7 @@
 #include "PCGSubsystem.h"
 #include "Compute/PCGComputeCommon.h"
 #include "Data/PCGLandscapeData.h"
+#include "Data/PCGTexture2DArrayData.h"
 #include "Data/PCGTextureData.h"
 #include "Helpers/PCGHelpers.h"
 #include "Metadata/Accessors/PCGAttributeAccessorHelpers.h"
@@ -123,7 +124,14 @@ TArray<FPCGPinProperties> UPCGGenerateLandscapeTexturesSettings::OutputPinProper
 		PinProperties.Emplace(PCGGenerateLandscapeTextures::OutputHeightPinLabel, EPCGDataType::Texture, /*bAllowMultipleConnections=*/true, /*bInAllowMultipleData=*/false);
 	}
 
-	PinProperties.Emplace(PCGGenerateLandscapeTextures::OutputGrassMapsPinLabel, EPCGDataType::Texture, /*bAllowMultipleConnections=*/true, /*bInAllowMultipleData=*/true);
+	if (bOutputTextureArray)
+	{
+		PinProperties.Emplace(PCGGenerateLandscapeTextures::OutputGrassMapsPinLabel, FPCGDataTypeInfoTexture2DArray::AsId(), /*bAllowMultipleConnections=*/true, /*bInAllowMultipleData=*/false);
+	}
+	else
+	{
+		PinProperties.Emplace(PCGGenerateLandscapeTextures::OutputGrassMapsPinLabel, EPCGDataType::Texture, /*bAllowMultipleConnections=*/true, /*bInAllowMultipleData=*/true);
+	}
 
 	return PinProperties;
 }
@@ -168,6 +176,11 @@ void FPCGGenerateLandscapeTexturesContext::AddExtraStructReferencedObjects(FRefe
 		}
 	}
 
+	if (GrassMapTextureArrayData)
+	{
+		Collector.AddReferencedObject(GrassMapTextureArrayData);
+	}
+
 	for (TObjectPtr<UTexture>& Texture : TexturesToStream)
 	{
 		if (Texture)
@@ -635,7 +648,7 @@ bool FPCGGenerateLandscapeTexturesElement::ExecuteInternal(FPCGContext* InContex
 		
 		TWeakPtr<FPCGContextHandle> ContextHandle = Context->GetOrCreateHandle();
 
-		ENQUEUE_RENDER_COMMAND(GenerateGrassMaps)([ContextHandle, LandscapeComponentExtent=Context->LandscapeComponentExtent, LandscapeTileCoords=MoveTemp(LandscapeTileCoords), LandscapeComponentResolution, LandscapeGridScale, LandscapeZ, bGenerateHeightMap=Settings->bGenerateHeightMap, NumTilesX, NumTilesY, OutputResolution, WindowStart, WindowEnd, IndexMapping=MoveTemp(Context->IndexMapping)](FRHICommandListImmediate& RHICmdList)
+		ENQUEUE_RENDER_COMMAND(PCGGenerateGrassMaps)([ContextHandle, LandscapeComponentExtent=Context->LandscapeComponentExtent, LandscapeTileCoords=MoveTemp(LandscapeTileCoords), LandscapeComponentResolution, LandscapeGridScale, LandscapeZ, bGenerateHeightMap=Settings->bGenerateHeightMap, NumTilesX, NumTilesY, OutputResolution, WindowStart, WindowEnd, IndexMapping=MoveTemp(Context->IndexMapping)](FRHICommandListImmediate& RHICmdList)
 		{
 			LLM_SCOPE_BYTAG(PCG);
 
@@ -763,14 +776,22 @@ bool FPCGGenerateLandscapeTexturesElement::ExecuteInternal(FPCGContext* InContex
 		Context->HeightTextureData = FPCGContext::NewObject_AnyThread<UPCGTextureData>(Context);
 	}
 
-	if (Context->GrassMapTextureDatas.IsEmpty())
+	if (!Context->bCreatedOutputGrassTextures)
 	{
-		// Create the texture data objects if they haven't been created already. There should be one per selected grass type.
-		for (int DataIndex = 0; DataIndex < Context->SelectedGrassTypes.Num(); ++DataIndex)
+		if (Settings->bOutputTextureArray)
 		{
-			UPCGTextureData* TextureData = FPCGContext::NewObject_AnyThread<UPCGTextureData>(Context);
-			Context->GrassMapTextureDatas.Add(TextureData);
+			Context->GrassMapTextureArrayData = FPCGContext::NewObject_AnyThread<UPCGTexture2DArrayData>(Context);
 		}
+		else
+		{
+			// Create the texture data objects if they haven't been created already. There should be one per selected grass type.
+			for (int DataIndex = 0; DataIndex < Context->SelectedGrassTypes.Num(); ++DataIndex)
+			{
+				Context->GrassMapTextureDatas.Add(FPCGContext::NewObject_AnyThread<UPCGTextureData>(Context));
+			}
+		}
+
+		Context->bCreatedOutputGrassTextures = true;
 	}
 
 	if (!Context->GrassMapHandle.IsValid() && !Context->HeightHandle.IsValid())
@@ -788,18 +809,28 @@ bool FPCGGenerateLandscapeTexturesElement::ExecuteInternal(FPCGContext* InContex
 
 	if (Context->GrassMapHandle)
 	{
-		for (int DataIndex = 0; DataIndex < Context->SelectedGrassTypes.Num(); ++DataIndex)
+		if (Settings->bOutputTextureArray)
 		{
-			check(Context->GrassMapTextureDatas.IsValidIndex(DataIndex));
-			UPCGTextureData* TextureData = Context->GrassMapTextureDatas[DataIndex];
-			check(TextureData);
+			FPCGTexture2DArrayDataInitParams InitParams;
+			InitParams.Transform = Context->OutputTextureTransform;
 
-			// Poll initialize (fine to be called even when initialization was already complete).
-			bAllTexturesInitialized &= TextureData->Initialize(
-				Context->GrassMapHandle,
-				/*TextureIndex=*/Context->SelectedGrassTypes[DataIndex].Get<1>(),
-				Context->OutputTextureTransform,
-				Settings->bSkipReadbackToCPU);
+			bAllTexturesInitialized &= Context->GrassMapTextureArrayData->Initialize(Context->GrassMapHandle, InitParams);
+		}
+		else
+		{
+			for (int DataIndex = 0; DataIndex < Context->SelectedGrassTypes.Num(); ++DataIndex)
+			{
+				check(Context->GrassMapTextureDatas.IsValidIndex(DataIndex));
+				UPCGTextureData* TextureData = Context->GrassMapTextureDatas[DataIndex];
+				check(TextureData);
+
+				// Poll initialize (fine to be called even when initialization was already complete).
+				bAllTexturesInitialized &= TextureData->Initialize(
+					Context->GrassMapHandle,
+					/*TextureIndex=*/Context->SelectedGrassTypes[DataIndex].Get<1>(),
+					Context->OutputTextureTransform,
+					Settings->bSkipReadbackToCPU);
+			}
 		}
 	}
 
@@ -821,7 +852,7 @@ bool FPCGGenerateLandscapeTexturesElement::ExecuteInternal(FPCGContext* InContex
 	}
 
 #if WITH_EDITOR
-	if ((Context->HeightTextureData || !Context->GrassMapTextureDatas.IsEmpty()) && !Settings->bSkipReadbackToCPU)
+	if ((Context->HeightTextureData || Context->GrassMapTextureArrayData || !Context->GrassMapTextureDatas.IsEmpty()) && !Settings->bSkipReadbackToCPU)
 	{
 		if (!Settings->bSkipReadbackToCPU && Context->Node && Context->GetStack() && Context->ExecutionSource.IsValid())
 		{
@@ -845,37 +876,46 @@ bool FPCGGenerateLandscapeTexturesElement::ExecuteInternal(FPCGContext* InContex
 		}
 	}
 
-	for (int DataIndex = 0; DataIndex < Context->SelectedGrassTypes.Num(); ++DataIndex)
+	if (Settings->bOutputTextureArray)
 	{
-		check(Context->GrassMapTextureDatas.IsValidIndex(DataIndex));
-		UPCGTextureData* TextureData = Context->GrassMapTextureDatas[DataIndex];
-		check(TextureData);
-
-		if (!TextureData->IsSuccessfullyInitialized())
-		{
-			PCGE_LOG(Warning, LogOnly, LOCTEXT("TextureInitFailed", "Data could not be retrieved for this texture, initialization failed."));
-			continue;
-		}
-
 		FPCGTaggedData& OutTaggedData = Context->OutputData.TaggedData.Emplace_GetRef();
 		OutTaggedData.Pin = PCGGenerateLandscapeTextures::OutputGrassMapsPinLabel;
-		OutTaggedData.Data = TextureData;
-		
-		FString GrassTypeName = TEXT("");
-
-		if (ULandscapeGrassType* GrassType = Context->SelectedGrassTypes[DataIndex].Get<0>().Get())
+		OutTaggedData.Data = Context->GrassMapTextureArrayData;
+	}
+	else
+	{
+		for (int DataIndex = 0; DataIndex < Context->SelectedGrassTypes.Num(); ++DataIndex)
 		{
-			GrassType->GetName(GrassTypeName);
-		}
+			check(Context->GrassMapTextureDatas.IsValidIndex(DataIndex));
+			UPCGTextureData* TextureData = Context->GrassMapTextureDatas[DataIndex];
+			check(TextureData);
 
-		if (!GrassTypeName.IsEmpty())
-		{
-			OutTaggedData.Tags.Add(PCGComputeHelpers::GetPrefixedDataLabel(GrassTypeName));
-		}
-		else
-		{
-			PCGE_LOG(Warning, LogOnly, LOCTEXT("MissingGrassName",
-				"Grass type name was missing, data could not be labeled. Make sure all GrassTypes in your landscape material have an asset associated."));
+			if (!TextureData->IsSuccessfullyInitialized())
+			{
+				PCGE_LOG(Warning, LogOnly, LOCTEXT("TextureInitFailed", "Data could not be retrieved for this texture, initialization failed."));
+				continue;
+			}
+
+			FPCGTaggedData& OutTaggedData = Context->OutputData.TaggedData.Emplace_GetRef();
+			OutTaggedData.Pin = PCGGenerateLandscapeTextures::OutputGrassMapsPinLabel;
+			OutTaggedData.Data = TextureData;
+
+			FString GrassTypeName = TEXT("");
+
+			if (ULandscapeGrassType* GrassType = Context->SelectedGrassTypes[DataIndex].Get<0>().Get())
+			{
+				GrassType->GetName(GrassTypeName);
+			}
+
+			if (!GrassTypeName.IsEmpty())
+			{
+				OutTaggedData.Tags.Add(PCGComputeHelpers::GetPrefixedDataLabel(GrassTypeName));
+			}
+			else
+			{
+				PCGE_LOG(Warning, LogOnly, LOCTEXT("MissingGrassName",
+					"Grass type name was missing, data could not be labeled. Make sure all GrassTypes in your landscape material have an asset associated."));
+			}
 		}
 	}
 
diff --git a/Engine/Plugins/PCG/Source/PCG/Public/Elements/PCGGenerateLandscapeTextures.h b/Engine/Plugins/PCG/Source/PCG/Public/Elements/PCGGenerateLandscapeTextures.h
--- a/Engine/Plugins/PCG/Source/PCG/Public/Elements/PCGGenerateLandscapeTextures.h
+++ b/Engine/Plugins/PCG/Source/PCG/Public/Elements/PCGGenerateLandscapeTextures.h
@@ -14,6 +14,7 @@ class FLandscapeGrassWeightExporter;
 class ULandscapeComponent;
 class ULandscapeGrassType;
 class UPCGTextureData;
+class UPCGTexture2DArrayData;
 class UTexture;
 struct IPooledRenderTarget;
 
@@ -72,6 +73,10 @@ public:
 	/** Generate a landscape height texture. */
 	UPROPERTY(BlueprintReadWrite, EditAnywhere, Category = Settings)
 	bool bGenerateHeightMap = false;
+
+	/** Output a single texture 2D array data rather than a texture 2D data for each slice. */
+	UPROPERTY(BlueprintReadWrite, EditAnywhere, Category = Settings)
+	bool bOutputTextureArray = false;
 };
 
 struct FPCGGenerateLandscapeTexturesContext : public FPCGContext
@@ -81,9 +86,13 @@ public:
 
 	FLandscapeGrassWeightExporter* LandscapeGrassWeightExporter = nullptr;
 
-	/** Output texture data objects. */
+	/** Output height texture data. */
 	TObjectPtr<UPCGTextureData> HeightTextureData;
+
+	/** Output grass map texture data, either a texture array data or multiple texture datas, depending on settings. */
 	TArray<TObjectPtr<UPCGTextureData>> GrassMapTextureDatas;
+	TObjectPtr<UPCGTexture2DArrayData> GrassMapTextureArrayData;
+	bool bCreatedOutputGrassTextures = false;
 
 	/** Exported result textures. */
 	TRefCountPtr<IPooledRenderTarget> HeightHandle = nullptr;
