[PCG] Revert a change (under a CVar) that is recompiling a graph when it has structural change to avoid unnecessary notify for graph changes

If notifications are paused (or no one is hooked to them), this recompile is wastefull.
This becomes more meaningfull when the changes are done in sequence by a script, it slows down the operations.

#Synced-CL 49559950

diff --git a/Engine/Plugins/PCG/Source/PCG/Private/PCGGraph.cpp b/Engine/Plugins/PCG/Source/PCG/Private/PCGGraph.cpp
--- a/Engine/Plugins/PCG/Source/PCG/Private/PCGGraph.cpp
+++ b/Engine/Plugins/PCG/Source/PCG/Private/PCGGraph.cpp
@@ -57,6 +57,11 @@ namespace PCGGraph
 		TEXT("pcg.Graph.EnableWarningOnPinDestructiveChangesInPostLoad"),
 		false,
 		TEXT("Output warnings when destructive changes happens on PostLoad in UpdatePins."));
+
+	static TAutoConsoleVariable<bool> CVarRecompileGraphOnChange(
+		TEXT("pcg.Graph.CompileTasksToDetectChanges"),
+		true,
+		TEXT("Recompile graph after edits to detect if compiled tasks changed, to avoid unnecessary graph execution"));
 }
 
 namespace PCGGraphUtils
@@ -992,6 +997,7 @@ UPCGNode* UPCGGraph::AddNodeOfType(TSubclassOf<class UPCGSettings> InSettingsCla
 
 UPCGNode* UPCGGraph::AddNode(UPCGSettingsInterface* InSettingsInterface)
 {
+	TRACE_CPUPROFILER_EVENT_SCOPE(UPCGGraph::AddNode);
 	if (!InSettingsInterface || !InSettingsInterface->GetSettings())
 	{
 		return nullptr;
@@ -1865,23 +1871,26 @@ void UPCGGraph::NotifyGraphStructureChanged(EPCGChangeType ChangeType, bool bFor
 {
 	bool bExecutionAffected = true;
 
-	// If settings were not changed, we can gate the change notification based on whether compiled graph output changed. This compilation check
-	// does not support settings changes.
-	if (!bForce && !(ChangeType & EPCGChangeType::Settings))
+	if (PCGGraph::CVarRecompileGraphOnChange.GetValueOnAnyThread())
 	{
-		bExecutionAffected = Recompile();
-	}
+		// If settings were not changed, we can gate the change notification based on whether compiled graph output changed. This compilation check
+		// does not support settings changes.
+		if (!bForce && !(ChangeType & EPCGChangeType::Settings))
+		{
+			bExecutionAffected = Recompile();
+		}
 
-	if (!bExecutionAffected)
-	{
-		// If compiled tasks are unchanged and the settings have not changed, then we can demote the change to avoid unnecessary graph executions.
-		// * Structural and GenerationGrid are frequent change types that trigger generation
-		// * Node because edges can change pin types on nodes.
-		// * Settings because changing a subgraph currently changes settings on subgraph node.
-		ChangeType &= ~(EPCGChangeType::Structural | EPCGChangeType::GenerationGrid | EPCGChangeType::Node | EPCGChangeType::Settings);
+		if (!bExecutionAffected)
+		{
+			// If compiled tasks are unchanged and the settings have not changed, then we can demote the change to avoid unnecessary graph executions.
+			// * Structural and GenerationGrid are frequent change types that trigger generation
+			// * Node because edges can change pin types on nodes.
+			// * Settings because changing a subgraph currently changes settings on subgraph node.
+			ChangeType &= ~(EPCGChangeType::Structural | EPCGChangeType::GenerationGrid | EPCGChangeType::Node | EPCGChangeType::Settings);
 
-		// Positively flagging as cosmetic is required because downstream things specifically test for this currently.
-		ChangeType |= EPCGChangeType::Cosmetic;
+			// Positively flagging as cosmetic is required because downstream things specifically test for this currently.
+			ChangeType |= EPCGChangeType::Cosmetic;
+		}
 	}
 
 	NotifyGraphChanged(ChangeType);
