[PCG] Fix InspectionData thread safety
- Replace GetInspectionData() api which was returning an unsafe raw pointer to a FPCGDataCollection by 2 options:
- InspectData: Thread safe in place inspection (using function ref under lock)
- GetInspectionDataPtr: FGCObject which guarantees UPCGData held by shared data collection is kept alive (without depending on inspection cache)
- Inspection data is stored as TSharedPtr<FPCGDataCollection> to handle collections outliving the inspection cache

#Synced-CL 49570007

diff --git a/Engine/Plugins/Experimental/ProceduralVegetationEditor/Source/ProceduralVegetationEditor/Private/PVEditor.cpp b/Engine/Plugins/Experimental/ProceduralVegetationEditor/Source/ProceduralVegetationEditor/Private/PVEditor.cpp
--- a/Engine/Plugins/Experimental/ProceduralVegetationEditor/Source/ProceduralVegetationEditor/Private/PVEditor.cpp
+++ b/Engine/Plugins/Experimental/ProceduralVegetationEditor/Source/ProceduralVegetationEditor/Private/PVEditor.cpp
@@ -390,8 +390,7 @@ void FPVEditor::UpdateInspectedCollection()
 	if (const UPCGNode* InspectedNode = NodeBeingInspected ? NodeBeingInspected->GetPCGNode() : nullptr)
 	{
 		CollectionBeingInspected = MakeShared<FManagedArrayCollection>();
-		FName OutputName;
-		GetCollectionFromNode(InspectedNode, OutputName, *CollectionBeingInspected);
+		GetCollectionFromNode(InspectedNode, *CollectionBeingInspected);
 	}
 }
 
@@ -574,8 +573,7 @@ FReply FPVEditor::ExportInternal()
 		for (UPCGNode* Node : NodesToExport)
 		{
 			FManagedArrayCollection& OutCollection = OutCollections.AddDefaulted_GetRef();
-			FName PinName = NAME_None;
-			if (!GetCollectionFromNode(Node, PinName, OutCollection))
+			if (!GetCollectionFromNode(Node, OutCollection))
 			{
 				TSharedRef<FTokenizedMessage> Error = OutMessages.Error(
 					FText::Format(
@@ -596,12 +594,13 @@ FReply FPVEditor::ExportInternal()
 		TArray<const UPCGNode*> NodesWithoutFoliage;
 		for (UPCGNode* Node : NodesToExport)
 		{
-			FName PinName = NAME_None;
-			const UPVData* PVData = GetPVDataFromNode(Node, PinName);
-			if (PVData && !Cast<UPVFoliageMeshData>(PVData))
+			GetPVDataFromNode(Node, [Node, &NodesWithoutFoliage](const UPVData& InPVData)
 			{
-				NodesWithoutFoliage.Add(Node);
-			}
+				if(!InPVData.IsA<UPVFoliageMeshData>())
+				{
+					NodesWithoutFoliage.Add(Node);
+				}
+			});
 		}
 
 		if (NodesWithoutFoliage.Num() > 0)
@@ -847,50 +846,67 @@ TSharedRef<SPCGEditorViewport> FPVEditor::CreateViewportWidget()
 	
 }
 
-const UPVData* FPVEditor::GetPVDataFromNode(const UPCGNode* InNode, FName& OutPinName) const
+FPCGStack FPVEditor::GetPVDataStack(const UPCGNode* InNode, const UPCGPin*& OutPin) const
 {
-	const UPCGPin* Pin = nullptr;
-		
+	OutPin = nullptr;
 	const TArray<TObjectPtr<UPCGPin>>& Pins = InNode->GetOutputPins();
 
 	if (Pins.IsValidIndex(0))
 	{
-		Pin = Pins[0];
+		OutPin = Pins[0];
 	}
-	
+
 	const FPCGStack* PCGStack = GetStackBeingInspected();
 	check(PCGStack)
-		
+
 	// Create a temporary stack with Node + valid Pin to query the exact DataCollection
 	FPCGStack Stack = *PCGStack;
 	TArray<FPCGStackFrame>& StackFrames = Stack.GetStackFramesMutable();
 	StackFrames.Reserve(StackFrames.Num() + 2);
 	StackFrames.Emplace(InNode);
-	if (Pin)
+	if (OutPin)
 	{
-		StackFrames.Emplace(Pin);
+		StackFrames.Emplace(OutPin);
 	}
-		
-	const FPCGDataCollection* Data = ExecutionSource->GetExecutionState().GetInspection().GetInspectionData(Stack);
-		
-	if(Data != nullptr)
+
+	return Stack;
+}
+
+const UPVData* FPVEditor::GetPVDataFromCollection(const FPCGDataCollection& InDataCollection, const UPCGPin* InPin) const
+{
+	const FPCGTaggedData* TaggedData = InDataCollection.TaggedData.GetData();
+
+	if (TaggedData != nullptr)
 	{
-		const FPCGTaggedData* TaggedData = Data->TaggedData.GetData();
-		
-		if(TaggedData != nullptr)
+		const UPVData* ProceduralVegetationData = Cast<UPVData>(TaggedData->Data);
+		if (ProceduralVegetationData != nullptr && InPin)
 		{
-			const UPVData* ProceduralVegetationData = Cast<UPVData>(TaggedData->Data);
-			if(ProceduralVegetationData != nullptr && Pin)
-			{
-				OutPinName = Pin->Properties.Label;
-				return ProceduralVegetationData;
-			}
+			return ProceduralVegetationData;
 		}
 	}
-
+	
 	return nullptr;
 }
 
+bool FPVEditor::GetPVDataFromNode(const UPCGNode* InNode, TFunctionRef<void(const UPVData&)> PVDataFunc) const
+{
+	const UPCGPin* OutPin = nullptr;
+	FPCGStack Stack = GetPVDataStack(InNode, OutPin);
+
+	bool bFoundPVData = false;
+
+	ExecutionSource->GetExecutionState().GetInspection().InspectData(Stack, [this, Pin = OutPin, &bFoundPVData, &PVDataFunc](const FPCGDataCollection& InDataCollection)
+	{
+		if(const UPVData* PVData = GetPVDataFromCollection(InDataCollection, Pin))
+		{
+			bFoundPVData = true;
+			PVDataFunc(*PVData);
+		}
+	});
+		
+	return bFoundPVData;
+}
+
 bool FPVEditor::CanToggleInspected() const
 {
 	return false;
@@ -906,16 +922,12 @@ UPCGEditorGraphNodeBase* FPVEditor::GetSelectedNode()
 	return nullptr;
 }
 
-bool FPVEditor::GetCollectionFromNode(const UPCGNode* InNode, FName& OutPinName, FManagedArrayCollection& OutCollection) const
+bool FPVEditor::GetCollectionFromNode(const UPCGNode* InNode, FManagedArrayCollection& OutCollection) const
 {
-	const UPVData* PVData = GetPVDataFromNode(InNode, OutPinName);
-	if(PVData != nullptr)
+	return GetPVDataFromNode(InNode, [&OutCollection](const UPVData& InPVData)
 	{
-		OutCollection = PVData->GetCollection();
-		return true;
-	}
-
-	return false;
+		OutCollection = InPVData.GetCollection();
+	});
 }
 
 TSharedRef<SCollectionSpreadSheetWidget> FPVEditor::CreateCollectionSpreadSheetWidget()
diff --git a/Engine/Plugins/Experimental/ProceduralVegetationEditor/Source/ProceduralVegetationEditor/Private/PVEditor.h b/Engine/Plugins/Experimental/ProceduralVegetationEditor/Source/ProceduralVegetationEditor/Private/PVEditor.h
--- a/Engine/Plugins/Experimental/ProceduralVegetationEditor/Source/ProceduralVegetationEditor/Private/PVEditor.h
+++ b/Engine/Plugins/Experimental/ProceduralVegetationEditor/Source/ProceduralVegetationEditor/Private/PVEditor.h
@@ -90,9 +90,11 @@ private:
 protected:
 	virtual TSharedRef<SPCGEditorViewport> CreateViewportWidget() override;
 	
-	const UPVData* GetPVDataFromNode(const UPCGNode* InNode, FName& OutPinName) const;
+	FPCGStack GetPVDataStack(const UPCGNode* InNode, UPCGPin const*& OutPin) const;
+	const UPVData* GetPVDataFromCollection(const FPCGDataCollection& InDataCollection, const UPCGPin* InPin) const;
+	bool GetPVDataFromNode(const UPCGNode* InNode, TFunctionRef<void(const UPVData&)> PVDataFunc) const;
 	virtual bool CanToggleInspected() const override;
-	bool GetCollectionFromNode(const UPCGNode* InNode, FName& OutPinName, FManagedArrayCollection& OutCollection) const;
+	bool GetCollectionFromNode(const UPCGNode* InNode, FManagedArrayCollection& OutCollection) const;
 
 	UPCGEditorGraphNodeBase* GetSelectedNode();
 };
diff --git a/Engine/Plugins/PCG/Source/PCG/Private/PCGComponent.cpp b/Engine/Plugins/PCG/Source/PCG/Private/PCGComponent.cpp
--- a/Engine/Plugins/PCG/Source/PCG/Private/PCGComponent.cpp
+++ b/Engine/Plugins/PCG/Source/PCG/Private/PCGComponent.cpp
@@ -2634,7 +2634,9 @@ void UPCGComponent::StoreInspectionData(const FPCGStack* InStack, const UPCGNode
 // UE_DEPRECATED(5.6)
 const FPCGDataCollection* UPCGComponent::GetInspectionData(const FPCGStack& InStack) const
 {
+PRAGMA_DISABLE_DEPRECATION_WARNINGS
 	return ExecutionInspection.GetInspectionData(InStack);
+PRAGMA_ENABLE_DEPRECATION_WARNINGS
 }
 
 // UE_DEPRECATED(5.6)
diff --git a/Engine/Plugins/PCG/Source/PCG/Private/PCGGraphExecutionInspection.cpp b/Engine/Plugins/PCG/Source/PCG/Private/PCGGraphExecutionInspection.cpp
--- a/Engine/Plugins/PCG/Source/PCG/Private/PCGGraphExecutionInspection.cpp
+++ b/Engine/Plugins/PCG/Source/PCG/Private/PCGGraphExecutionInspection.cpp
@@ -9,6 +9,19 @@
 #include "PCGSettings.h"
 #include "PCGSubgraph.h"
 
+void FPCGInspectionData::AddReferencedObjects(FReferenceCollector& Collector)
+{
+	if (Data)
+	{
+		Data->AddReferences(Collector);
+	}
+}
+
+FString FPCGInspectionData::GetReferencerName() const
+{
+	return TEXT("FPCGInspectionData");
+}
+
 bool FPCGGraphExecutionInspection::IsInspecting() const
 {
 	return InspectionCounter > 0;
@@ -129,7 +142,7 @@ void FPCGGraphExecutionInspection::StoreInspectionData(const FPCGStack* InStack,
 		InInputData.MarkUsage(EPCGDataUsage::ComponentInspectionData);
 		InOutputData.MarkUsage(EPCGDataUsage::ComponentInspectionData);
 
-		auto StorePinInspectionDataFromNode = [](const FPCGStack* InStack, const TArray<TObjectPtr<UPCGPin>>& InPins, const FPCGDataCollection& InData, TMap<FPCGStack, FPCGDataCollection>& InOutInspectionCache)
+		auto StorePinInspectionDataFromNode = [](const FPCGStack* InStack, const TArray<TObjectPtr<UPCGPin>>& InPins, const FPCGDataCollection& InData, TMap<FPCGStack, TSharedPtr<FPCGDataCollection>>& InOutInspectionCache)
 		{
 			for (const UPCGPin* Pin : InPins)
 			{
@@ -145,19 +158,19 @@ void FPCGGraphExecutionInspection::StoreInspectionData(const FPCGStack* InStack,
 				// Implementation note: since static subgraphs actually are visited twice and the second time the input doesn't match the input pins, we don't clear the data.
 				if (!PinDataCollection.TaggedData.IsEmpty())
 				{
-					if (FPCGDataCollection* CollectionInCache = InOutInspectionCache.Find(Stack))
+					if (TSharedPtr<FPCGDataCollection>* CollectionInCache = InOutInspectionCache.Find(Stack))
 					{
-						(*CollectionInCache) += PinDataCollection;
+						*(*CollectionInCache) += PinDataCollection;
 					}
 					else
 					{
-						InOutInspectionCache.Add(Stack, PinDataCollection);
+						InOutInspectionCache.Add(Stack, MakeShared<FPCGDataCollection>(MoveTemp(PinDataCollection)));
 					}
 				}
 			}
 		};
 
-		auto StorePinInspectionData = [InStack, InNode, &StorePinInspectionDataFromNode](const TArray<TObjectPtr<UPCGPin>>& InPins, const FPCGDataCollection& InData, TMap<FPCGStack, FPCGDataCollection>& InOutInspectionCache)
+		auto StorePinInspectionData = [InStack, InNode, &StorePinInspectionDataFromNode](const TArray<TObjectPtr<UPCGPin>>& InPins, const FPCGDataCollection& InData, TMap<FPCGStack, TSharedPtr<FPCGDataCollection>>& InOutInspectionCache)
 		{
 			FPCGStack Stack = *InStack;
 
@@ -193,10 +206,34 @@ void FPCGGraphExecutionInspection::StoreInspectionData(const FPCGStack* InStack,
 	}
 }
 
-const FPCGDataCollection* FPCGGraphExecutionInspection::GetInspectionData(const FPCGStack& InStack) const
+bool FPCGGraphExecutionInspection::InspectData(const FPCGStack& InStack, TFunctionRef<void(const FPCGDataCollection&)> InspectFunc) const
 {
 	PCG::TSharedScopeLock ReadLock(InspectionCacheLock);
-	return InspectionCache.Find(InStack);
+	if (const TSharedPtr<FPCGDataCollection>* DataPtr = InspectionCache.Find(InStack))
+	{
+		if(const FPCGDataCollection* Data = DataPtr->Get())
+		{
+			InspectFunc(*Data);
+			return true;
+		}
+	}
+
+	return false;
+}
+
+TSharedPtr<FPCGInspectionData> FPCGGraphExecutionInspection::GetInspectionDataPtr(const FPCGStack& InStack) const
+{
+	TSharedPtr<FPCGInspectionData> InspectionDataPtr;
+
+	PCG::TSharedScopeLock ReadLock(InspectionCacheLock);
+	if (const TSharedPtr<FPCGDataCollection>* DataPtr = InspectionCache.Find(InStack); DataPtr && DataPtr->IsValid())
+	{
+		// Only allocate FPCGInspectionData if we found some data to return to avoid creating FGCObjects for nothing
+		InspectionDataPtr = MakeShared<FPCGInspectionData>();
+		InspectionDataPtr->Data = *DataPtr;
+	}
+
+	return InspectionDataPtr;
 }
 
 void FPCGGraphExecutionInspection::ClearInspectionData(bool bClearPerNodeExecutionData)
@@ -204,9 +241,9 @@ void FPCGGraphExecutionInspection::ClearInspectionData(bool bClearPerNodeExecuti
 	{
 		PCG::TUniqueScopeLock WriteLock(InspectionCacheLock);
 
-		for (TPair<FPCGStack, FPCGDataCollection>& Entry : InspectionCache)
+		for (TPair<FPCGStack, TSharedPtr<FPCGDataCollection>>& Entry : InspectionCache)
 		{
-			Entry.Value.ClearUsage(EPCGDataUsage::ComponentInspectionData);
+			Entry.Value->ClearUsage(EPCGDataUsage::ComponentInspectionData);
 		}
 
 		InspectionCache.Reset();
@@ -287,7 +324,7 @@ void FPCGGraphExecutionInspection::AddReferencedObjects(FReferenceCollector& Col
 	PCG::TSharedScopeLock ReadLock(InspectionCacheLock);
 	for (auto& Pair : InspectionCache)
 	{
-		Pair.Value.AddReferences(Collector);
+		Pair.Value->AddReferences(Collector);
 	}
 }
 
diff --git a/Engine/Plugins/PCG/Source/PCG/Public/PCGGraphExecutionInspection.h b/Engine/Plugins/PCG/Source/PCG/Public/PCGGraphExecutionInspection.h
--- a/Engine/Plugins/PCG/Source/PCG/Public/PCGGraphExecutionInspection.h
+++ b/Engine/Plugins/PCG/Source/PCG/Public/PCGGraphExecutionInspection.h
@@ -2,17 +2,28 @@
 
 #pragma once
 
+#if WITH_EDITOR
+
 #include "PCGCommon.h"
 #include "PCGData.h"
 #include "PCGNode.h"
 #include "Graph/PCGStackContext.h"
 #include "Utils/PCGExtraCapture.h"
 
+#include "UObject/GCObject.h"
 #include "UObject/ObjectKey.h"
 
-#define UE_API PCG_API
+/* Helper struct to keep references on data collection UPCGData visible to GC even if inspection data source gets cleared */
+struct FPCGInspectionData : public FGCObject
+{
+	//~ GCObject Interface
+	PCG_API virtual void AddReferencedObjects(FReferenceCollector& Collector) override;
+	PCG_API virtual FString GetReferencerName() const override;
+	//~ GCObject Interface
+
+	TSharedPtr<FPCGDataCollection> Data;
+};
 
-#if WITH_EDITOR
 class FPCGGraphExecutionInspection
 {
 public:
@@ -30,42 +41,50 @@ public:
 		PCGUtils::FCallTime Timer;
 	};
 
-	UE_API bool IsInspecting() const;
-	UE_API void EnableInspection();
-	UE_API void DisableInspection();
-	UE_API void StoreInspectionData(const FPCGStack* InStack, const UPCGNode* InNode, const PCGUtils::FCallTime* InTimer, const FPCGDataCollection& InInputData, const FPCGDataCollection& InOutputData, bool bUsedCache);
-	UE_API const FPCGDataCollection* GetInspectionData(const FPCGStack& InStack) const;
-	UE_API void ClearInspectionData(bool bClearPerNodeExecutionData = true);
+	PCG_API bool IsInspecting() const;
+	PCG_API void EnableInspection();
+	PCG_API void DisableInspection();
+	PCG_API void StoreInspectionData(const FPCGStack* InStack, const UPCGNode* InNode, const PCGUtils::FCallTime* InTimer, const FPCGDataCollection& InInputData, const FPCGDataCollection& InOutputData, bool bUsedCache);
+
+	/** Data is only valid inside the function call. If you want the data to outlive the call, use GetInspectionDataPtr instead. */
+	PCG_API bool InspectData(const FPCGStack& InStack, TFunctionRef<void(const FPCGDataCollection&)> InspectFunc) const;
+	
+	/** If found, the inspection data is returned as a garbage collection visible object (FGCObject). Prefer using InspectData if you can manage UObject references yourself. */
+	PCG_API TSharedPtr<FPCGInspectionData> GetInspectionDataPtr(const FPCGStack& InStack) const;
+	
+	PCG_API void ClearInspectionData(bool bClearPerNodeExecutionData = true);
 
 	/** Whether a task for the given node and stack was executed during the last execution. */
-	UE_API bool WasNodeExecuted(const UPCGNode* InNode, const FPCGStack& Stack) const;
+	PCG_API bool WasNodeExecuted(const UPCGNode* InNode, const FPCGStack& Stack) const;
 
 	/** Called at execution time each time a node has been executed. */
-	UE_API void NotifyNodeExecuted(const UPCGNode* InNode, const FPCGStack* InStack, const PCGUtils::FCallTime* InTimer, bool bNodeUsedCache);
+	PCG_API void NotifyNodeExecuted(const UPCGNode* InNode, const FPCGStack* InStack, const PCGUtils::FCallTime* InTimer, bool bNodeUsedCache);
 		
-	UE_API TMap<TObjectKey<const UPCGNode>, TSet<FNodeExecutedNotificationData>> GetExecutedNodeStacks() const;
+	PCG_API TMap<TObjectKey<const UPCGNode>, TSet<FNodeExecutedNotificationData>> GetExecutedNodeStacks() const;
 
 	/** Retrieve the inactive pin bitmask for the given node and stack in the last execution. */
-	UE_API uint64 GetNodeInactivePinMask(const UPCGNode* InNode, const FPCGStack& Stack) const;
+	PCG_API uint64 GetNodeInactivePinMask(const UPCGNode* InNode, const FPCGStack& Stack) const;
 
 	/** Whether the given node was culled by a dynamic branch in the given stack. */
-	UE_API void NotifyNodeDynamicInactivePins(const UPCGNode* InNode, const FPCGStack* InStack, uint64 InactivePinBitmask) const;
+	PCG_API void NotifyNodeDynamicInactivePins(const UPCGNode* InNode, const FPCGStack* InStack, uint64 InactivePinBitmask) const;
 
 	/** Did the given node produce one or more data items in the given stack during the last execution. */
-	UE_API bool HasNodeProducedData(const UPCGNode* InNode, const FPCGStack& Stack) const;
+	PCG_API bool HasNodeProducedData(const UPCGNode* InNode, const FPCGStack& Stack) const;
 
-	UE_API void NotifyGPUToCPUReadback(const UPCGNode* InNode, const FPCGStack* InStack) const;
-	UE_API void NotifyCPUToGPUUpload(const UPCGNode* InNode, const FPCGStack* InStack) const;
+	PCG_API void NotifyGPUToCPUReadback(const UPCGNode* InNode, const FPCGStack* InStack) const;
+	PCG_API void NotifyCPUToGPUUpload(const UPCGNode* InNode, const FPCGStack* InStack) const;
 
-	UE_API bool DidNodeTriggerGPUToCPUReadback(const UPCGNode* InNode, const FPCGStack& Stack) const;
-	UE_API bool DidNodeTriggerCPUToGPUUpload(const UPCGNode* InNode, const FPCGStack& Stack) const;
+	PCG_API bool DidNodeTriggerGPUToCPUReadback(const UPCGNode* InNode, const FPCGStack& Stack) const;
+	PCG_API bool DidNodeTriggerCPUToGPUUpload(const UPCGNode* InNode, const FPCGStack& Stack) const;
 
-	UE_API void AddReferencedObjects(FReferenceCollector& Collector);
+	PCG_API void AddReferencedObjects(FReferenceCollector& Collector);
 
+	UE_DEPRECATED(5.8, "Use GetInspectionDataPtr/InspectData instead")
+	PCG_API const FPCGDataCollection* GetInspectionData(const FPCGStack& InStack) const { return nullptr;}
 private:
 	int32 InspectionCounter = 0;
 
-	TMap<FPCGStack, FPCGDataCollection> InspectionCache;
+	TMap<FPCGStack, TSharedPtr<FPCGDataCollection>> InspectionCache;
 
 	mutable PCG::FSharedLock InspectionCacheLock;
 
@@ -87,5 +106,3 @@ private:
 	mutable PCG::FSharedLock NodeToStacksTriggeringGPUTransfersLock;
 };
 #endif
-
-#undef UE_API
diff --git a/Engine/Plugins/PCG/Source/PCGEditor/Private/Managers/PCGEditorInspectionDataManager.cpp b/Engine/Plugins/PCG/Source/PCGEditor/Private/Managers/PCGEditorInspectionDataManager.cpp
--- a/Engine/Plugins/PCG/Source/PCGEditor/Private/Managers/PCGEditorInspectionDataManager.cpp
+++ b/Engine/Plugins/PCG/Source/PCGEditor/Private/Managers/PCGEditorInspectionDataManager.cpp
@@ -78,18 +78,47 @@ void FPCGEditorInspectionDataManager::OnSourceGenerated(IPCGGraphExecutionSource
 	}
 }
 
-const FPCGDataCollection* FPCGEditorInspectionDataManager::GetInspectionData(const FPCGEditorInspectionDataEntrySetupParams& InParams) const
+const TSharedPtr<FPCGInspectionData> FPCGEditorInspectionDataManager::GetInspectionDataPtr(const FPCGEditorInspectionDataEntrySetupParams& InParams) const
 {
-	const FPCGDataCollection* DataCollection = nullptr;
 	FPCGStack Stack{};
 	const UPCGPin* Pin = nullptr;
 
-	return GetInspectionData_Internal(InParams, DataCollection, Stack, Pin) ? DataCollection : nullptr;
+	const IPCGGraphExecutionSource* ExecutionSource = PCGSourceBeingInspected.Get();
+	if (!ExecutionSource)
+	{
+		return nullptr;
+	}
+
+	if (!GetInspectionDataParams(InParams, Stack, Pin))
+	{
+		return nullptr;
+	}
+
+	return ExecutionSource->GetExecutionState().GetInspection().GetInspectionDataPtr(Stack);
+}
+
+bool FPCGEditorInspectionDataManager::InspectData(const FPCGEditorInspectionDataEntrySetupParams& InParams, TFunctionRef<void(const FPCGDataCollection&)> InspectFunc) const
+{
+	FPCGStack Stack{};
+	const UPCGPin* Pin = nullptr;
+
+	const IPCGGraphExecutionSource* ExecutionSource = PCGSourceBeingInspected.Get();
+	if (!ExecutionSource)
+	{
+		return false;
+	}
+
+	if (!GetInspectionDataParams(InParams, Stack, Pin))
+	{
+		return false;
+	}
+
+	return ExecutionSource->GetExecutionState().GetInspection().InspectData(Stack, InspectFunc);
 }
 
 void FPCGEditorInspectionDataManager::SetupInspectionEntry(const FPCGEditorInspectionDataEntrySetupParams& InParams)
 {
-	InspectionEntries[InParams.EntryIndex] = ConstructEntry(InParams);
+	ConstructEntry(InParams, InspectionEntries[InParams.EntryIndex]);
 	OnInspectionEntryChangedDelegate.Broadcast(InspectionEntries[InParams.EntryIndex], InParams.EntryIndex);
 }
 
@@ -107,40 +136,49 @@ void FPCGEditorInspectionDataManager::AddReferencedObjects(FReferenceCollector&
 	}
 }
 
-FPCGEditorInspectionDataEntry FPCGEditorInspectionDataManager::ConstructEntry(const FPCGEditorInspectionDataEntrySetupParams& InParams) const
+void FPCGEditorInspectionDataManager::ConstructEntry(const FPCGEditorInspectionDataEntrySetupParams& InParams, FPCGEditorInspectionDataEntry& OutEntry) const
 {
-	FPCGEditorInspectionDataEntry Result{};
+	OutEntry = {};
+
+	const IPCGGraphExecutionSource* ExecutionSource = PCGSourceBeingInspected.Get();
+	if(!ExecutionSource)
+	{
+		return;
+	}
 
 	if (InParams.EntryIndex < 0
 		|| InParams.EntryIndex >= InspectionEntries.Num()
 		|| !InParams.Node
-		|| !PCGSourceBeingInspected.IsValid()
 		|| StackBeingInspected.GetStackFrames().IsEmpty())
 	{
-		return Result;
+		return;
 	}
 
-	const FPCGDataCollection* DataCollection = nullptr;
 	FPCGStack Stack{};
 	const UPCGPin* Pin = nullptr;
 	
-	if (!GetInspectionData_Internal(InParams, DataCollection, Stack, Pin))
+	if (!GetInspectionDataParams(InParams, Stack, Pin))
 	{
-		return Result;
+		return;
 	}
 
-	check(DataCollection && Pin);
+	check(Pin);
 
-	Result.InspectionData = *DataCollection;
-	Result.Node = InParams.Node;
-	Result.PinName = Pin->Properties.Label;
-	Result.bIsOutputPin = InParams.bIsOutputPin;
-	Result.Stack = MoveTemp(Stack);
+	const bool bFoundData = ExecutionSource->GetExecutionState().GetInspection().InspectData(Stack, [&OutEntry](const FPCGDataCollection& InspectionData)
+	{
+		OutEntry.InspectionData = InspectionData;
+	});
 
-	return Result;
+	if (bFoundData)
+	{
+		OutEntry.Node = InParams.Node;
+		OutEntry.PinName = Pin->Properties.Label;
+		OutEntry.bIsOutputPin = InParams.bIsOutputPin;
+		OutEntry.Stack = MoveTemp(Stack);
+	}
 }
 
-bool FPCGEditorInspectionDataManager::GetInspectionData_Internal(const FPCGEditorInspectionDataEntrySetupParams& InParams, FPCGDataCollection const*& OutData, FPCGStack& OutStack, UPCGPin const*& OutPin) const
+bool FPCGEditorInspectionDataManager::GetInspectionDataParams(const FPCGEditorInspectionDataEntrySetupParams& InParams, FPCGStack& OutStack, UPCGPin const*& OutPin) const
 {
 	const UPCGNode* PCGNode = InParams.Node ? InParams.Node->GetPCGNode() : nullptr;
 	if (!PCGNode)
@@ -176,10 +214,5 @@ bool FPCGEditorInspectionDataManager::GetInspectionData_Internal(const FPCGEdito
 	StackFrames.Emplace(PCGNode);
 	StackFrames.Emplace(OutPin);
 
-	if (IPCGGraphExecutionSource* Source = PCGSourceBeingInspected.Get())
-	{
-		OutData = Source->GetExecutionState().GetInspection().GetInspectionData(OutStack);
-	}
-	
-	return OutData != nullptr;
+	return true;
 }
diff --git a/Engine/Plugins/PCG/Source/PCGEditor/Private/Schema/PCGEditorGraphSchema.cpp b/Engine/Plugins/PCG/Source/PCGEditor/Private/Schema/PCGEditorGraphSchema.cpp
--- a/Engine/Plugins/PCG/Source/PCGEditor/Private/Schema/PCGEditorGraphSchema.cpp
+++ b/Engine/Plugins/PCG/Source/PCGEditor/Private/Schema/PCGEditorGraphSchema.cpp
@@ -1126,14 +1126,15 @@ bool FPCGEditorConnectionDrawingPolicy::UpdateParamsIfDebugging(UEdGraphPin* Out
 	StackFrames.Emplace(NodeToInspect);
 	StackFrames.Emplace(PinToInspect);
 
-	if (const FPCGDataCollection* DataCollection = PCGSource->GetExecutionState().GetInspection().GetInspectionData(Stack))
+	const bool bFoundData = PCGSource->GetExecutionState().GetInspection().InspectData(Stack, [&Params](const FPCGDataCollection& InspectionData)
 	{
-		if (DataCollection->TaggedData.Num() > 1)
+		if(InspectionData.TaggedData.Num() > 1)
 		{
 			Params.WireThickness *= GetDefault<UPCGEditorSettings>()->MultiDataEdgeDebugEmphasis;
 		}
-	}
-	else
+	});
+
+	if(!bFoundData)
 	{
 		Params.WireColor = Params.WireColor.Desaturate(GetDefault<UPCGEditorSettings>()->EmptyEdgeDebugDesaturateFactor);
 	}
diff --git a/Engine/Plugins/PCG/Source/PCGEditor/Private/Widgets/SPCGEditorGraphAttributeListView.cpp b/Engine/Plugins/PCG/Source/PCGEditor/Private/Widgets/SPCGEditorGraphAttributeListView.cpp
--- a/Engine/Plugins/PCG/Source/PCGEditor/Private/Widgets/SPCGEditorGraphAttributeListView.cpp
+++ b/Engine/Plugins/PCG/Source/PCGEditor/Private/Widgets/SPCGEditorGraphAttributeListView.cpp
@@ -846,16 +846,40 @@ void SPCGEditorGraphAttributeListView::AddReferencedObjects(FReferenceCollector&
 	}
 }
 
-const FPCGDataCollection* SPCGEditorGraphAttributeListView::GetInspectionData(const TSharedPtr<FPinComboBoxItem>& EditorPin) const
+TSharedPtr<FPCGInspectionData> SPCGEditorGraphAttributeListView::GetInspectionDataPtr(const TSharedPtr<FPinComboBoxItem>& EditorPin) const
 {
 	if (!EditorPin)
 	{
 		return nullptr;
 	}
 
-	const FPCGEditorInspectionDataEntrySetupParams Params{WidgetEntryNumber, PCGEditorGraphNode.Get(), EditorPin->PinIndex, EditorPin->bIsOutputPin};
+	if (TSharedPtr<FPCGEditor> PCGEditor = PCGEditorPtr.Pin())
+	{
+		const FPCGEditorInspectionDataEntrySetupParams Params{WidgetEntryNumber, PCGEditorGraphNode.Get(), EditorPin->PinIndex, EditorPin->bIsOutputPin};
+		return PCGEditor->GetInspectionDataManager().GetInspectionDataPtr(Params);
+	}
+	else
+	{
+		return nullptr;
+	}
+}
 
-	return PCGEditorPtr.IsValid() ? PCGEditorPtr.Pin()->GetInspectionDataManager().GetInspectionData(Params) : nullptr;
+bool SPCGEditorGraphAttributeListView::InspectData(const TSharedPtr<FPinComboBoxItem>& EditorPin, TFunctionRef<void(const FPCGDataCollection&)> InspectFunc) const
+{
+	if (!EditorPin)
+	{
+		return false;
+	}
+
+	if (TSharedPtr<FPCGEditor> PCGEditor = PCGEditorPtr.Pin())
+	{
+		const FPCGEditorInspectionDataEntrySetupParams Params{ WidgetEntryNumber, PCGEditorGraphNode.Get(), EditorPin->PinIndex, EditorPin->bIsOutputPin };
+		return PCGEditor->GetInspectionDataManager().InspectData(Params, InspectFunc);
+	}
+	else
+	{
+		return false;
+	}
 }
 
 const FPCGDataCollection* SPCGEditorGraphAttributeListView::GetInspectionData() const
@@ -1437,23 +1461,30 @@ TSharedRef<SWidget> SPCGEditorGraphAttributeListView::OnGenerateAdditionalOperat
 	return MenuBuilder.MakeWidget();
 }
 
-FPCGDataCollection SPCGEditorGraphAttributeListView::BuildDataCollectionForSave(bool bUsePinComboIndex, bool bUseDataComboIndex) const
+bool SPCGEditorGraphAttributeListView::BuildDataCollectionForSave(bool bUsePinComboIndex, bool bUseDataComboIndex, FPCGDataCollection* OutSaveData) const
 {
+	// Return if we found any data to save, when OutSaveData is nullptr this method can be used to test without copying any data
+	bool bDataToSave = false;
+
 	// Can't save N'th data on all pins, as it makes no sense
 	check(!bUseDataComboIndex || bUsePinComboIndex);
 
-	// Build a data collection from the indices we're given.
-	FPCGDataCollection Collection;
-
 	if (!bUsePinComboIndex)
 	{
 		for (const auto& PinComboBoxItem : PinComboBoxItems)
 		{
-			const FPCGDataCollection* PinInspectionData = GetInspectionData(PinComboBoxItem);
-			if (PinInspectionData)
+			InspectData(PinComboBoxItem, [&bDataToSave, OutSaveData](const FPCGDataCollection& InspectionData)
 			{
-				Collection.TaggedData.Append(PinInspectionData->TaggedData);
-			}
+				if(InspectionData.TaggedData.Num() > 0)
+				{
+					bDataToSave = true;
+				}
+
+				if (OutSaveData)
+				{
+					OutSaveData->TaggedData.Append(InspectionData.TaggedData);
+				}
+			});
 		}
 	}
 	else
@@ -1462,29 +1493,43 @@ FPCGDataCollection SPCGEditorGraphAttributeListView::BuildDataCollectionForSave(
 		{
 			if (!bUseDataComboIndex)
 			{
-				Collection = *PinInspectionData;
+				if (PinInspectionData->TaggedData.Num() > 0)
+				{
+					bDataToSave = true;
+				}
+
+				if (OutSaveData)
+				{
+					*OutSaveData = *PinInspectionData;
+				}
 			}
 			else if(int SelectedIndex = GetSelectedDataIndex(); PinInspectionData->TaggedData.IsValidIndex(SelectedIndex))
 			{
-				Collection.TaggedData.Add(PinInspectionData->TaggedData[SelectedIndex]);
+				bDataToSave = true;
+				
+				if (OutSaveData)
+				{
+					OutSaveData->TaggedData.Add(PinInspectionData->TaggedData[SelectedIndex]);
+				}
 			}
 		}
 	}
 
-	return Collection;
+	return bDataToSave;
 }
 
 void SPCGEditorGraphAttributeListView::SaveData(bool bUsePinComboIndex, bool bUseDataComboIndex)
 {
-	FPCGDataCollection Collection = BuildDataCollectionForSave(bUsePinComboIndex, bUseDataComboIndex);
-
-	if (Collection.TaggedData.IsEmpty())
+	// FPCGInspectionData will make sure data is visible to GC while we save
+	FPCGInspectionData SaveData;
+	SaveData.Data = MakeShared<FPCGDataCollection>();
+	if (!BuildDataCollectionForSave(bUsePinComboIndex, bUseDataComboIndex, SaveData.Data.Get()))
 	{
 		return;
 	}
 
 	UPCGDataCollectionExporter* Exporter = NewObject<UPCGDataCollectionExporter>();
-	Exporter->Data = Collection;
+	Exporter->Data = *SaveData.Data;
 
 	FPCGAssetExporterParameters Parameters;
 	UPCGAssetExporterUtils::CreateAsset(Exporter, Parameters);
@@ -1492,8 +1537,8 @@ void SPCGEditorGraphAttributeListView::SaveData(bool bUsePinComboIndex, bool bUs
 
 bool SPCGEditorGraphAttributeListView::CanSaveData(bool bUsePinComboIndex, bool bUseDataComboIndex) const
 {
-	FPCGDataCollection Collection = BuildDataCollectionForSave(bUsePinComboIndex, bUseDataComboIndex);
-	return !Collection.TaggedData.IsEmpty();
+	// Query only
+	return BuildDataCollectionForSave(bUsePinComboIndex, bUseDataComboIndex, nullptr);
 }
 
 FText SPCGEditorGraphAttributeListView::OnGenerateSelectedPinText() const
diff --git a/Engine/Plugins/PCG/Source/PCGEditor/Private/Widgets/SPCGEditorGraphAttributeListView.h b/Engine/Plugins/PCG/Source/PCGEditor/Private/Widgets/SPCGEditorGraphAttributeListView.h
--- a/Engine/Plugins/PCG/Source/PCGEditor/Private/Widgets/SPCGEditorGraphAttributeListView.h
+++ b/Engine/Plugins/PCG/Source/PCGEditor/Private/Widgets/SPCGEditorGraphAttributeListView.h
@@ -18,6 +18,7 @@
 enum class EPCGMetadataTypes : uint8;
 enum class EPCGEditorPanel;
 struct FPCGDataCollection;
+struct FPCGInspectionData;
 class FPCGEditor;
 struct FPCGEditorInspectionDataEntry;
 class FPCGEditorInspectionDataManager;
@@ -207,7 +208,8 @@ private:
 		bool bIsOutputPin = true;
 	};
 
-	const FPCGDataCollection* GetInspectionData(const TSharedPtr<FPinComboBoxItem>& EditorPin) const;
+	TSharedPtr<FPCGInspectionData> GetInspectionDataPtr(const TSharedPtr<FPinComboBoxItem>& EditorPin) const;
+	bool InspectData(const TSharedPtr<FPinComboBoxItem>& EditorPin, TFunctionRef<void(const FPCGDataCollection&)> InspectFunc) const;
 	const FPCGDataCollection* GetInspectionData() const;
 
 	FText OnGenerateSelectedPinText() const;
@@ -229,7 +231,7 @@ private:
 	ECheckBoxState GetAnyAttributeEnabledState() const;
 	bool IsAttributeEnabled(FName InAttributeName) const;
 
-	FPCGDataCollection BuildDataCollectionForSave(bool bUsePinComboIndex, bool bUseDataComboIndex) const;
+	bool BuildDataCollectionForSave(bool bUsePinComboIndex, bool bUseDataComboIndex, FPCGDataCollection* OutData) const;
 	void SaveData(bool bUsePinIndex, bool bUseDataIndex);
 	bool CanSaveData(bool bUsePinIndex, bool bUseDataIndex) const;
 
diff --git a/Engine/Plugins/PCG/Source/PCGEditor/Public/Managers/PCGEditorInspectionDataManager.h b/Engine/Plugins/PCG/Source/PCGEditor/Public/Managers/PCGEditorInspectionDataManager.h
--- a/Engine/Plugins/PCG/Source/PCGEditor/Public/Managers/PCGEditorInspectionDataManager.h
+++ b/Engine/Plugins/PCG/Source/PCGEditor/Public/Managers/PCGEditorInspectionDataManager.h
@@ -7,6 +7,7 @@
 #include "Graph/PCGStackContext.h"
 
 class UPCGEditorGraphNodeBase;
+struct FPCGInspectionData;
 
 struct FPCGEditorInspectionDataEntry
 {
@@ -68,7 +69,11 @@ public:
 
 	PCGEDITOR_API void OnSourceGenerated(IPCGGraphExecutionSource* InSource) const;
 
-	PCGEDITOR_API const FPCGDataCollection* GetInspectionData(const FPCGEditorInspectionDataEntrySetupParams& InParams) const;
+	UE_DEPRECATED(5.8, "Use GetInspectionDataPtr/InspectData instead")
+	PCGEDITOR_API const FPCGDataCollection* GetInspectionData(const FPCGEditorInspectionDataEntrySetupParams& InParams) const { return nullptr; }
+
+	PCGEDITOR_API const TSharedPtr<FPCGInspectionData> GetInspectionDataPtr(const FPCGEditorInspectionDataEntrySetupParams& InParams) const;
+	PCGEDITOR_API bool InspectData(const FPCGEditorInspectionDataEntrySetupParams& InParams, TFunctionRef<void(const FPCGDataCollection&)> InspectFunc) const;
 	
 	PCGEDITOR_API void SetupInspectionEntry(const FPCGEditorInspectionDataEntrySetupParams& InParams);
 
@@ -83,9 +88,9 @@ public:
 	FOnInspectionEntryChanged OnInspectionEntryChangedDelegate;
 
 private:
-	FPCGEditorInspectionDataEntry ConstructEntry(const FPCGEditorInspectionDataEntrySetupParams& InParams) const;
+	void ConstructEntry(const FPCGEditorInspectionDataEntrySetupParams& InParams, FPCGEditorInspectionDataEntry& OutEntry) const;
 
-	bool GetInspectionData_Internal(const FPCGEditorInspectionDataEntrySetupParams& InParams, FPCGDataCollection const*& OutData, FPCGStack& OutStack, UPCGPin const*& OutPin) const;
+	bool GetInspectionDataParams(const FPCGEditorInspectionDataEntrySetupParams& InParams, FPCGStack& OutStack, UPCGPin const*& OutPin) const;
 	
 	FPCGSoftGraphExecutionSource PCGSourceBeingInspected;
 	// Implementation note: we'll keep the last valid source inspected so we don't un-inspect on spurious selection changes
