[PCG] DebugDisplay extended to different execution sources
- If execution source isn't actor based, spawn a transient debug actor to handle the debug primitives

#Synced-CL 49570007

diff --git a/Engine/Plugins/PCG/Source/PCG/Private/Elements/PCGDebugElement.cpp b/Engine/Plugins/PCG/Source/PCG/Private/Elements/PCGDebugElement.cpp
--- a/Engine/Plugins/PCG/Source/PCG/Private/Elements/PCGDebugElement.cpp
+++ b/Engine/Plugins/PCG/Source/PCG/Private/Elements/PCGDebugElement.cpp
@@ -59,7 +59,7 @@ bool FPCGDebugElement::ExecuteInternal(FPCGContext* Context) const
 
 		if (!TargetActor)
 		{
-			TargetActor = Context->GetTypedExecutionTarget<AActor>();
+			TargetActor = Context->GetDebugDisplayActor();
 		}
 
 		if (const IPCGDataVisualization* DataVis = DataVisRegistry.GetDataVisualization(Input.Data->GetClass()))
diff --git a/Engine/Plugins/PCG/Source/PCG/Private/Helpers/PCGActorHelpers.cpp b/Engine/Plugins/PCG/Source/PCG/Private/Helpers/PCGActorHelpers.cpp
--- a/Engine/Plugins/PCG/Source/PCG/Private/Helpers/PCGActorHelpers.cpp
+++ b/Engine/Plugins/PCG/Source/PCG/Private/Helpers/PCGActorHelpers.cpp
@@ -74,9 +74,9 @@ UPCGManagedISMComponent* UPCGActorHelpers::GetOrCreateManagedISMC(AActor* InTarg
 	return GetOrCreateManagedISMC(InTargetActor, InSourceComponent, InParams);
 }
 
-UInstancedStaticMeshComponent* UPCGActorHelpers::GetOrCreateISMC(AActor* InTargetActor, UPCGComponent* InSourceComponent, const FPCGISMComponentBuilderParams& InParams, FPCGContext* OptionalContext)
+UInstancedStaticMeshComponent* UPCGActorHelpers::GetOrCreateISMC(AActor* InTargetActor, IPCGGraphExecutionSource* ExecutionSource, const FPCGISMComponentBuilderParams& InParams, FPCGContext* OptionalContext)
 {
-	UPCGManagedISMComponent* MISMC = GetOrCreateManagedISMC(InTargetActor, InSourceComponent, InParams, OptionalContext);
+	UPCGManagedISMComponent* MISMC = GetOrCreateManagedISMC(InTargetActor, ExecutionSource, InParams, OptionalContext);
 	if (MISMC)
 	{
 		return MISMC->GetComponent();
@@ -87,12 +87,20 @@ UInstancedStaticMeshComponent* UPCGActorHelpers::GetOrCreateISMC(AActor* InTarge
 	}
 }
 
-UPCGManagedISMComponent* UPCGActorHelpers::GetOrCreateManagedISMC(AActor* InTargetActor, UPCGComponent* InSourceComponent, const FPCGISMComponentBuilderParams& InParams, FPCGContext* OptionalContext)
+UPCGManagedISMComponent* UPCGActorHelpers::GetOrCreateManagedISMC(AActor* InTargetActor, IPCGGraphExecutionSource* InExecutionSource, const FPCGISMComponentBuilderParams& InParams, FPCGContext* OptionalContext)
 {
-	check(InTargetActor && InSourceComponent);
+	check(InTargetActor && InExecutionSource);
 
 	TRACE_CPUPROFILER_EVENT_SCOPE(UPCGActorHelpers::GetOrCreateManagedISMC);
 
+	UObject* ExecutionSourceObject = CastChecked<UObject>(InExecutionSource);
+	
+	FPCGManagedResourceContainer* ResourceContainer = InExecutionSource->GetExecutionState().GetManagedResourceContainer();
+	if (!ResourceContainer)
+	{
+		return nullptr;
+	}
+
 	FISMComponentDescriptor Descriptor(InParams.Descriptor);
 
 	const UStaticMesh* StaticMesh = Descriptor.StaticMesh;
@@ -101,10 +109,10 @@ UPCGManagedISMComponent* UPCGActorHelpers::GetOrCreateManagedISMC(AActor* InTarg
 		return nullptr;
 	}
 
-	auto AddTagsToComponent = [InSourceComponent, &InParams](UInstancedStaticMeshComponent* ISMC)
+	auto AddTagsToComponent = [ExecutionSourceObject, &InParams](UInstancedStaticMeshComponent* ISMC)
 	{
 		ISMC->ComponentTags.AddUnique(PCGHelpers::DefaultPCGTag);
-		ISMC->ComponentTags.AddUnique(InSourceComponent->GetFName());
+		ISMC->ComponentTags.AddUnique(ExecutionSourceObject->GetFName());
 			
 		for (FName ComponentTag : InParams.Descriptor.ComponentTags)
 		{
@@ -150,10 +158,10 @@ UPCGManagedISMComponent* UPCGActorHelpers::GetOrCreateManagedISMC(AActor* InTarg
 		UPCGManagedISMComponent* MatchingResource = nullptr;
 		if(InParams.SettingsCrc.IsValid())
 		{
-			// If we explicitely are trying to find a RF_Transient through the params or if the PCGComponent itself is in preview mode
-			const bool bTransient = InParams.bTransient || (InSourceComponent->GetEditingMode() == EPCGEditorDirtyMode::Preview);
+			// If we explicitely are trying to find a RF_Transient through the params or if the Execution source itself is in preview mode
+			const bool bTransient = InParams.bTransient || (InExecutionSource->GetExecutionState().IsInPreviewMode());
 
-			InSourceComponent->ForEachManagedResource([&MatchingResource, &InParams, &InTargetActor, &Descriptor, bTransient](UPCGManagedResource* InResource)
+			ResourceContainer->ForEachManagedResource([&MatchingResource, &InParams, &InTargetActor, &Descriptor, bTransient](UPCGManagedResource* InResource)
 			{
 				// Early out if already found a match
 				if (MatchingResource)
@@ -194,7 +202,7 @@ UPCGManagedISMComponent* UPCGActorHelpers::GetOrCreateManagedISMC(AActor* InTarg
 			UInstancedStaticMeshComponent* ISMC = Cast<UInstancedStaticMeshComponent>(MatchingResource->GeneratedComponent.Get());
 			if (ensure(ISMC))
 			{
-				ISMC->Modify(!InSourceComponent->IsInPreviewMode());
+				ISMC->Modify(!InExecutionSource->GetExecutionState().IsInPreviewMode());
 				AddTagsToComponent(ISMC);
 			}
 
@@ -203,7 +211,7 @@ UPCGManagedISMComponent* UPCGActorHelpers::GetOrCreateManagedISMC(AActor* InTarg
 	}
 
 	// No matching ISM component found, let's create a new one
-	InTargetActor->Modify(!InSourceComponent->IsInPreviewMode());
+	InTargetActor->Modify(!InExecutionSource->GetExecutionState().IsInPreviewMode());
 
 	FString ComponentName;
 
@@ -218,7 +226,7 @@ UPCGManagedISMComponent* UPCGActorHelpers::GetOrCreateManagedISMC(AActor* InTarg
 
 	ComponentName += StaticMesh->GetName();
 
-	EObjectFlags ObjectFlags = (InSourceComponent->IsInPreviewMode() ? RF_Transient | RF_NonPIEDuplicateTransient : RF_NoFlags);
+	EObjectFlags ObjectFlags = (InExecutionSource->GetExecutionState().IsInPreviewMode() ? RF_Transient | RF_NonPIEDuplicateTransient : RF_NoFlags);
 
 	// Used for debug data visualization
 	if (InParams.bTransient)
@@ -250,7 +258,7 @@ UPCGManagedISMComponent* UPCGActorHelpers::GetOrCreateManagedISMC(AActor* InTarg
 	AddTagsToComponent(ISMC);
 
 	// Create managed resource on source component
-	UPCGManagedISMComponent* Resource = NewObject<UPCGManagedISMComponent>(InSourceComponent, NAME_None, InParams.bTransient ? RF_Transient : RF_NoFlags);
+	UPCGManagedISMComponent* Resource = NewObject<UPCGManagedISMComponent>(ExecutionSourceObject, NAME_None, InParams.bTransient ? RF_Transient : RF_NoFlags);
 	Resource->SetComponent(ISMC);
 	Resource->SetDescriptor(Descriptor);
 	if (InTargetActor->GetRootComponent())
@@ -259,17 +267,25 @@ UPCGManagedISMComponent* UPCGActorHelpers::GetOrCreateManagedISMC(AActor* InTarg
 	}
 	
 	Resource->SetSettingsCrc(InParams.SettingsCrc);
-	InSourceComponent->AddToManagedResources(Resource);
+	ResourceContainer->AddManagedResource(Resource);
 
 	return Resource;
 }
 
-UPCGManagedISKMComponent* UPCGActorHelpers::GetOrCreateManagedABMC(AActor* InTargetActor, UPCGComponent* InSourceComponent, const FPCGSkinnedMeshComponentBuilderParams& InParams, FPCGContext* OptionalContext)
+UPCGManagedISKMComponent* UPCGActorHelpers::GetOrCreateManagedABMC(AActor* InTargetActor, IPCGGraphExecutionSource* InExecutionSource, const FPCGSkinnedMeshComponentBuilderParams& InParams, FPCGContext* OptionalContext)
 {
-	check(InTargetActor && InSourceComponent);
+	check(InTargetActor && InExecutionSource);
 
 	TRACE_CPUPROFILER_EVENT_SCOPE(UPCGActorHelpers::GetOrCreateManagedABMC);
 
+	UObject* ExecutionSourceObject = CastChecked<UObject>(InExecutionSource);
+
+	FPCGManagedResourceContainer* ResourceContainer = InExecutionSource->GetExecutionState().GetManagedResourceContainer();
+	if (!ResourceContainer)
+	{
+		return nullptr;
+	}
+
 	FSkinnedMeshComponentDescriptor Descriptor(InParams.Descriptor);
 
 	const USkinnedAsset* SkinnedAsset = Descriptor.SkinnedAsset;
@@ -278,10 +294,10 @@ UPCGManagedISKMComponent* UPCGActorHelpers::GetOrCreateManagedABMC(AActor* InTar
 		return nullptr;
 	}
 
-	auto AddTagsToComponent = [InSourceComponent, &InParams](UInstancedSkinnedMeshComponent* ISKMC)
+	auto AddTagsToComponent = [ExecutionSourceObject, &InParams](UInstancedSkinnedMeshComponent* ISKMC)
 	{
 		ISKMC->ComponentTags.AddUnique(PCGHelpers::DefaultPCGTag);
-		ISKMC->ComponentTags.AddUnique(InSourceComponent->GetFName());
+		ISKMC->ComponentTags.AddUnique(ExecutionSourceObject->GetFName());
 			
 		for (const FName& ComponentTag : InParams.Descriptor.ComponentTags)
 		{
@@ -299,7 +315,7 @@ UPCGManagedISKMComponent* UPCGActorHelpers::GetOrCreateManagedABMC(AActor* InTar
 		UPCGManagedISKMComponent* MatchingResource = nullptr;
 		if(InParams.SettingsCrc.IsValid())
 		{
-			InSourceComponent->ForEachManagedResource([&MatchingResource, &InParams, &InTargetActor, &Descriptor](UPCGManagedResource* InResource)
+			ResourceContainer->ForEachManagedResource([&MatchingResource, &InParams, &InTargetActor, &Descriptor](UPCGManagedResource* InResource)
 			{
 				// Early out if already found a match
 				if (MatchingResource)
@@ -337,7 +353,7 @@ UPCGManagedISKMComponent* UPCGActorHelpers::GetOrCreateManagedABMC(AActor* InTar
 			UInstancedSkinnedMeshComponent* ISKMC = Cast<UInstancedSkinnedMeshComponent>(MatchingResource->GeneratedComponent.Get());
 			if (ensure(ISKMC))
 			{
-				ISKMC->Modify(!InSourceComponent->IsInPreviewMode());
+				ISKMC->Modify(!InExecutionSource->GetExecutionState().IsInPreviewMode());
 				AddTagsToComponent(ISKMC);
 			}
 
@@ -346,11 +362,11 @@ UPCGManagedISKMComponent* UPCGActorHelpers::GetOrCreateManagedABMC(AActor* InTar
 	}
 
 	// No matching ABM component found, let's create a new one
-	InTargetActor->Modify(!InSourceComponent->IsInPreviewMode());
+	InTargetActor->Modify(!InExecutionSource->GetExecutionState().IsInPreviewMode());
 
 	FString ComponentName = TEXT("ABM_") + SkinnedAsset->GetName();
 
-	EObjectFlags ObjectFlags = (InSourceComponent->IsInPreviewMode() ? RF_Transient | RF_NonPIEDuplicateTransient : RF_NoFlags);
+	EObjectFlags ObjectFlags = (InExecutionSource->GetExecutionState().IsInPreviewMode() ? RF_Transient | RF_NonPIEDuplicateTransient : RF_NoFlags);
 
 	// Used for debug data visualization
 	if (InParams.bTransient)
@@ -374,7 +390,7 @@ UPCGManagedISKMComponent* UPCGActorHelpers::GetOrCreateManagedABMC(AActor* InTar
 	AddTagsToComponent(ISKMC);
 
 	// Create managed resource on source component
-	UPCGManagedISKMComponent* Resource = NewObject<UPCGManagedISKMComponent>(InSourceComponent, NAME_None, InParams.bTransient ? RF_Transient : RF_NoFlags);
+	UPCGManagedISKMComponent* Resource = NewObject<UPCGManagedISKMComponent>(ExecutionSourceObject, NAME_None, InParams.bTransient ? RF_Transient : RF_NoFlags);
 	Resource->SetComponent(ISKMC);
 	Resource->SetDescriptor(Descriptor);
 	if (InTargetActor->GetRootComponent())
@@ -383,7 +399,8 @@ UPCGManagedISKMComponent* UPCGActorHelpers::GetOrCreateManagedABMC(AActor* InTar
 	}
 	
 	Resource->SetSettingsCrc(InParams.SettingsCrc);
-	InSourceComponent->AddToManagedResources(Resource);
+
+	ResourceContainer->AddManagedResource(Resource);
 
 	return Resource;
 }
@@ -401,9 +418,9 @@ UPCGManagedSplineMeshComponent* UPCGActorHelpers::GetOrCreateManagedSplineMeshCo
 	return GetOrCreateManagedSplineMeshComponent(InTargetActor, InSourceComponent, InParams);
 }
 
-USplineMeshComponent* UPCGActorHelpers::GetOrCreateSplineMeshComponent(AActor* InTargetActor, UPCGComponent* InSourceComponent, const FPCGSplineMeshComponentBuilderParameters& InParams, FPCGContext* OptionalContext)
+USplineMeshComponent* UPCGActorHelpers::GetOrCreateSplineMeshComponent(AActor* InTargetActor, IPCGGraphExecutionSource* InExecutionSource, const FPCGSplineMeshComponentBuilderParameters& InParams, FPCGContext* OptionalContext)
 {
-	UPCGManagedSplineMeshComponent* ManagedComponent = GetOrCreateManagedSplineMeshComponent(InTargetActor, InSourceComponent, InParams, OptionalContext);
+	UPCGManagedSplineMeshComponent* ManagedComponent = GetOrCreateManagedSplineMeshComponent(InTargetActor, InExecutionSource, InParams, OptionalContext);
 	if (ManagedComponent)
 	{
 		return ManagedComponent->GetComponent();
@@ -414,9 +431,15 @@ USplineMeshComponent* UPCGActorHelpers::GetOrCreateSplineMeshComponent(AActor* I
 	}
 }
 
-UPCGManagedSplineMeshComponent* UPCGActorHelpers::GetOrCreateManagedSplineMeshComponent(AActor* InTargetActor, UPCGComponent* InSourceComponent, const FPCGSplineMeshComponentBuilderParameters& InParams, FPCGContext* OptionalContext)
+UPCGManagedSplineMeshComponent* UPCGActorHelpers::GetOrCreateManagedSplineMeshComponent(AActor* InTargetActor, IPCGGraphExecutionSource* InExecutionSource, const FPCGSplineMeshComponentBuilderParameters& InParams, FPCGContext* OptionalContext)
 {
-	check(InTargetActor && InSourceComponent);
+	check(InTargetActor && InExecutionSource);
+
+	FPCGManagedResourceContainer* ResourceContainer = InExecutionSource->GetExecutionState().GetManagedResourceContainer();
+	if (!ResourceContainer)
+	{
+		return nullptr;
+	}
 
 	const UStaticMesh* StaticMesh = InParams.Descriptor.StaticMesh;
 
@@ -455,7 +478,7 @@ UPCGManagedSplineMeshComponent* UPCGActorHelpers::GetOrCreateManagedSplineMeshCo
 		UPCGManagedSplineMeshComponent* MatchingResource = nullptr;
 		if(InParams.SettingsCrc.IsValid())
 		{
-			InSourceComponent->ForEachManagedResource([&MatchingResource, &InParams, &InTargetActor](UPCGManagedResource* InResource)
+			ResourceContainer->ForEachManagedResource([&MatchingResource, &InParams, &InTargetActor](UPCGManagedResource* InResource)
 			{
 				// Early out if already found a match
 				if (MatchingResource)
@@ -492,7 +515,7 @@ UPCGManagedSplineMeshComponent* UPCGActorHelpers::GetOrCreateManagedSplineMeshCo
 			USplineMeshComponent* SplineMeshComponent = MatchingResource->GetComponent();
 			if (ensure(SplineMeshComponent))
 			{
-				SplineMeshComponent->Modify(!InSourceComponent->IsInPreviewMode());
+				SplineMeshComponent->Modify(!InExecutionSource->GetExecutionState().IsInPreviewMode());
 
 				const USceneComponent* RootComponent = InTargetActor->GetRootComponent();
 				
@@ -511,10 +534,10 @@ UPCGManagedSplineMeshComponent* UPCGActorHelpers::GetOrCreateManagedSplineMeshCo
 	}
 
 	// No matching component found, let's create a new one.
-	InTargetActor->Modify(!InSourceComponent->IsInPreviewMode());
+	InTargetActor->Modify(!InExecutionSource->GetExecutionState().IsInPreviewMode());
 
 	FString ComponentName = TEXT("PCGSplineMeshComponent_") + StaticMesh->GetName();
-	const EObjectFlags ObjectFlags = (InSourceComponent->IsInPreviewMode() ? RF_Transient | RF_NonPIEDuplicateTransient : RF_NoFlags);
+	const EObjectFlags ObjectFlags = (InExecutionSource->GetExecutionState().IsInPreviewMode() ? RF_Transient | RF_NonPIEDuplicateTransient : RF_NoFlags);
 	USplineMeshComponent* SplineMeshComponent = InParams.Descriptor.CreateComponent(InTargetActor, MakeUniqueObjectName(InTargetActor, InParams.Descriptor.ComponentClass, FName(ComponentName)), ObjectFlags);
 
 	// Init Component
@@ -540,16 +563,18 @@ UPCGManagedSplineMeshComponent* UPCGActorHelpers::GetOrCreateManagedSplineMeshCo
 
 	AttachToRoot(SplineMeshComponent);
 
-	SplineMeshComponent->ComponentTags.Add(InSourceComponent->GetFName());
+	UObject* ExecutionSourceObject = CastChecked<UObject>(InExecutionSource);
+	SplineMeshComponent->ComponentTags.Add(ExecutionSourceObject->GetFName());
 	SplineMeshComponent->ComponentTags.Add(PCGHelpers::DefaultPCGTag);
 
 	// Create managed resource on source component
-	UPCGManagedSplineMeshComponent* Resource = NewObject<UPCGManagedSplineMeshComponent>(InSourceComponent);
+	UPCGManagedSplineMeshComponent* Resource = NewObject<UPCGManagedSplineMeshComponent>(ExecutionSourceObject);
 	Resource->SetComponent(SplineMeshComponent);
 	Resource->SetDescriptor(InParams.Descriptor);
 	Resource->SetSplineMeshParams(InParams.SplineMeshParams);
 	Resource->SetSettingsCrc(InParams.SettingsCrc);
-	InSourceComponent->AddToManagedResources(Resource);
+	
+	ResourceContainer->AddManagedResource(Resource);
 
 	return Resource;
 }
diff --git a/Engine/Plugins/PCG/Source/PCG/Private/PCGComponentExecutionState.cpp b/Engine/Plugins/PCG/Source/PCG/Private/PCGComponentExecutionState.cpp
--- a/Engine/Plugins/PCG/Source/PCG/Private/PCGComponentExecutionState.cpp
+++ b/Engine/Plugins/PCG/Source/PCG/Private/PCGComponentExecutionState.cpp
@@ -194,6 +194,12 @@ uint32 FPCGComponentExecutionState::GetGenerationGridSize() const
 	return (IsPartitioned() || IsLocalSource()) ? Component->GetGenerationGridSize() : PCGHiGenGrid::UninitializedGridSize();
 }
 
+bool FPCGComponentExecutionState::IsInPreviewMode() const
+{
+	check(Component);
+	return Component->IsInPreviewMode();
+}
+
 FPCGTaskId FPCGComponentExecutionState::GetGenerationTaskId() const
 {
 	check(Component);
diff --git a/Engine/Plugins/PCG/Source/PCG/Private/PCGContext.cpp b/Engine/Plugins/PCG/Source/PCG/Private/PCGContext.cpp
--- a/Engine/Plugins/PCG/Source/PCG/Private/PCGContext.cpp
+++ b/Engine/Plugins/PCG/Source/PCG/Private/PCGContext.cpp
@@ -622,6 +622,23 @@ bool FPCGContext::IsExecutingGraphInspected() const
 
 	return Graph && Graph->IsInspecting() && Graph->GetInspectedStack() == *StackPtr;
 }
+
+AActor* FPCGContext::GetDebugDisplayActor() const
+{
+	if (AActor* Actor = GetTypedExecutionTarget<AActor>())
+	{
+		return Actor;
+	}
+	else if (IPCGGraphExecutionSource* ExecutionSourcePtr = ExecutionSource.Get())
+	{
+		if (UPCGSubsystem* PCGSubsystem = Cast<UPCGSubsystem>(ExecutionSource->GetExecutionState().GetSubsystem()))
+		{
+			return PCGSubsystem->GetOrCreateDebugActor();
+		}
+	}
+
+	return nullptr;
+}
 #endif // WITH_EDITOR
 
 void FPCGContext::AddStructReferencedObjects(FReferenceCollector& Collector)
diff --git a/Engine/Plugins/PCG/Source/PCG/Private/PCGElement.cpp b/Engine/Plugins/PCG/Source/PCG/Private/PCGElement.cpp
--- a/Engine/Plugins/PCG/Source/PCG/Private/PCGElement.cpp
+++ b/Engine/Plugins/PCG/Source/PCG/Private/PCGElement.cpp
@@ -993,7 +993,7 @@ bool IPCGElement::DebugDisplay(FPCGContext* Context) const
 	{
 		if (const IPCGDataVisualization* DataVis = DataVisRegistry.GetDataVisualization(Data.Data->GetClass()))
 		{
-			DataVis->ExecuteDebugDisplay(Context, SettingsInterface, Data.Data, Context->GetTypedExecutionTarget<AActor>());
+			DataVis->ExecuteDebugDisplay(Context, SettingsInterface, Data.Data, Context->GetDebugDisplayActor());
 		}
 	}
 
diff --git a/Engine/Plugins/PCG/Source/PCG/Private/Subsystems/PCGSubsystem.cpp b/Engine/Plugins/PCG/Source/PCG/Private/Subsystems/PCGSubsystem.cpp
--- a/Engine/Plugins/PCG/Source/PCG/Private/Subsystems/PCGSubsystem.cpp
+++ b/Engine/Plugins/PCG/Source/PCG/Private/Subsystems/PCGSubsystem.cpp
@@ -513,6 +513,22 @@ void UPCGSubsystem::DestroyAllPCGWorldActors()
 	}
 }
 
+AActor* UPCGSubsystem::GetOrCreateDebugActor()
+{
+	if (!DebugActor.IsValid())
+	{
+		FActorSpawnParameters SpawnParams;
+		SpawnParams.bCreateActorPackage = false;
+		SpawnParams.bHideFromSceneOutliner = true;
+		SpawnParams.Name = TEXT("PCG_Debug");
+		SpawnParams.ObjectFlags = RF_Transient;
+		UPCGActorHelpers::FSpawnDefaultActorParams Params(GetWorld(), AActor::StaticClass(), FTransform::Identity, SpawnParams);
+		DebugActor = UPCGActorHelpers::SpawnDefaultActor(Params);
+	}
+
+	return DebugActor.Get();
+}
+
 void UPCGSubsystem::LogAbnormalComponentStates(bool bGroupByState) const
 {
 	TArray<UPCGComponent*> DeactivatedComponents;
diff --git a/Engine/Plugins/PCG/Source/PCG/Public/Helpers/PCGActorHelpers.h b/Engine/Plugins/PCG/Source/PCG/Public/Helpers/PCGActorHelpers.h
--- a/Engine/Plugins/PCG/Source/PCG/Public/Helpers/PCGActorHelpers.h
+++ b/Engine/Plugins/PCG/Source/PCG/Public/Helpers/PCGActorHelpers.h
@@ -142,11 +142,11 @@ PRAGMA_ENABLE_DEPRECATION_WARNINGS
 	UE_DEPRECATED(5.6, "Use signature with no SettingsUID. Make sure to fill Params.SettingsCRC to enable component reuse.")
 	static PCG_API UPCGManagedSplineMeshComponent* GetOrCreateManagedSplineMeshComponent(AActor* InTargetActor, UPCGComponent* SourceComponent, uint64 SettingsUID, const FPCGSplineMeshComponentBuilderParameters& Params);
 	
-	static PCG_API UInstancedStaticMeshComponent* GetOrCreateISMC(AActor* InTargetActor, UPCGComponent* SourceComponent, const FPCGISMComponentBuilderParams& Params, FPCGContext* OptionalContext = nullptr);
-	static PCG_API UPCGManagedISMComponent* GetOrCreateManagedISMC(AActor* InTargetActor, UPCGComponent* SourceComponent, const FPCGISMComponentBuilderParams& Params, FPCGContext* OptionalContext = nullptr);
-	static PCG_API UPCGManagedISKMComponent* GetOrCreateManagedABMC(AActor* InTargetActor, UPCGComponent* SourceComponent, const FPCGSkinnedMeshComponentBuilderParams& Params, FPCGContext* OptionalContext = nullptr);
-	static PCG_API USplineMeshComponent* GetOrCreateSplineMeshComponent(AActor* InTargetActor, UPCGComponent* SourceComponent, const FPCGSplineMeshComponentBuilderParameters& Params, FPCGContext* OptionalContext = nullptr);
-	static PCG_API UPCGManagedSplineMeshComponent* GetOrCreateManagedSplineMeshComponent(AActor* InTargetActor, UPCGComponent* SourceComponent, const FPCGSplineMeshComponentBuilderParameters& Params, FPCGContext* OptionalContext = nullptr);
+	static PCG_API UInstancedStaticMeshComponent* GetOrCreateISMC(AActor* InTargetActor, IPCGGraphExecutionSource* ExecutionSource, const FPCGISMComponentBuilderParams& Params, FPCGContext* OptionalContext = nullptr);
+	static PCG_API UPCGManagedISMComponent* GetOrCreateManagedISMC(AActor* InTargetActor, IPCGGraphExecutionSource* ExecutionSource, const FPCGISMComponentBuilderParams& Params, FPCGContext* OptionalContext = nullptr);
+	static PCG_API UPCGManagedISKMComponent* GetOrCreateManagedABMC(AActor* InTargetActor, IPCGGraphExecutionSource* ExecutionSource, const FPCGSkinnedMeshComponentBuilderParams& Params, FPCGContext* OptionalContext = nullptr);
+	static PCG_API USplineMeshComponent* GetOrCreateSplineMeshComponent(AActor* InTargetActor, IPCGGraphExecutionSource* ExecutionSource, const FPCGSplineMeshComponentBuilderParameters& Params, FPCGContext* OptionalContext = nullptr);
+	static PCG_API UPCGManagedSplineMeshComponent* GetOrCreateManagedSplineMeshComponent(AActor* InTargetActor, IPCGGraphExecutionSource* ExecutionSource, const FPCGSplineMeshComponentBuilderParameters& Params, FPCGContext* OptionalContext = nullptr);
 
 	static PCG_API bool DeleteActors(UWorld* World, const TArray<TSoftObjectPtr<AActor>>& ActorsToDelete);
 
diff --git a/Engine/Plugins/PCG/Source/PCG/Public/PCGComponentExecutionState.h b/Engine/Plugins/PCG/Source/PCG/Public/PCGComponentExecutionState.h
--- a/Engine/Plugins/PCG/Source/PCG/Public/PCGComponentExecutionState.h
+++ b/Engine/Plugins/PCG/Source/PCG/Public/PCGComponentExecutionState.h
@@ -40,6 +40,7 @@ public:
 	virtual FPCGGridDescriptor GetGridDescriptor(uint32 InGridSize) const override;
 	virtual EPCGHiGenGrid GetGenerationGrid() const override;
 	virtual uint32 GetGenerationGridSize() const override;
+	virtual bool IsInPreviewMode() const override;
 	virtual FPCGTaskId GetGenerationTaskId() const override;
 	virtual void StoreOutputDataForPin(const FString& InResourceKey, const FPCGDataCollection& InData) const override;
 	virtual const FPCGDataCollection* RetrieveOutputDataForPin(const FString& InResourceKey) const override;
diff --git a/Engine/Plugins/PCG/Source/PCG/Public/PCGContext.h b/Engine/Plugins/PCG/Source/PCG/Public/PCGContext.h
--- a/Engine/Plugins/PCG/Source/PCG/Public/PCGContext.h
+++ b/Engine/Plugins/PCG/Source/PCG/Public/PCGContext.h
@@ -232,6 +232,9 @@ PRAGMA_ENABLE_DEPRECATION_WARNINGS
 
 	/** Whether the graph for the current frame in the execution stack is selected for debugging in the graph editor. */
 	UE_API bool IsExecutingGraphInspected() const;
+
+	/** Returns the proper actor to do debug visualization. */
+	UE_API AActor* GetDebugDisplayActor() const;
 #endif // WITH_EDITOR
 
 	/** Gathers references to objects to prevent them from being garbage collected. Extend resources to be collected with AddExtractStructReferencedObjects. */
diff --git a/Engine/Plugins/PCG/Source/PCG/Public/PCGGraphExecutionStateInterface.h b/Engine/Plugins/PCG/Source/PCG/Public/PCGGraphExecutionStateInterface.h
--- a/Engine/Plugins/PCG/Source/PCG/Public/PCGGraphExecutionStateInterface.h
+++ b/Engine/Plugins/PCG/Source/PCG/Public/PCGGraphExecutionStateInterface.h
@@ -227,6 +227,9 @@ public:
 
 	/** Get the size for the HiGen grid this execution source executes on. */
 	virtual uint32 GetGenerationGridSize() const { return PCGHiGenGrid::UninitializedGridSize(); }
+
+	/** Returns true if the generated artifacts are transient */
+	virtual bool IsInPreviewMode() const { return false; }
 };
 
 UINTERFACE(BlueprintType, MinimalAPI)
diff --git a/Engine/Plugins/PCG/Source/PCG/Public/Subsystems/PCGSubsystem.h b/Engine/Plugins/PCG/Source/PCG/Public/Subsystems/PCGSubsystem.h
--- a/Engine/Plugins/PCG/Source/PCG/Public/Subsystems/PCGSubsystem.h
+++ b/Engine/Plugins/PCG/Source/PCG/Public/Subsystems/PCGSubsystem.h
@@ -136,6 +136,8 @@ public:
 	PCG_API void DestroyAllPCGWorldActors();
 	PCG_API void DestroyCurrentPCGWorldActor();
 	PCG_API void LogAbnormalComponentStates(bool bGroupByState) const;
+
+	PCG_API AActor* GetOrCreateDebugActor();
 #endif
 	PCG_API void RegisterPCGWorldActor(APCGWorldActor* InActor);
 	PCG_API void UnregisterPCGWorldActor(APCGWorldActor* InActor);
@@ -425,6 +427,9 @@ private:
 	PCG::FLock PCGWorldActorLock;
 
 #if WITH_EDITOR
+	/** Debug transient actor spawned to display in world PCG debug visualization data */
+	TWeakObjectPtr<AActor> DebugActor;
+
 	int32 NumLevelUnloadingActors = 0;
 
 	using FConstructionScriptSourceComponents = TMap<FName, TObjectKey<UPCGComponent>>;
diff --git a/Engine/Plugins/PCG/Source/PCGEditor/Private/DataVisualizations/PCGSpatialDataVisualization.cpp b/Engine/Plugins/PCG/Source/PCGEditor/Private/DataVisualizations/PCGSpatialDataVisualization.cpp
--- a/Engine/Plugins/PCG/Source/PCGEditor/Private/DataVisualizations/PCGSpatialDataVisualization.cpp
+++ b/Engine/Plugins/PCG/Source/PCGEditor/Private/DataVisualizations/PCGSpatialDataVisualization.cpp
@@ -184,8 +184,6 @@ void IPCGSpatialDataVisualization::ExecuteDebugDisplayHelper(
 	Params[1] = Params[0];
 	Params[1].Descriptor.bReverseCulling = true;
 
-	UPCGComponent* SourceComponent = Cast<UPCGComponent>(Context ? Context->ExecutionSource.Get() : nullptr);
-
 	// Since the instance count is global, track the current instances applied and previously belong to the ISMCs.
 	int32 NumCurrentInstances = 0;
 	for (int32 Direction = 0; Direction < 2; ++Direction)
@@ -199,9 +197,9 @@ void IPCGSpatialDataVisualization::ExecuteDebugDisplayHelper(
 
 		UInstancedStaticMeshComponent* ISMC = nullptr;
 		
-		if (TargetActor && SourceComponent)
+		if (TargetActor && Context && Context->ExecutionSource.Get())
 		{
-			ISMC = UPCGActorHelpers::GetOrCreateISMC(TargetActor, SourceComponent, Params[Direction], Context);
+			ISMC = UPCGActorHelpers::GetOrCreateISMC(TargetActor, Context->ExecutionSource.Get(), Params[Direction], Context);
 		}
 		else
 		{
