[PCG] Fix graph change detection failing when edges made between inactive/culled nodes.

#rnx
#Synced-CL 49592575

diff --git a/Engine/Plugins/PCG/Source/PCG/Private/PCGGraph.cpp b/Engine/Plugins/PCG/Source/PCG/Private/PCGGraph.cpp
--- a/Engine/Plugins/PCG/Source/PCG/Private/PCGGraph.cpp
+++ b/Engine/Plugins/PCG/Source/PCG/Private/PCGGraph.cpp
@@ -1882,11 +1882,8 @@ void UPCGGraph::NotifyGraphStructureChanged(EPCGChangeType ChangeType, bool bFor
 
 		if (!bExecutionAffected)
 		{
-			// If compiled tasks are unchanged and the settings have not changed, then we can demote the change to avoid unnecessary graph executions.
-			// * Structural and GenerationGrid are frequent change types that trigger generation
-			// * Node because edges can change pin types on nodes.
-			// * Settings because changing a subgraph currently changes settings on subgraph node.
-			ChangeType &= ~(EPCGChangeType::Structural | EPCGChangeType::GenerationGrid | EPCGChangeType::Node | EPCGChangeType::Settings);
+			// Remove a range of change types related to graph editing which may prevent unnecessary execution. Note settings changes always recompile (checked above).
+			ChangeType &= ~(EPCGChangeType::Structural | EPCGChangeType::GenerationGrid | EPCGChangeType::Node | EPCGChangeType::Edge);
 
 			// Positively flagging as cosmetic is required because downstream things specifically test for this currently.
 			ChangeType |= EPCGChangeType::Cosmetic;
