[PCG] Add Downsample Texture node.

[RN] [PCG] Added Downsample Texture node to generate mips for textures.
#Synced-CL 49592575

diff --git a/Engine/Plugins/PCG/Shaders/Private/PCGTexureDownsample.usf b/Engine/Plugins/PCG/Shaders/Private/PCGTexureDownsample.usf
new file mode 100644
--- /dev/null
+++ b/Engine/Plugins/PCG/Shaders/Private/PCGTexureDownsample.usf
@@ -0,0 +1,99 @@
+// Copyright Epic Games, Inc. All Rights Reserved.
+
+#pragma once
+
+#include "/Engine/Private/Common.ush"
+#include "/Engine/Private/GammaCorrectionCommon.ush"
+
+#define PCG_DOWNSAMPLE_MODE_AVERAGE		0
+#define PCG_DOWNSAMPLE_MODE_MIN			1
+#define PCG_DOWNSAMPLE_MODE_MAX			2
+#define PCG_DOWNSAMPLE_MODE_SUM			3
+
+SamplerState MipSampler;
+uint2 MipOutSize;
+float2 TexelSize;
+
+#if PCG_DOWNSAMPLE_SRGB
+#define PCG_DATA_TYPE half4
+#else
+#define PCG_DATA_TYPE float4
+#endif
+
+#if PCG_DOWNSAMPLE_TEXTUREARRAY
+Texture2DArray MipArrayInSRV;
+RWTexture2DArray<PCG_DATA_TYPE> MipArrayOutUAV;
+#define PCG_TEXTURE_SAMPLE() MipArrayInSRV.SampleLevel(MipSampler, float3(UV, InThreadID.z), 0)
+#define PCG_TEXTURE_LOAD(Coords) MipArrayInSRV[uint3(Coords, InThreadID.z)]
+#define PCG_STORE_RESULT() MipArrayOutUAV[InThreadID] = OutColor
+#else
+Texture2D MipInSRV;
+RWTexture2D<PCG_DATA_TYPE> MipOutUAV;
+#define PCG_TEXTURE_SAMPLE() MipInSRV.SampleLevel(MipSampler, UV, 0)
+#define PCG_TEXTURE_LOAD(Coords) MipInSRV[Coords]
+#define PCG_STORE_RESULT() MipOutUAV[InThreadID.xy] = OutColor
+#endif
+
+[numthreads(8, 8, 1)]
+void TextureDownsampleCS(const uint3 InThreadID : SV_DispatchThreadID)
+{
+	if (any(InThreadID.xy >= MipOutSize))
+	{
+		return;
+	}
+
+#if PCG_DOWNSAMPLE_MODE == PCG_DOWNSAMPLE_MODE_AVERAGE
+
+	const float2 UV = TexelSize * (InThreadID.xy + 0.5f);
+
+    PCG_DATA_TYPE OutColor = PCG_TEXTURE_SAMPLE();
+#if PCG_DOWNSAMPLE_SRGB
+    // Sample operation already did sRGB->linear conversion, so just need to convert back.
+    OutColor = PCG_DATA_TYPE(LinearToSrgb(OutColor.xyz), OutColor.w);
+#endif
+
+#else // PCG_DOWNSAMPLE_MODE == PCG_DOWNSAMPLE_MODE_AVERAGE
+
+	PCG_DATA_TYPE InColor00 = PCG_TEXTURE_LOAD(InThreadID.xy * 2);
+	PCG_DATA_TYPE InColor01 = PCG_TEXTURE_LOAD(InThreadID.xy * 2 + uint2(0, 1));
+	PCG_DATA_TYPE InColor10 = PCG_TEXTURE_LOAD(InThreadID.xy * 2 + uint2(1, 0));
+	PCG_DATA_TYPE InColor11 = PCG_TEXTURE_LOAD(InThreadID.xy * 2 + uint2(1, 1));
+
+#if PCG_DOWNSAMPLE_MODE == PCG_DOWNSAMPLE_MODE_MIN
+
+	// Min is color space agnostic - load and store in whatever color space is in use.
+	const PCG_DATA_TYPE OutColor = min(min(InColor00, InColor01), min(InColor10, InColor11));
+
+#elif PCG_DOWNSAMPLE_MODE == PCG_DOWNSAMPLE_MODE_MAX
+
+	// Max is color space agnostic - load and store in whatever color space is in use.
+	const PCG_DATA_TYPE OutColor = max(max(InColor00, InColor01), max(InColor10, InColor11));
+
+#elif PCG_DOWNSAMPLE_MODE == PCG_DOWNSAMPLE_MODE_SUM
+
+#if PCG_DOWNSAMPLE_SRGB
+	// Convert to linear space before doing math, then convert result back.
+	InColor00.xyz = sRGBToLinear(InColor00.xyz);
+	InColor01.xyz = sRGBToLinear(InColor01.xyz);
+	InColor10.xyz = sRGBToLinear(InColor10.xyz);
+	InColor11.xyz = sRGBToLinear(InColor11.xyz);
+    const PCG_DATA_TYPE OutColor = PCG_DATA_TYPE(LinearToSrgb(InColor00.xyz + InColor01.xyz + InColor10.xyz + InColor11.xyz), InColor00.w + InColor01.w + InColor10.w + InColor11.w);
+#else // PCG_DOWNSAMPLE_SRGB
+	const PCG_DATA_TYPE OutColor = InColor00 + InColor01 + InColor10 + InColor11;
+#endif // PCG_DOWNSAMPLE_SRGB
+
+#endif // PCG_DOWNSAMPLE_MODE == PCG_DOWNSAMPLE_MODE_SUM
+
+#endif // PCG_DOWNSAMPLE_MODE == PCG_DOWNSAMPLE_MODE_AVERAGE
+	
+	PCG_STORE_RESULT();
+}
+
+#undef PCG_DOWNSAMPLE_MODE_AVERAGE
+#undef PCG_DOWNSAMPLE_MODE_MIN
+#undef PCG_DOWNSAMPLE_MODE_MAX
+#undef PCG_DOWNSAMPLE_MODE_SUM
+#undef PCG_DATA_TYPE
+#undef PCG_TEXTURE_SAMPLE
+#undef PCG_TEXTURE_LOAD
+#undef PCG_STORE_RESULT
diff --git a/Engine/Plugins/PCG/Source/PCG/Internal/Elements/PCGDownsampleTexture.h b/Engine/Plugins/PCG/Source/PCG/Internal/Elements/PCGDownsampleTexture.h
new file mode 100644
--- /dev/null
+++ b/Engine/Plugins/PCG/Source/PCG/Internal/Elements/PCGDownsampleTexture.h
@@ -0,0 +1,86 @@
+// Copyright Epic Games, Inc. All Rights Reserved.
+
+#pragma once
+
+#include "PCGContext.h"
+#include "PCGSettings.h"
+
+#include "PCGTextureDownsample.h"
+
+#include "PCGDownsampleTexture.generated.h"
+
+class UPCGTextureData;
+class UPCGTexture2DArrayData;
+class UPCGTexture2DBaseData;
+
+UENUM()
+enum class EPCGDownsampleTextureType : uint8
+{
+	Texture2D,
+	Texture2DArray,
+};
+
+/** Downsample a texture and populate its mip chain. */
+UCLASS(MinimalAPI, BlueprintType, ClassGroup = (Procedural))
+class UPCGDownsampleTextureSettings : public UPCGSettings
+{
+	GENERATED_BODY()
+
+public:
+	//~Begin UPCGSettings interface
+#if WITH_EDITOR
+	virtual FName GetDefaultNodeName() const override { return TEXT("DownsampleTexture"); }
+	virtual FText GetDefaultNodeTitle() const override;
+	virtual FText GetNodeTooltipText() const override;
+	virtual EPCGSettingsType GetType() const override { return EPCGSettingsType::Generic; }
+#endif
+	virtual bool HasDynamicPins() const override { return true; }
+
+	virtual bool PadResolutionToSquare() const { return bPadResolutionToSquare; }
+	virtual bool PadResolutionToNextPowerOfTwo() const { return bPadResolutionToNextPowerOfTwo; }
+
+protected:
+	virtual TArray<FPCGPinProperties> InputPinProperties() const override;
+	virtual TArray<FPCGPinProperties> OutputPinProperties() const override;
+	virtual FPCGElementPtr CreateElement() const override;
+	//~End UPCGSettings interface
+
+public:
+	EPCGTextureDownsampleMode GetMode() const { return Mode; }
+
+protected:
+	UPROPERTY(EditAnywhere, Category = Settings)
+	EPCGTextureDownsampleMode Mode = EPCGTextureDownsampleMode::Average;
+
+	UPROPERTY(EditAnywhere, Category = Settings, AdvancedDisplay)
+	bool bPadResolutionToSquare = false;
+
+	UPROPERTY(EditAnywhere, Category = Settings, AdvancedDisplay)
+	bool bPadResolutionToNextPowerOfTwo = false;
+};
+
+struct FPCGDownsampleTextureContext : public FPCGContext
+{
+public:
+	bool bSubmittedRenderCommands = false;
+	bool bCreatedOutputData = false;
+
+	struct FTextureOutputInfo
+	{
+		TRefCountPtr<IPooledRenderTarget> OutputHandle;
+		FTransform Transform;
+		EPCGTextureFilter Filter = EPCGTextureFilter::Bilinear;
+		int32 ArraySize = 1;
+		UPCGTexture2DBaseData* OutputData = nullptr;
+	};
+
+	TArray<FTextureOutputInfo> TextureOutputInfo;
+};
+
+class FPCGDownsampleTextureElement : public IPCGElement
+{
+protected:
+	virtual bool ExecuteInternal(FPCGContext* InContext) const override;
+	virtual bool SupportsGPUResidentData(FPCGContext* InContext) const { return true; }
+	virtual FPCGContext* CreateContext() override { return new FPCGDownsampleTextureContext(); }
+};
diff --git a/Engine/Plugins/PCG/Source/PCG/Private/Elements/PCGDownsampleTexture.cpp b/Engine/Plugins/PCG/Source/PCG/Private/Elements/PCGDownsampleTexture.cpp
new file mode 100644
--- /dev/null
+++ b/Engine/Plugins/PCG/Source/PCG/Private/Elements/PCGDownsampleTexture.cpp
@@ -0,0 +1,248 @@
+// Copyright Epic Games, Inc. All Rights Reserved.
+
+#include "Elements/PCGDownsampleTexture.h"
+
+#include "PCGContext.h"
+#include "PCGData.h"
+#include "Data/PCGTexture2DArrayData.h"
+#include "Data/PCGTextureData.h"
+
+#include "RenderCaptureInterface.h"
+#include "RenderGraphBuilder.h"
+#include "RenderGraphUtils.h"
+#include "RHIStaticStates.h"
+#include "Algo/AnyOf.h"
+
+#include UE_INLINE_GENERATED_CPP_BY_NAME(PCGDownsampleTexture)
+
+#define LOCTEXT_NAMESPACE "PCGDownsampleTextureElement"
+
+namespace PCGDownsampleTexture
+{
+#if !UE_BUILD_SHIPPING
+	static int32 TriggerGPUCaptureDispatchIndex = 0;
+	static FAutoConsoleVariableRef CVarTriggerGPUCaptureDispatchIndex(
+		TEXT("pcg.GPU.TriggerRenderCaptures.DownsampleTexture"),
+		TriggerGPUCaptureDispatchIndex,
+		TEXT("Index of the next dispatch to capture. I.e. if set to 2, will ignore one dispatch and then trigger a capture on the next one."),
+		ECVF_RenderThreadSafe
+	);
+#endif // !UE_BUILD_SHIPPING
+}
+
+#if WITH_EDITOR
+FText UPCGDownsampleTextureSettings::GetDefaultNodeTitle() const
+{
+	return LOCTEXT("NodeTitle", "Downsample Texture");
+}
+
+FText UPCGDownsampleTextureSettings::GetNodeTooltipText() const
+{
+	return LOCTEXT("NodeTooltip", "Downsample a texture and populate its mip chain.");
+}
+#endif // WITH_EDITOR
+
+TArray<FPCGPinProperties> UPCGDownsampleTextureSettings::InputPinProperties() const
+{
+	TArray<FPCGPinProperties> PinProperties;
+	PinProperties.Emplace(PCGPinConstants::DefaultInputLabel, FPCGDataTypeInfoTexture2DBase::AsId(), /*bInAllowMultipleConnections=*/true, /*bAllowMultipleData=*/true);
+
+	return PinProperties;
+}
+
+TArray<FPCGPinProperties> UPCGDownsampleTextureSettings::OutputPinProperties() const
+{
+	TArray<FPCGPinProperties> PinProperties;
+	PinProperties.Emplace(PCGPinConstants::DefaultOutputLabel, EPCGDataType::Any);
+
+	return PinProperties;
+}
+
+FPCGElementPtr UPCGDownsampleTextureSettings::CreateElement() const
+{
+	return MakeShared<FPCGDownsampleTextureElement>();
+}
+
+bool FPCGDownsampleTextureElement::ExecuteInternal(FPCGContext* InContext) const
+{
+	TRACE_CPUPROFILER_EVENT_SCOPE(FPCGDownsampleTextureElement::Execute);
+	check(InContext);
+	FPCGDownsampleTextureContext* Context = static_cast<FPCGDownsampleTextureContext*>(InContext);
+
+	const UPCGDownsampleTextureSettings* Settings = Context->GetInputSettings<UPCGDownsampleTextureSettings>();
+	check(Settings);
+
+	if (!Context->bSubmittedRenderCommands)
+	{
+		for (const FPCGTaggedData& Input : Context->InputData.GetInputsByPin(PCGPinConstants::DefaultInputLabel))
+		{
+			const UPCGTexture2DBaseData* InputTexture2D = Cast<UPCGTexture2DBaseData>(Input.Data);
+			if (!InputTexture2D)
+			{
+				continue;
+			}
+			
+			const TRefCountPtr<IPooledRenderTarget> InputTextureHandle = InputTexture2D->GetRefCountedTexture();
+			if (!InputTextureHandle)
+			{
+				continue;
+			}
+
+			const int32 OutputIndex = Context->TextureOutputInfo.Num();
+
+			FPCGDownsampleTextureContext::FTextureOutputInfo& OutputInfo = Context->TextureOutputInfo.Emplace_GetRef();
+			OutputInfo.Transform = InputTexture2D->GetTransform();
+			OutputInfo.Filter = InputTexture2D->Filter;
+			OutputInfo.ArraySize = InputTexture2D->GetArraySize();
+
+			FIntPoint Resolution = InputTexture2D->GetResolution();
+
+			if (Settings->PadResolutionToNextPowerOfTwo())
+			{
+				Resolution.X = 1u << FMath::CeilLogTwo(Resolution.X);
+				Resolution.Y = 1u << FMath::CeilLogTwo(Resolution.Y);
+			}
+
+			if (Settings->PadResolutionToSquare())
+			{
+				const int32 MaxDim = FMath::Max(Resolution.X, Resolution.Y);
+				Resolution.X = Resolution.Y = MaxDim;
+			}
+
+			FRDGTextureDesc OutputTextureDesc = FRDGTextureDesc::Create2DArray(
+				Resolution,
+				InputTexture2D->GetFormat(),
+				FClearValueBinding(),
+				ETextureCreateFlags::ShaderResource | ETextureCreateFlags::UAV,
+				InputTexture2D->GetArraySize()
+			);
+
+			OutputTextureDesc.NumMips = FMath::FloorLog2(FMath::Max(Resolution.X, Resolution.Y)) + 1;
+
+			ENQUEUE_RENDER_COMMAND(PCGDownsampleTexture)([ContextHandle=Context->GetOrCreateHandle(), InputTextureHandle, OutputTextureDesc, OutputIndex, Mode=Settings->GetMode()](FRHICommandListImmediate& RHICmdList)
+			{
+				LLM_SCOPE_BYTAG(PCG);
+
+				FPCGContext::FSharedContext<FPCGDownsampleTextureContext> SharedContext(ContextHandle);
+				FPCGDownsampleTextureContext* ContextInner = SharedContext.Get();
+				if (!ContextInner)
+				{
+					return;
+				}
+
+				FRDGBuilder GraphBuilder(RHICmdList);
+				TRefCountPtr<IPooledRenderTarget> OutputTextureExported;
+
+				{
+#if !UE_BUILD_SHIPPING
+					RenderCaptureInterface::FScopedCapture RenderCapture(PCGDownsampleTexture::TriggerGPUCaptureDispatchIndex > 0, GraphBuilder, TEXT("PCGDownsampleTexture"));
+					PCGDownsampleTexture::TriggerGPUCaptureDispatchIndex = FMath::Max(PCGDownsampleTexture::TriggerGPUCaptureDispatchIndex - 1, 0);
+#endif // !UE_BUILD_SHIPPING
+
+					FRDGTextureRef InputTexture = GraphBuilder.RegisterExternalTexture(InputTextureHandle);
+					FRDGTextureRef OutputTexture = GraphBuilder.CreateTexture(OutputTextureDesc, TEXT("PCGDownsampledTexture"));
+
+					// Populate mip 0
+					{
+						FRHICopyTextureInfo CopyInfo;
+						CopyInfo.SourceMipIndex = 0;
+						CopyInfo.DestMipIndex = 0;
+						CopyInfo.NumMips = 1;
+						CopyInfo.SourceSliceIndex = 0;
+						CopyInfo.DestSliceIndex = 0;
+						CopyInfo.NumSlices = OutputTextureDesc.ArraySize;
+
+						AddCopyTexturePass(GraphBuilder, InputTexture, OutputTexture, CopyInfo);
+					}
+
+					// Downsample passes to populate rest of mips
+					{
+						PCGTextureDownsample::FParams Params;
+						Params.Texture = OutputTexture;
+						Params.Sampler = TStaticSamplerState<SF_Bilinear, AM_Clamp, AM_Clamp, AM_Clamp>::GetRHI();
+						Params.SliceIndex = 0;
+						Params.NumSlices = OutputTextureDesc.ArraySize;
+						Params.Mode = Mode;
+
+						PCGTextureDownsample::DownsampleTexture(GraphBuilder, Params);
+					}
+
+					OutputTextureExported = GraphBuilder.ConvertToExternalTexture(OutputTexture);
+					GraphBuilder.SetTextureAccessFinal(OutputTexture, ERHIAccess::SRVCompute);
+				}
+
+				GraphBuilder.Execute();
+
+				// Pass exported buffer back to game thread and wake up this element.
+				ExecuteOnGameThread(UE_SOURCE_LOCATION, [ContextHandle, OutputTextureExported, OutputIndex]()
+				{
+					LLM_SCOPE_BYTAG(PCG);
+
+					FPCGContext::FSharedContext<FPCGDownsampleTextureContext> SharedContext(ContextHandle);
+					if (FPCGDownsampleTextureContext* Context = SharedContext.Get())
+					{
+						Context->TextureOutputInfo[OutputIndex].OutputHandle = OutputTextureExported;
+						Context->bIsPaused = false;
+					}
+				});
+			});
+		}
+
+		Context->bSubmittedRenderCommands = true;
+
+		// Render command will wake this task up after completing.
+		Context->bIsPaused = true;
+
+		return false;
+	}
+
+	if (Algo::AnyOf(Context->TextureOutputInfo, [](const FPCGDownsampleTextureContext::FTextureOutputInfo& InInfo) { return !InInfo.OutputHandle; }))
+	{
+		UE_LOG(LogPCG, Error, TEXT("Downsample operation failed - null output texture handle encountered."));
+		return true;
+	}
+
+	if (!Context->bCreatedOutputData)
+	{
+		for (FPCGDownsampleTextureContext::FTextureOutputInfo& OutputInfo : Context->TextureOutputInfo)
+		{
+			if (OutputInfo.ArraySize == 1)
+			{
+				OutputInfo.OutputData = FPCGContext::NewObject_AnyThread<UPCGTextureData>(Context);
+				OutputInfo.OutputData->Filter = OutputInfo.Filter;
+			}
+			else
+			{
+				OutputInfo.OutputData = FPCGContext::NewObject_AnyThread<UPCGTexture2DArrayData>(Context);
+			}
+
+			FPCGTaggedData& OutTaggedData = Context->OutputData.TaggedData.Emplace_GetRef();
+			OutTaggedData.Data = OutputInfo.OutputData;
+			OutTaggedData.Pin = PCGPinConstants::DefaultOutputLabel;
+		}
+
+		Context->bCreatedOutputData = true;
+	}
+
+	bool bAllInitialized = true;
+
+	for (FPCGDownsampleTextureContext::FTextureOutputInfo& OutputInfo : Context->TextureOutputInfo)
+	{
+		if (UPCGTextureData* TextureData = Cast<UPCGTextureData>(OutputInfo.OutputData))
+		{
+			bAllInitialized &= TextureData->Initialize(OutputInfo.OutputHandle, /*InTextureIndex=*/0, OutputInfo.Transform, /*bInSkipReadbackToCPU=*/true);
+		}
+		else if (UPCGTexture2DArrayData* TextureArrayData = Cast<UPCGTexture2DArrayData>(OutputInfo.OutputData))
+		{
+			FPCGTexture2DArrayDataInitParams Params;
+			Params.Transform = OutputInfo.Transform;
+			Params.Filter = OutputInfo.Filter;
+
+			bAllInitialized &= TextureArrayData->Initialize(OutputInfo.OutputHandle, Params);
+		}
+	}
+
+	return bAllInitialized;
+}
+
+#undef LOCTEXT_NAMESPACE
diff --git a/Engine/Plugins/PCG/Source/PCGCompute/Internal/PCGTextureDownsample.h b/Engine/Plugins/PCG/Source/PCGCompute/Internal/PCGTextureDownsample.h
new file mode 100644
--- /dev/null
+++ b/Engine/Plugins/PCG/Source/PCGCompute/Internal/PCGTextureDownsample.h
@@ -0,0 +1,31 @@
+// Copyright Epic Games, Inc. All Rights Reserved.
+
+#pragma once
+
+#include "RHI.h"
+#include "RendererInterface.h"
+
+#include "PCGTextureDownsample.generated.h"
+
+UENUM()
+enum class EPCGTextureDownsampleMode : uint8
+{
+	Average = 0,
+	Min = 1,
+	Max = 2,
+	Sum = 3,
+};
+
+namespace PCGTextureDownsample
+{
+	struct FParams
+	{
+		FRDGTextureRef Texture = nullptr;
+		FRHISamplerState* Sampler = nullptr;
+		int SliceIndex = 0;
+		int NumSlices = 1;
+		EPCGTextureDownsampleMode Mode = EPCGTextureDownsampleMode::Average;
+	};
+
+	PCGCOMPUTE_API void DownsampleTexture(FRDGBuilder& GraphBuilder, FParams& InParams);
+};
diff --git a/Engine/Plugins/PCG/Source/PCGCompute/PCGCompute.Build.cs b/Engine/Plugins/PCG/Source/PCGCompute/PCGCompute.Build.cs
--- a/Engine/Plugins/PCG/Source/PCGCompute/PCGCompute.Build.cs
+++ b/Engine/Plugins/PCG/Source/PCGCompute/PCGCompute.Build.cs
@@ -12,6 +12,7 @@ namespace UnrealBuildTool.Rules
 				new string[]
 				{
 					"Core",
+					"CoreUObject",
 					"Projects",
 					"RHI",
 					"RenderCore",
diff --git a/Engine/Plugins/PCG/Source/PCGCompute/Private/PCGTextureDownsample.cpp b/Engine/Plugins/PCG/Source/PCGCompute/Private/PCGTextureDownsample.cpp
new file mode 100644
--- /dev/null
+++ b/Engine/Plugins/PCG/Source/PCGCompute/Private/PCGTextureDownsample.cpp
@@ -0,0 +1,96 @@
+// Copyright Epic Games, Inc. All Rights Reserved.
+
+#include "PCGTextureDownsample.h"
+
+#include "PCGComputeModule.h"
+
+#include "GlobalShader.h"
+#include "RHIGPUReadback.h"
+#include "RenderGraphUtils.h"
+#include "Async/Async.h"
+
+#define LOCTEXT_NAMESPACE "PCGCompute"
+
+class FPCGTextureDownsampleCS : public FGlobalShader
+{
+public:
+	DECLARE_GLOBAL_SHADER(FPCGTextureDownsampleCS)
+	SHADER_USE_PARAMETER_STRUCT(FPCGTextureDownsampleCS, FGlobalShader)
+
+	class FDownsampleTextureArray : SHADER_PERMUTATION_BOOL("PCG_DOWNSAMPLE_TEXTUREARRAY");
+	class FDownsampleSRGB : SHADER_PERMUTATION_BOOL("PCG_DOWNSAMPLE_SRGB");
+	class FDownsampleMode : SHADER_PERMUTATION_INT("PCG_DOWNSAMPLE_MODE", 4);
+	using FPermutationDomain = TShaderPermutationDomain<FDownsampleTextureArray, FDownsampleSRGB, FDownsampleMode>;
+
+	BEGIN_SHADER_PARAMETER_STRUCT(FParameters, )
+		SHADER_PARAMETER(FUintVector2, MipOutSize)
+		SHADER_PARAMETER(FVector2f, TexelSize)
+		SHADER_PARAMETER_SAMPLER(SamplerState, MipSampler)
+		SHADER_PARAMETER_RDG_TEXTURE_SRV(Texture2D, MipInSRV)
+		SHADER_PARAMETER_RDG_TEXTURE_UAV(RWTexture2D, MipOutUAV)
+		SHADER_PARAMETER_RDG_TEXTURE_SRV(Texture2DArray, MipArrayInSRV)
+		SHADER_PARAMETER_RDG_TEXTURE_UAV(RWTexture2DArray, MipArrayOutUAV)
+	END_SHADER_PARAMETER_STRUCT()
+};
+
+IMPLEMENT_GLOBAL_SHADER(FPCGTextureDownsampleCS, "/PCGComputeShaders/PCGTexureDownsample.usf", "TextureDownsampleCS", SF_Compute);
+
+namespace PCGTextureDownsample
+{
+	void DownsampleTexture(FRDGBuilder& GraphBuilder, FParams& InParams)
+	{
+		FRDGTextureRef Texture = InParams.Texture;
+		FRHISamplerState* Sampler = InParams.Sampler;
+		check(Texture);
+		check(Sampler);
+
+		FGlobalShaderMap* ShaderMap = GetGlobalShaderMap(GMaxRHIFeatureLevel);
+
+		const FRDGTextureDesc& TextureDesc = Texture->Desc;
+
+		// Select compute shader variant (normal vs. sRGB etc.)
+		const bool bMipsSRGB = EnumHasAnyFlags(TextureDesc.Flags, TexCreate_SRGB);
+		FPCGTextureDownsampleCS::FPermutationDomain PermutationVector;
+		PermutationVector.Set<FPCGTextureDownsampleCS::FDownsampleTextureArray>(InParams.NumSlices > 1);
+		PermutationVector.Set<FPCGTextureDownsampleCS::FDownsampleSRGB>(bMipsSRGB);
+		PermutationVector.Set<FPCGTextureDownsampleCS::FDownsampleMode>(static_cast<int>(InParams.Mode));
+		TShaderMapRef<FPCGTextureDownsampleCS> ComputeShader(ShaderMap, PermutationVector);
+
+		// Loop through each level of the mips that require creation and add a dispatch pass per level.
+		for (uint8 MipLevel = 1, MipCount = TextureDesc.NumMips; MipLevel < MipCount; ++MipLevel)
+		{
+			const FUintVector2 DestTextureSize(
+				FMath::Max(TextureDesc.Extent.X >> MipLevel, 1),
+				FMath::Max(TextureDesc.Extent.Y >> MipLevel, 1));
+
+			FPCGTextureDownsampleCS::FParameters* PassParameters = GraphBuilder.AllocParameters<FPCGTextureDownsampleCS::FParameters>();
+			PassParameters->MipOutSize = DestTextureSize;
+			PassParameters->TexelSize = FVector2f(1.0f / DestTextureSize.X, 1.0f / DestTextureSize.Y);
+			PassParameters->MipSampler = Sampler;
+
+			if (InParams.NumSlices > 1)
+			{
+				PassParameters->MipArrayInSRV = GraphBuilder.CreateSRV(FRDGTextureSRVDesc::CreateForMipLevel(Texture, MipLevel - 1));
+				PassParameters->MipArrayOutUAV = GraphBuilder.CreateUAV(FRDGTextureUAVDesc(Texture, MipLevel));
+				PassParameters->MipInSRV = nullptr;
+				PassParameters->MipOutUAV = nullptr;
+			}
+			else
+			{
+				PassParameters->MipInSRV = GraphBuilder.CreateSRV(FRDGTextureSRVDesc::CreateForMipLevel(Texture, MipLevel - 1));
+				PassParameters->MipOutUAV = GraphBuilder.CreateUAV(FRDGTextureUAVDesc(Texture, MipLevel));
+				PassParameters->MipArrayInSRV = nullptr;
+				PassParameters->MipArrayOutUAV = nullptr;
+			}
+
+			FComputeShaderUtils::AddPass(
+				GraphBuilder,
+				RDG_EVENT_NAME("GenerateMips DestMipLevel=%d", MipLevel),
+				ComputeShader,
+				PassParameters,
+				FIntVector(FMath::DivideAndRoundUp<int>(DestTextureSize.X, 8), FMath::DivideAndRoundUp<int>(DestTextureSize.Y, 8), InParams.NumSlices));
+		}
+	}
+}
+
+#undef LOCTEXT_NAMESPACE
