[PCG] Fix 'Treat Editor as Viewport' runtime gen getting duplicated in PIE
- Also avoid package override for runtime generated actors

#Synced-CL 49621001

diff --git a/Engine/Plugins/PCG/Source/PCG/Private/Elements/PCGCreateTargetActor.cpp b/Engine/Plugins/PCG/Source/PCG/Private/Elements/PCGCreateTargetActor.cpp
--- a/Engine/Plugins/PCG/Source/PCG/Private/Elements/PCGCreateTargetActor.cpp
+++ b/Engine/Plugins/PCG/Source/PCG/Private/Elements/PCGCreateTargetActor.cpp
@@ -392,7 +392,7 @@ bool FPCGCreateTargetActorElement::ExecuteInternal(FPCGContext* Context) const
 	}
 
 	UPCGActorHelpers::FSpawnDefaultActorParams SpawnDefaultActorParams(TargetActor->GetWorld(), Settings->TemplateActorClass, Transform, SpawnParams);
-	
+	SpawnDefaultActorParams.bIsRuntime = (SourceComponent && SourceComponent->IsManagedByRuntimeGenSystem());
 #if WITH_EDITOR
 	SpawnDefaultActorParams.bIsPreviewActor = (SourceComponent && SourceComponent->IsInPreviewMode());
 
diff --git a/Engine/Plugins/PCG/Source/PCG/Private/Elements/PCGSpawnActor.cpp b/Engine/Plugins/PCG/Source/PCG/Private/Elements/PCGSpawnActor.cpp
--- a/Engine/Plugins/PCG/Source/PCG/Private/Elements/PCGSpawnActor.cpp
+++ b/Engine/Plugins/PCG/Source/PCG/Private/Elements/PCGSpawnActor.cpp
@@ -864,7 +864,7 @@ void FPCGSpawnActorElement::SpawnActors(FPCGSubgraphContext* Context, AActor* Ta
 
 	UPCGActorHelpers::FSpawnDefaultActorParams SpawnDefaultActorParams(TargetActor->GetWorld(), InTemplateActorClass, FTransform::Identity, SpawnParams);
 	SpawnDefaultActorParams.bForceStaticMobility = false; // Always respect the actor's mobility
-
+	SpawnDefaultActorParams.bIsRuntime = (SourceComponent && SourceComponent->IsManagedByRuntimeGenSystem());
 #if WITH_EDITOR
 	SpawnDefaultActorParams.bIsPreviewActor = (SourceComponent && SourceComponent->IsInPreviewMode());
 
diff --git a/Engine/Plugins/PCG/Source/PCG/Private/Helpers/PCGActorHelpers.cpp b/Engine/Plugins/PCG/Source/PCG/Private/Helpers/PCGActorHelpers.cpp
--- a/Engine/Plugins/PCG/Source/PCG/Private/Helpers/PCGActorHelpers.cpp
+++ b/Engine/Plugins/PCG/Source/PCG/Private/Helpers/PCGActorHelpers.cpp
@@ -707,16 +707,22 @@ AActor* UPCGActorHelpers::SpawnDefaultActor(const FSpawnDefaultActorParams& Para
 
 	if (PCGHelpers::IsRuntimeOrPIE() || Params.bIsPreviewActor)
 	{
-		SpawnParams.ObjectFlags |= RF_Transient | RF_NonPIEDuplicateTransient;
+		SpawnParams.ObjectFlags |= RF_Transient;
+
+		// Only allow actors to be duplicated in PIE if they aren't generated at runtime (Treat editor as viewport code path which runs runtime gen in an editor world)
+		if (!Params.bIsRuntime)
+		{
+			SpawnParams.ObjectFlags |= RF_NonPIEDuplicateTransient;
+		}
 	}
 
 #if WITH_EDITOR
 	// Capture label before potentially changing it for previewing
 	FString Label = SpawnParams.Name.IsNone() ? FString() : SpawnParams.Name.ToString();
 
-	// If we are spawning a preview actor in a world partition world, we need to assign it an external package as Data Layers are not supported on non OFPA actors.
+	// If we are spawning a non-runtime preview actor in a world partition world, we need to assign it an external package as Data Layers are not supported on non OFPA actors.
 	// Note: this could at somepoint become an option in the FActorSpawnParameters to allow RF_Transient actors to create their own package in UWorld::SpawnActor.
-	if (Params.bIsPreviewActor && SpawnParams.OverrideLevel->GetWorldPartition() && !Params.World->IsGameWorld())
+	if (Params.bIsPreviewActor && !Params.bIsRuntime && SpawnParams.OverrideLevel->GetWorldPartition() && !Params.World->IsGameWorld())
 	{
 		// Remove trailing _C from BP generated class
 		FString ActorClassName = Params.ActorClass->GetName();
@@ -742,6 +748,9 @@ AActor* UPCGActorHelpers::SpawnDefaultActor(const FSpawnDefaultActorParams& Para
 
 	SpawnParams.InitialActorLabel = Label;
 
+	// Avoid creating an actor package for runtime actors
+	SpawnParams.bCreateActorPackage &= !Params.bIsRuntime;
+
 	// Find External Data Layer if it exists so we can create scope that will allow actor to be properly created
 	const UExternalDataLayerInstance* ExternalDataLayerInstance = nullptr;
 	bool bPushedContext = false;
diff --git a/Engine/Plugins/PCG/Source/PCG/Public/Helpers/PCGActorHelpers.h b/Engine/Plugins/PCG/Source/PCG/Public/Helpers/PCGActorHelpers.h
--- a/Engine/Plugins/PCG/Source/PCG/Public/Helpers/PCGActorHelpers.h
+++ b/Engine/Plugins/PCG/Source/PCG/Public/Helpers/PCGActorHelpers.h
@@ -231,6 +231,7 @@ PRAGMA_ENABLE_DEPRECATION_WARNINGS
 		AActor* Parent = nullptr;
 		bool bForceStaticMobility = true;
 		bool bIsPreviewActor = false;
+		bool bIsRuntime = false;
 #if WITH_EDITOR
 		UHLODLayer* HLODLayer = nullptr;
 		TArray<const UDataLayerInstance*> DataLayerInstances;
