[PCG][BugHawk] Check execution source is valid and freshen code slightly.

#rnx
#Synced-CL 49622738

diff --git a/Engine/Plugins/PCG/Source/PCG/Private/Elements/PCGTextureSampler.cpp b/Engine/Plugins/PCG/Source/PCG/Private/Elements/PCGTextureSampler.cpp
--- a/Engine/Plugins/PCG/Source/PCG/Private/Elements/PCGTextureSampler.cpp
+++ b/Engine/Plugins/PCG/Source/PCG/Private/Elements/PCGTextureSampler.cpp
@@ -147,15 +147,20 @@ bool FPCGTextureSamplerElement::PrepareDataInternal(FPCGContext* InContext) cons
 bool FPCGTextureSamplerElement::ExecuteInternal(FPCGContext* InContext) const
 {
 	TRACE_CPUPROFILER_EVENT_SCOPE(FPCGTextureSamplerElement::Execute);
+	check(InContext);
 
-	FPCGTextureSamplerContext* Context = static_cast<FPCGTextureSamplerContext*>(InContext);
-	check(Context);
-
-	if (Context->bIsPaused)
+	if (InContext->bIsPaused)
 	{
 		return false;
 	}
 
+	if (!InContext->ExecutionSource.IsValid())
+	{
+		return true;
+	}
+
+	FPCGTextureSamplerContext* Context = static_cast<FPCGTextureSamplerContext*>(InContext);
+
 	const UPCGTextureSamplerSettings* Settings = Context->GetInputSettings<UPCGTextureSamplerSettings>();
 	check(Settings);
 
