[PCG] Add ability to sample from a specific mip to the texture data interface.

#rnx
#Synced-CL 49645631

diff --git a/Engine/Plugins/PCG/Shaders/Private/PCGTextureDataInterface.ush b/Engine/Plugins/PCG/Shaders/Private/PCGTextureDataInterface.ush
--- a/Engine/Plugins/PCG/Shaders/Private/PCGTextureDataInterface.ush
+++ b/Engine/Plugins/PCG/Shaders/Private/PCGTextureDataInterface.ush
@@ -292,11 +292,11 @@ float4 SampleWorldPos_{DataInterfaceName}(uint InDataIndex, float3 WorldPos)
 	{ \
 		if (TextureDimension == PCG_TextureDimension_Texture2D) \
 		{ \
-			return {DataInterfaceName}_Texture_##Index[InElementIndex]; \
+			return {DataInterfaceName}_Texture_##Index.Load(uint3(InElementIndex, InMipIndex)); \
 		} \
 		else if (TextureDimension == PCG_TextureDimension_Texture2DArray) \
 		{ \
-			return {DataInterfaceName}_TextureArray_##Index[uint3(InElementIndex, SliceIndex)]; \
+			return {DataInterfaceName}_TextureArray_##Index.Load(uint4(InElementIndex, SliceIndex, InMipIndex)); \
 		} \
 		else \
 		{ \
@@ -305,7 +305,7 @@ float4 SampleWorldPos_{DataInterfaceName}(uint InDataIndex, float3 WorldPos)
 	} \
 	else
 
-float4 Load_{DataInterfaceName}(uint InDataIndex, uint2 InElementIndex)
+float4 Load_{DataInterfaceName}(uint InDataIndex, uint2 InElementIndex, uint InMipIndex)
 {
 	uint BindingIndex = 0;
 	uint SliceIndex = 0;
@@ -329,11 +329,11 @@ float4 Load_{DataInterfaceName}(uint InDataIndex, uint2 InElementIndex)
 #else
 	if (TextureDimension == PCG_TextureDimension_Texture2D)
 	{
-		return {DataInterfaceName}_Texture_0[InElementIndex];
+		return {DataInterfaceName}_Texture_0.Load(uint3(InElementIndex, InMipIndex));
 	}
 	else if (TextureDimension == PCG_TextureDimension_Texture2DArray)
 	{
-		return {DataInterfaceName}_TextureArray_0[uint3(InElementIndex, SliceIndex)];
+		return {DataInterfaceName}_TextureArray_0.Load(uint4(InElementIndex, SliceIndex, InMipIndex));
 	}
 	else
 	{
@@ -342,6 +342,11 @@ float4 Load_{DataInterfaceName}(uint InDataIndex, uint2 InElementIndex)
 #endif
 }
 
+float4 Load_{DataInterfaceName}(uint InDataIndex, uint2 InElementIndex)
+{
+	return Load_{DataInterfaceName}(InDataIndex, InElementIndex, 0u);
+}
+
 #undef PCG_CONDITIONAL_TEXTURE_LOAD
 
 #define PCG_CONDITIONAL_TEXTURE_STORE(Index) \
diff --git a/Engine/Plugins/PCG/Source/PCG/Private/Compute/DataInterfaces/PCGTextureDataInterface.cpp b/Engine/Plugins/PCG/Source/PCG/Private/Compute/DataInterfaces/PCGTextureDataInterface.cpp
--- a/Engine/Plugins/PCG/Source/PCG/Private/Compute/DataInterfaces/PCGTextureDataInterface.cpp
+++ b/Engine/Plugins/PCG/Source/PCG/Private/Compute/DataInterfaces/PCGTextureDataInterface.cpp
@@ -102,6 +102,13 @@ void UPCGTextureDataInterface::GetSupportedInputs(TArray<FShaderFunctionDefiniti
 		.AddParam(EShaderFundamentalType::Uint) // InDataIndex
 		.AddParam(EShaderFundamentalType::Uint, 2); // InElementIndex
 
+	OutFunctions.AddDefaulted_GetRef()
+		.SetName(TEXT("Load"))
+		.AddReturnType(EShaderFundamentalType::Float, 4)
+		.AddParam(EShaderFundamentalType::Uint) // InDataIndex
+		.AddParam(EShaderFundamentalType::Uint, 2) // InElementIndex
+		.AddParam(EShaderFundamentalType::Uint); // InMipIndex
+
 	OutFunctions.AddDefaulted_GetRef()
 		.SetName(TEXT("GetTransform"))
 		.AddReturnType(EShaderFundamentalType::Float, 4, 4)
diff --git a/Engine/Plugins/PCG/Source/PCG/Private/Compute/Elements/PCGCustomHLSL.cpp b/Engine/Plugins/PCG/Source/PCG/Private/Compute/Elements/PCGCustomHLSL.cpp
--- a/Engine/Plugins/PCG/Source/PCG/Private/Compute/Elements/PCGCustomHLSL.cpp
+++ b/Engine/Plugins/PCG/Source/PCG/Private/Compute/Elements/PCGCustomHLSL.cpp
@@ -625,6 +625,7 @@ void UPCGCustomHLSLSettings::UpdateInputDeclarations()
 			"// Computes sample coordinates of the WorldPos relative to the texture data's transform.\n"
 			"float4 {0}_SampleWorldPos(uint DataIndex, float3 WorldPos);\n"
 			"float4 {0}_Load(uint DataIndex, uint2 ElementIndex);\n"
+			"float4 {0}_Load(uint DataIndex, uint2 ElementIndex, uint MipIndex);\n"
 			"float2 {0}_GetTexelSize(uint DataIndex);\n"
 			"float2 {0}_GetTexelSizeWorld(uint DataIndex);\n"
 			"float4x4 {0}_GetTransform(uint DataIndex);\n"
