[PCG] Add Developer Editor Settings to force Viewport generation mode (Always, Never, default=PCGWorldActor)
- Also add possibility to change this setting through Show Flags viewport menu

#Synced-CL 49625772

diff --git a/Engine/Plugins/PCG/Source/PCG/Private/Editor/IPCGEditorModule.cpp b/Engine/Plugins/PCG/Source/PCG/Private/Editor/IPCGEditorModule.cpp
--- a/Engine/Plugins/PCG/Source/PCG/Private/Editor/IPCGEditorModule.cpp
+++ b/Engine/Plugins/PCG/Source/PCG/Private/Editor/IPCGEditorModule.cpp
@@ -13,4 +13,16 @@ void IPCGEditorModule::SetEditorModule(IPCGEditorModule* InModule)
 {
 	check(PCGEditorModulePtr == nullptr || InModule == nullptr);
 	PCGEditorModulePtr = InModule;
-}
\ No newline at end of file
+}
+
+#if WITH_EDITOR
+bool IPCGEditorModule::ShouldTreatViewportAsGenerationSource(const APCGWorldActor* InPCGWorldActor)
+{
+	if (IPCGEditorModule* PCGEditorModule = IPCGEditorModule::Get())
+	{
+		return PCGEditorModule->ShouldTreatViewportAsGenerationSourceInternal(InPCGWorldActor);
+	}
+
+	return false;
+}
+#endif
\ No newline at end of file
diff --git a/Engine/Plugins/PCG/Source/PCG/Private/RuntimeGen/PCGGenSourceManager.cpp b/Engine/Plugins/PCG/Source/PCG/Private/RuntimeGen/PCGGenSourceManager.cpp
--- a/Engine/Plugins/PCG/Source/PCG/Private/RuntimeGen/PCGGenSourceManager.cpp
+++ b/Engine/Plugins/PCG/Source/PCG/Private/RuntimeGen/PCGGenSourceManager.cpp
@@ -19,6 +19,8 @@
 #include "WorldPartition/WorldPartition.h"
 
 #if WITH_EDITOR
+#include "Editor/IPCGEditorModule.h"
+
 #include "EditorViewportClient.h"
 #endif
 
@@ -250,7 +252,7 @@ void FPCGGenSourceManager::UpdatePerTickGenSources(const APCGWorldActor* InPCGWo
 	EditorCameraGenSource->EditorViewportClient = nullptr;
 
 	// Update the active editor viewport client for the EditorCameraGenSource, only if requested by the world actor, in-editor, and the viewport is visible.
-	if (InPCGWorldActor->bTreatEditorViewportAsGenerationSource && !Context.World->IsGameWorld() && GEditor)
+	if (IPCGEditorModule::ShouldTreatViewportAsGenerationSource(InPCGWorldActor) && !Context.World->IsGameWorld() && GEditor)
 	{
 		if (FViewport* Viewport = GEditor->GetActiveViewport())
 		{
diff --git a/Engine/Plugins/PCG/Source/PCG/Private/RuntimeGen/PCGRuntimeGenScheduler.cpp b/Engine/Plugins/PCG/Source/PCG/Private/RuntimeGen/PCGRuntimeGenScheduler.cpp
--- a/Engine/Plugins/PCG/Source/PCG/Private/RuntimeGen/PCGRuntimeGenScheduler.cpp
+++ b/Engine/Plugins/PCG/Source/PCG/Private/RuntimeGen/PCGRuntimeGenScheduler.cpp
@@ -32,6 +32,8 @@
 #include "WorldPartition/WorldPartitionSubsystem.h"
 
 #if WITH_EDITOR
+#include "Editor/IPCGEditorModule.h"
+
 #include "EditorViewportClient.h"
 #endif
 
@@ -1397,7 +1399,7 @@ void FPCGRuntimeGenScheduler::TickCVars(const APCGWorldActor* InPCGWorldActor)
 	}
 
 #if WITH_EDITOR
-	const bool bTreatEditorViewportAsGenSource = InPCGWorldActor->bTreatEditorViewportAsGenerationSource;
+	const bool bTreatEditorViewportAsGenSource = IPCGEditorModule::ShouldTreatViewportAsGenerationSource(InPCGWorldActor);
 	
 	if (bTreatEditorViewportAsGenSource != bTreatEditorViewportAsGenSourcePreviousFrame && Context.World && Context.World->IsEditorWorld())
 	{
diff --git a/Engine/Plugins/PCG/Source/PCG/Public/Editor/IPCGEditorModule.h b/Engine/Plugins/PCG/Source/PCG/Public/Editor/IPCGEditorModule.h
--- a/Engine/Plugins/PCG/Source/PCG/Public/Editor/IPCGEditorModule.h
+++ b/Engine/Plugins/PCG/Source/PCG/Public/Editor/IPCGEditorModule.h
@@ -11,6 +11,7 @@
 
 class IPCGBaseSubsystem;
 struct FPCGStackContext;
+class APCGWorldActor;
 class UPCGComponent;
 
 enum class EPCGEditorPermissionMode : uint8
@@ -38,6 +39,8 @@ public:
 
 	DECLARE_EVENT_OneParam(IPCGEditorModule, FPermissionModeChanged, EPCGEditorPermissionMode);
 	virtual FPermissionModeChanged& OnPermissionModeChanged() = 0;
+
+	static bool ShouldTreatViewportAsGenerationSource(const APCGWorldActor* InPCGWorldActor);
 #endif
 
 protected:
@@ -64,5 +67,9 @@ protected:
 
 	virtual void NotifyGraphChanged(UPCGGraph* InGraph, EPCGChangeType ChangeType) = 0;
 
+#if WITH_EDITOR
+	virtual bool ShouldTreatViewportAsGenerationSourceInternal(const APCGWorldActor* InPCGWorldActor) const = 0;
+#endif
+
 	PCG_API static void SetEditorModule(IPCGEditorModule* InModule);
 };
\ No newline at end of file
diff --git a/Engine/Plugins/PCG/Source/PCGEditor/Private/PCGEditorModule.cpp b/Engine/Plugins/PCG/Source/PCGEditor/Private/PCGEditorModule.cpp
--- a/Engine/Plugins/PCG/Source/PCGEditor/Private/PCGEditorModule.cpp
+++ b/Engine/Plugins/PCG/Source/PCGEditor/Private/PCGEditorModule.cpp
@@ -35,6 +35,7 @@
 #include "PCGEditorSettings.h"
 #include "PCGEditorStyle.h"
 #include "PCGEditorUtils.h"
+#include "PCGWorldActor.h"
 #include "ChangeTracking/PCGActorTracker.h"
 #include "ChangeTracking/PCGComponentChangeHandler.h"
 #include "ChangeTracking/PCGLandscapeTracker.h"
@@ -76,7 +77,10 @@
 #include "LevelEditorMenuContext.h"
 #include "PropertyEditorModule.h"
 #include "ScopedTransaction.h"
+#include "SEnumCombo.h"
+#include "SLevelViewport.h"
 #include "ToolMenus.h"
+#include "ToolMenuEntry.h"
 #include "Data/DataView/PCGDataViewData.h"
 #include "DataVisualizations/PCGDataViewVisualization.h"
 #include "Details/PCGRaycastFilterRuleCollectionDetails.h"
@@ -194,6 +198,80 @@ void FPCGEditorModule::ShutdownModule()
 	FEditorBuildUtils::UnregisterCustomBuildType(PCGEditorModule::PCGBuildType);
 }
 
+void FPCGEditorModule::ExtendShowFlagsMenu()
+{
+	check(GEditor);
+
+	UToolMenu* const Menu = UToolMenus::Get()->ExtendMenu("LevelEditor.ViewportToolbar.Show");
+	if (!Menu)
+	{
+		return;
+	}
+	
+	FToolMenuSection& AllShowFlagsSection = Menu->FindOrAddSection("AllShowFlags");
+
+	AllShowFlagsSection.AddEntry(FToolMenuEntry::InitSubMenu(
+		"ShowFlagsPCG",
+		LOCTEXT("ShowFlagsPCGMenu", "PCG"),
+		LOCTEXT("ShowFlagsPCGMenu_ToolTip", "PCG Specific show flags"),
+		FNewToolMenuDelegate::CreateLambda(
+			[](UToolMenu* Submenu)
+			{
+				TSharedPtr<::SLevelViewport> Viewport = ULevelViewportContext::GetLevelViewport(Submenu);
+				if (!Viewport)
+				{
+					return;
+				}
+
+				if (!Viewport->GetWorld())
+				{
+					return;
+				}
+
+				{
+					FToolMenuSection& Section = Submenu->AddSection("ShowFlagsRuntimePCG", LOCTEXT("ShowFlagsRuntimePCGLabel", "Runtime"));
+
+					auto CreateTreatViewportAsGenerationSourceWidget = []()
+					{
+						return SNew(SEnumComboBox, StaticEnum<EPCGTreatViewportAsGenerationSource>())
+								.CurrentValue_Static(&FPCGEditorModule::GetTreatEditorViewportAsGenerationSourceValue)
+								.OnEnumSelectionChanged_Static(&FPCGEditorModule::OnTreatEditorViewportAsGenerationSourceValueChanged);
+					};
+
+					FToolMenuEntry TreatEditorViewportAsGenerationSourceMenuEntry = FToolMenuEntry::InitWidget(
+						"TreatEditorViewportAsGenerationSource",
+						CreateTreatViewportAsGenerationSourceWidget(),
+						LOCTEXT("TreatEditorViewportAsGenerationSourceLabel", "Treat Editor Viewport as Generation Source")
+					);
+					Section.AddEntry(TreatEditorViewportAsGenerationSourceMenuEntry);
+				}
+			}
+		),
+		false,
+		FSlateIcon(FPCGEditorStyle::Get().GetStyleSetName(), "PCG.EditorIcon")
+	));
+}
+
+int32 FPCGEditorModule::GetTreatEditorViewportAsGenerationSourceValue()
+{
+	EPCGTreatViewportAsGenerationSource Value = EPCGTreatViewportAsGenerationSource::PCGWorldActor;
+	if (const UPCGEditorSettings* EditorSettings = GetDefault<UPCGEditorSettings>())
+	{
+		Value = EditorSettings->TreatViewportAsGenerationSource;
+	}
+	return (int32)Value;
+}
+
+void FPCGEditorModule::OnTreatEditorViewportAsGenerationSourceValueChanged(int32 Value, ESelectInfo::Type SelectInfo)
+{
+	EPCGTreatViewportAsGenerationSource NewValue = static_cast<EPCGTreatViewportAsGenerationSource>(Value);
+	if (UPCGEditorSettings* EditorSettings = GetMutableDefault<UPCGEditorSettings>())
+	{
+		EditorSettings->TreatViewportAsGenerationSource = NewValue;
+		EditorSettings->SaveConfig();
+	}
+}
+
 void FPCGEditorModule::OnLevelEditorCreated(TSharedPtr<ILevelEditor> InLevelEditor)
 {
 	if (GEditor)
@@ -208,6 +286,8 @@ void FPCGEditorModule::OnLevelEditorCreated(TSharedPtr<ILevelEditor> InLevelEdit
 		}
 
 		GEditor->OnSceneMaterialsModifiedEvent().AddRaw(this, &FPCGEditorModule::OnSceneMaterialsModified);
+
+		ExtendShowFlagsMenu();
 	}
 }
 
@@ -224,6 +304,30 @@ bool FPCGEditorModule::ShouldDisableCPUThrottling()
 	return false;
 }
 
+bool FPCGEditorModule::ShouldTreatViewportAsGenerationSourceInternal(const APCGWorldActor* InPCGWorldActor) const
+{
+	EPCGTreatViewportAsGenerationSource Value = EPCGTreatViewportAsGenerationSource::PCGWorldActor;
+
+	if (const UPCGEditorSettings* EditorSettings = GetDefault<UPCGEditorSettings>())
+	{
+		Value = EditorSettings->TreatViewportAsGenerationSource;
+	}
+
+	switch (Value)
+	{
+		case EPCGTreatViewportAsGenerationSource::PCGWorldActor:
+			return InPCGWorldActor && InPCGWorldActor->bTreatEditorViewportAsGenerationSource;
+		case EPCGTreatViewportAsGenerationSource::Always:
+			return true;
+		case EPCGTreatViewportAsGenerationSource::Never:
+			return false;
+		default:
+			check(false);
+	}
+
+	return false;
+}
+
 void FPCGEditorModule::OnPreExit()
 {
 	UnregisterPinColorAndIcons();
diff --git a/Engine/Plugins/PCG/Source/PCGEditor/Public/PCGEditorModule.h b/Engine/Plugins/PCG/Source/PCGEditor/Public/PCGEditorModule.h
--- a/Engine/Plugins/PCG/Source/PCGEditor/Public/PCGEditorModule.h
+++ b/Engine/Plugins/PCG/Source/PCGEditor/Public/PCGEditorModule.h
@@ -74,6 +74,10 @@ protected:
 
 	void OnPreExit();
 
+	void ExtendShowFlagsMenu();
+	static int32 GetTreatEditorViewportAsGenerationSourceValue();
+	static void OnTreatEditorViewportAsGenerationSourceValueChanged(int32 Value, ESelectInfo::Type SelectInfo);
+
 	/** [EXPERIMENTAL] Used to refresh procedural instances when materials are modified which can otherwise be lost.
 	* Note: This function subject to change/removal without deprecation.
 	*/
@@ -81,6 +85,8 @@ protected:
 
 	bool ShouldDisableCPUThrottling();
 
+	virtual bool ShouldTreatViewportAsGenerationSourceInternal(const APCGWorldActor* InPCGWorldActor) const override;
+
 	TArray<TSharedRef<IAssetTypeActions>> RegisteredAssetTypeActions;
 	static EAssetTypeCategories::Type PCGAssetCategory;
 
diff --git a/Engine/Plugins/PCG/Source/PCGEditor/Public/PCGEditorSettings.h b/Engine/Plugins/PCG/Source/PCGEditor/Public/PCGEditorSettings.h
--- a/Engine/Plugins/PCG/Source/PCGEditor/Public/PCGEditorSettings.h
+++ b/Engine/Plugins/PCG/Source/PCGEditor/Public/PCGEditorSettings.h
@@ -22,6 +22,14 @@ enum class EPCGEditorDoubleClickAction : uint8
 	JumpToSourceFile UMETA(DisplayName = "View Source File (C++)")
 };
 
+UENUM()
+enum class EPCGTreatViewportAsGenerationSource : uint8
+{
+	PCGWorldActor UMETA(DisplayName = "Use PCG World Actor flag"),
+	Always,
+	Never,
+};
+
 UCLASS(MinimalAPI, config=Editor, meta=(DisplayName="PCG"))
 class UPCGEditorProjectSettings : public UDeveloperSettings
 {
@@ -269,6 +277,10 @@ public:
 	UPROPERTY(EditAnywhere, config, Category = "Workflow")
 	bool bCanSelectPartitionActors = true;
 
+	/** Controls if we want to treat the viewport as a runtime generation source */
+	UPROPERTY(EditAnywhere, config, Category = "Editor Performance")
+	EPCGTreatViewportAsGenerationSource TreatViewportAsGenerationSource = EPCGTreatViewportAsGenerationSource::PCGWorldActor;
+
 	/** Specify if we want to disable CPU Throttling when a PCG Graph is executing, this will improve execution time when app is out of focus/minimized */
 	UPROPERTY(EditAnywhere, config, Category = "Editor Performance", meta = (DisplayName = "Disable CPU throttling during graph execution"))
 	bool bDisableCPUThrottlingDuringGraphExecution = false;
