Fix issue where PCG sampling of RVT layers greater than 0 would fail.
Issue was that the RVT texture layer was being interpreted as the page table texture index.
RVTs all have a single page table texture which covers up to 4 RVT layers (the maximum used is 3). So the page table texture index is always 0.
Also changed page table sampling so that the single shared page table texture is only bound and sampled once instead of once per layer.
Fixed the channel selection when reading from the page table with TextureVirtualSample(). Index is always 0 when the RVT has a Single Physical Space, otherwise it is the index of the sample layer.
#rnx
#Synced-CL 49645631

diff --git a/Engine/Plugins/PCG/Shaders/Private/PCGLandscapeDataInterface.ush b/Engine/Plugins/PCG/Shaders/Private/PCGLandscapeDataInterface.ush
--- a/Engine/Plugins/PCG/Shaders/Private/PCGLandscapeDataInterface.ush
+++ b/Engine/Plugins/PCG/Shaders/Private/PCGLandscapeDataInterface.ush
@@ -21,6 +21,7 @@ Texture2D			{DataInterfaceName}_NormalVirtualTexture1;
 Texture2D<uint4>	{DataInterfaceName}_NormalVirtualTexturePageTable;
 Texture2D<uint>		{DataInterfaceName}_NormalVirtualTexturePageTableIndirection;
 uint				{DataInterfaceName}_NormalVirtualTextureAdaptive;
+uint2				{DataInterfaceName}_NormalVirtualTexturePageTableLayers;
 SamplerState		{DataInterfaceName}_NormalVirtualTexture0Sampler;
 SamplerState		{DataInterfaceName}_NormalVirtualTexture1Sampler;
 float3				{DataInterfaceName}_NormalVirtualTextureLWCTile;
@@ -115,7 +116,7 @@ float GetHeight_{DataInterfaceName}(float3 InWorldPos)
 			{DataInterfaceName}_HeightVirtualTexture,
 			{DataInterfaceName}_HeightVirtualTextureSampler,
 			PageTable,
-			0,
+			0, // Height is always in layer 0.
 			VTUniform_Unpack({DataInterfaceName}_HeightVirtualTextureUniforms));
 
 		return VirtualTextureUnpackHeight(PackedValue, UnpackHeightScaleBias);
@@ -188,14 +189,14 @@ float3 GetNormal_{DataInterfaceName}(float3 InWorldPos)
 			{DataInterfaceName}_NormalVirtualTexture0,
 			{DataInterfaceName}_NormalVirtualTexture0Sampler,
 			PageTable,
-			0,
+			{DataInterfaceName}_NormalVirtualTexturePageTableLayers.x,
 			VTUniform_Unpack({DataInterfaceName}_NormalVirtualTextureUniforms0));
 
 		const float4 PackedValue1 = TextureVirtualSample(
 			{DataInterfaceName}_NormalVirtualTexture1,
 			{DataInterfaceName}_NormalVirtualTexture1Sampler,
 			PageTable,
-			0,
+			{DataInterfaceName}_NormalVirtualTexturePageTableLayers.y,
 			VTUniform_Unpack({DataInterfaceName}_NormalVirtualTextureUniforms1));
 
 		if ({DataInterfaceName}_NormalVirtualTextureUnpackMode == VIRTUAL_TEXTURE_UNPACK_NORMAL_BC3BC3)
@@ -262,7 +263,7 @@ float3 GetBaseColor_{DataInterfaceName}(float3 InWorldPos)
 			{DataInterfaceName}_BaseColorVirtualTexture,
 			{DataInterfaceName}_BaseColorVirtualTextureSampler,
 			PageTable,
-			0,
+			0, // BaseColor is always in layer 0.
 			VTUniform_Unpack({DataInterfaceName}_BaseColorVirtualTextureUniforms));
 
 		if ({DataInterfaceName}_BaseColorVirtualTextureUnpackSRGB)
diff --git a/Engine/Plugins/PCG/Shaders/Private/PCGVirtualTextureDataInterface.ush b/Engine/Plugins/PCG/Shaders/Private/PCGVirtualTextureDataInterface.ush
--- a/Engine/Plugins/PCG/Shaders/Private/PCGVirtualTextureDataInterface.ush
+++ b/Engine/Plugins/PCG/Shaders/Private/PCGVirtualTextureDataInterface.ush
@@ -6,12 +6,8 @@
 	Texture2D {DataInterfaceName}_VirtualTexture0_##DataIndex; \
 	Texture2D {DataInterfaceName}_VirtualTexture1_##DataIndex; \
 	Texture2D {DataInterfaceName}_VirtualTexture2_##DataIndex; \
-	Texture2D<uint4> {DataInterfaceName}_PageTable0_##DataIndex; \
-	Texture2D<uint4> {DataInterfaceName}_PageTable1_##DataIndex; \
-	Texture2D<uint4> {DataInterfaceName}_PageTable2_##DataIndex; \
-	Texture2D<uint> {DataInterfaceName}_PageTableIndirection0_##DataIndex; \
-	Texture2D<uint> {DataInterfaceName}_PageTableIndirection1_##DataIndex; \
-	Texture2D<uint> {DataInterfaceName}_PageTableIndirection2_##DataIndex;
+	Texture2D<uint4> {DataInterfaceName}_PageTable_##DataIndex; \
+	Texture2D<uint> {DataInterfaceName}_PageTableIndirection_##DataIndex;
 
 PCG_VIRTUAL_TEXTURE_DECLARATION(0)
 PCG_VIRTUAL_TEXTURE_DECLARATION(1)
@@ -20,28 +16,28 @@ PCG_VIRTUAL_TEXTURE_DECLARATION(3)
 
 #undef PCG_VIRTUAL_TEXTURE_DECLARATION
 
-// Layer infos
-uint4 {DataInterfaceName}_MaterialTypes[PCG_MAX_NUM_BINDINGS];
-uint4 {DataInterfaceName}_ValidLayerMasks[PCG_MAX_NUM_BINDINGS];
+// Binding infos
+uint4 {DataInterfaceName}_PageTableUniforms[PCG_MAX_NUM_BINDINGS * 2];
+uint4 {DataInterfaceName}_PageTableInfo[PCG_MAX_NUM_BINDINGS];
 float4 {DataInterfaceName}_LWCTiles[PCG_MAX_NUM_BINDINGS];
 float4x4 {DataInterfaceName}_WorldToUVTransforms[PCG_MAX_NUM_BINDINGS];
-uint4 {DataInterfaceName}_Adaptive[PCG_MAX_NUM_BINDINGS * PCG_MAX_NUM_LAYERS];
+
+// Layer infos
 uint4 {DataInterfaceName}_Uniforms[PCG_MAX_NUM_BINDINGS * PCG_MAX_NUM_LAYERS];
-uint4 {DataInterfaceName}_PageTableUniforms[PCG_MAX_NUM_BINDINGS * 2];
+uint4 {DataInterfaceName}_ValidLayerMasks[PCG_MAX_NUM_BINDINGS];
 
 // Shared
 uint {DataInterfaceName}_NumVirtualTextures;
 SamplerState {DataInterfaceName}_Sampler;
 
-float4 SampleLayer_{DataInterfaceName}(
+// Sample the page table
+VTPageTableResult SamplePageTable_{DataInterfaceName}(
 	uint InDataIndex,
 	float2 InSampleUV,
-	Texture2D InTexture,
 	Texture2D<uint4> InPageTable,
 	Texture2D<uint> InPageTableIndirection,
-	uint InAdaptive,
-	uint4 InUniforms,
-	uint4 PageTableUniforms[2])
+	uint4 PageTableUniforms[2],
+	bool InAdaptive)
 {
 	VTPageTableResult PageTable;
 
@@ -67,23 +63,32 @@ float4 SampleLayer_{DataInterfaceName}(
 			0.0f);
 	}
 
-	return TextureVirtualSample(InTexture, {DataInterfaceName}_Sampler, PageTable, 0, VTUniform_Unpack(InUniforms));
+	return PageTable;
+}
+
+// Sample the physical texture
+float4 SampleLayer_{DataInterfaceName}(
+	VTPageTableResult InPageTableResult,
+	Texture2D InTexture,
+	uint4 InUniforms,
+	uint InLayerIndex)
+{
+	return TextureVirtualSample(InTexture, {DataInterfaceName}_Sampler, InPageTableResult, InLayerIndex, VTUniform_Unpack(InUniforms));
 }
 
 #define PCG_CONDITIONAL_VIRTUAL_TEXTURE_SAMPLE(DataIndex) \
 	if (InDataIndex == DataIndex) \
 	{ \
+		VTPageTableResult PageTableResult = SamplePageTable_{DataInterfaceName}( \
+			InDataIndex, \
+			SampleUV, \
+			{DataInterfaceName}_PageTable_##DataIndex, \
+			{DataInterfaceName}_PageTableIndirection_##DataIndex, \
+			PageTableUniforms, \
+			bAdaptive); \
 		if (ValidLayerMask & 0x1) \
 		{ \
-			LayerSample[0] = SampleLayer_{DataInterfaceName}( \
-				InDataIndex, \
-				SampleUV, \
-				{DataInterfaceName}_VirtualTexture0_##DataIndex, \
-				{DataInterfaceName}_PageTable0_##DataIndex, \
-				{DataInterfaceName}_PageTableIndirection0_##DataIndex, \
-				Adaptive0, \
-				Uniforms0, \
-				PageTableUniforms); \
+			LayerSample[0] = SampleLayer_{DataInterfaceName}(PageTableResult, {DataInterfaceName}_VirtualTexture0_##DataIndex, Uniforms0, 0); \
 		} \
 		else \
 		{ \
@@ -91,15 +96,7 @@ float4 SampleLayer_{DataInterfaceName}(
 		} \
 		if (ValidLayerMask & 0x2) \
 		{ \
-			LayerSample[1] = SampleLayer_{DataInterfaceName}( \
-				InDataIndex, \
-				SampleUV, \
-				{DataInterfaceName}_VirtualTexture1_##DataIndex, \
-				{DataInterfaceName}_PageTable1_##DataIndex, \
-				{DataInterfaceName}_PageTableIndirection1_##DataIndex, \
-				Adaptive1, \
-				Uniforms1, \
-				PageTableUniforms); \
+			LayerSample[1] = SampleLayer_{DataInterfaceName}(PageTableResult, {DataInterfaceName}_VirtualTexture1_##DataIndex, Uniforms1, bSinglePhysicalSpace ? 0 : 1); \
 		} \
 		else \
 		{ \
@@ -107,15 +104,7 @@ float4 SampleLayer_{DataInterfaceName}(
 		} \
 		if (ValidLayerMask & 0x4) \
 		{ \
-			LayerSample[2] = SampleLayer_{DataInterfaceName}( \
-				InDataIndex, \
-				SampleUV, \
-				{DataInterfaceName}_VirtualTexture2_##DataIndex, \
-				{DataInterfaceName}_PageTable2_##DataIndex, \
-				{DataInterfaceName}_PageTableIndirection2_##DataIndex, \
-				Adaptive2, \
-				Uniforms2, \
-				PageTableUniforms); \
+			LayerSample[2] = SampleLayer_{DataInterfaceName}(PageTableResult, {DataInterfaceName}_VirtualTexture2_##DataIndex, Uniforms2, bSinglePhysicalSpace ? 0 : 2); \
 		} \
 		else \
 		{ \
@@ -170,22 +159,21 @@ void SampleVirtualTexture_{DataInterfaceName}(
 	const uint ValidLayerMask = {DataInterfaceName}_ValidLayerMasks[InDataIndex].x;
 
 	{
-		const uint PackedLayerIndex = InDataIndex * PCG_MAX_NUM_LAYERS;
-
-		uint Adaptive0 = {DataInterfaceName}_Adaptive[PackedLayerIndex + 0].x;
-		uint Adaptive1 = {DataInterfaceName}_Adaptive[PackedLayerIndex + 1].x;
-		uint Adaptive2 = {DataInterfaceName}_Adaptive[PackedLayerIndex + 2].x;
-
-		uint4 Uniforms0 = {DataInterfaceName}_Uniforms[PackedLayerIndex + 0];
-		uint4 Uniforms1 = {DataInterfaceName}_Uniforms[PackedLayerIndex + 1];
-		uint4 Uniforms2 = {DataInterfaceName}_Uniforms[PackedLayerIndex + 2];
-
 		uint4 PageTableUniforms[2] =
 		{
 			{DataInterfaceName}_PageTableUniforms[InDataIndex * 2 + 0],
 			{DataInterfaceName}_PageTableUniforms[InDataIndex * 2 + 1],
 		};
 
+		bool bAdaptive = {DataInterfaceName}_PageTableInfo[InDataIndex].y != 0;
+		bool bSinglePhysicalSpace = {DataInterfaceName}_PageTableInfo[InDataIndex].z != 0;
+
+		const uint PackedLayerIndex = InDataIndex * PCG_MAX_NUM_LAYERS;
+
+		uint4 Uniforms0 = {DataInterfaceName}_Uniforms[PackedLayerIndex + 0];
+		uint4 Uniforms1 = {DataInterfaceName}_Uniforms[PackedLayerIndex + 1];
+		uint4 Uniforms2 = {DataInterfaceName}_Uniforms[PackedLayerIndex + 2];
+
 		PCG_CONDITIONAL_VIRTUAL_TEXTURE_SAMPLE(0)
 #if ENABLE_MULTIPLE_VIRTUAL_TEXTURES
 		PCG_CONDITIONAL_VIRTUAL_TEXTURE_SAMPLE(1)
@@ -195,7 +183,7 @@ void SampleVirtualTexture_{DataInterfaceName}(
 	}
 
 	// Unpack layer samples based on material type.
-	const uint MaterialType = {DataInterfaceName}_MaterialTypes[InDataIndex].x;
+	const uint MaterialType = {DataInterfaceName}_PageTableInfo[InDataIndex].x;
 
 	switch (MaterialType)
 	{
diff --git a/Engine/Plugins/PCG/Source/PCG/Private/Compute/Data/PCGVirtualTextureCommon.cpp b/Engine/Plugins/PCG/Source/PCG/Private/Compute/Data/PCGVirtualTextureCommon.cpp
--- a/Engine/Plugins/PCG/Source/PCG/Private/Compute/Data/PCGVirtualTextureCommon.cpp
+++ b/Engine/Plugins/PCG/Source/PCG/Private/Compute/Data/PCGVirtualTextureCommon.cpp
@@ -35,7 +35,7 @@ namespace PCGVirtualTextureCommon
 		return TextureSRV.IsValid();
 	}
 	
-	void FVirtualTexturePageTable::Initialize(const URuntimeVirtualTexture* VirtualTexture, uint32 PageTableIndex, bool bIncludeWorldToUV, bool bIncludeHeightUnpack)
+	void FVirtualTexturePageTable::Initialize(const URuntimeVirtualTexture* VirtualTexture)
 	{
 		Reset();
 	
@@ -43,25 +43,22 @@ namespace PCGVirtualTextureCommon
 		{
 			if (IAllocatedVirtualTexture* AllocatedTexture = VirtualTexture->GetAllocatedVirtualTexture())
 			{
-				PageTableRef = PageTableIndex < AllocatedTexture->GetNumPageTableTextures() ? AllocatedTexture->GetPageTableTexture(PageTableIndex) : nullptr;
+				// RVT should always have a single PageTable texture.
+				check(AllocatedTexture->GetNumPageTableTextures() == 1);
+				PageTableRef = AllocatedTexture->GetPageTableTexture(0);
 				PageTableIndirectionRef = AllocatedTexture->GetPageTableIndirectionTexture();
+
 				bIsAdaptive = VirtualTexture->GetAdaptivePageTable();
+				bSinglePhysicalSpace = VirtualTexture->GetSinglePhysicalSpace();
 	
 				if (PageTableRef.IsValid())
 				{
 					AllocatedTexture->GetPackedPageTableUniform(PageTableUniforms);
 	
-					if (bIncludeWorldToUV)
-					{
-						WorldToUVParameters[0] = VirtualTexture->GetUniformParameter(ERuntimeVirtualTextureShaderUniform_WorldToUVTransform0);
-						WorldToUVParameters[1] = VirtualTexture->GetUniformParameter(ERuntimeVirtualTextureShaderUniform_WorldToUVTransform1);
-						WorldToUVParameters[2] = VirtualTexture->GetUniformParameter(ERuntimeVirtualTextureShaderUniform_WorldToUVTransform2);
-					}
-	
-					if (bIncludeHeightUnpack)
-					{
-						WorldToUVParameters[3] = VirtualTexture->GetUniformParameter(ERuntimeVirtualTextureShaderUniform_WorldHeightUnpack);
-					}
+					WorldToUVParameters[0] = VirtualTexture->GetUniformParameter(ERuntimeVirtualTextureShaderUniform_WorldToUVTransform0);
+					WorldToUVParameters[1] = VirtualTexture->GetUniformParameter(ERuntimeVirtualTextureShaderUniform_WorldToUVTransform1);
+					WorldToUVParameters[2] = VirtualTexture->GetUniformParameter(ERuntimeVirtualTextureShaderUniform_WorldToUVTransform2);
+					WorldToUVParameters[3] = VirtualTexture->GetUniformParameter(ERuntimeVirtualTextureShaderUniform_WorldHeightUnpack);
 				}
 			}
 		}
@@ -72,6 +69,7 @@ namespace PCGVirtualTextureCommon
 		PageTableRef = nullptr;
 		PageTableIndirectionRef = nullptr;
 		bIsAdaptive = false;
+		bSinglePhysicalSpace = false;
 		PageTableUniforms[0] = FUintVector4::ZeroValue;
 		PageTableUniforms[1] = FUintVector4::ZeroValue;
 		WorldToUVParameters[0] = FVector4::Zero();
diff --git a/Engine/Plugins/PCG/Source/PCG/Private/Compute/Data/PCGVirtualTextureCommon.h b/Engine/Plugins/PCG/Source/PCG/Private/Compute/Data/PCGVirtualTextureCommon.h
--- a/Engine/Plugins/PCG/Source/PCG/Private/Compute/Data/PCGVirtualTextureCommon.h
+++ b/Engine/Plugins/PCG/Source/PCG/Private/Compute/Data/PCGVirtualTextureCommon.h
@@ -40,12 +40,13 @@ namespace PCGVirtualTextureCommon
 
 	struct FVirtualTexturePageTable
 	{
-		void Initialize(const URuntimeVirtualTexture* VirtualTexture, uint32 PageTableIndex, bool bIncludeWorldToUV, bool bIncludeHeightUnpack);
+		void Initialize(const URuntimeVirtualTexture* VirtualTexture);
 		void Reset();
 		bool IsValid() const;
 
 		FTextureRHIRef PageTableRef;
 		FTextureRHIRef PageTableIndirectionRef;
+		bool bSinglePhysicalSpace = false;
 		bool bIsAdaptive = false;
 		FUintVector4 PageTableUniforms[2];
 		FVector4 WorldToUVParameters[4];
diff --git a/Engine/Plugins/PCG/Source/PCG/Private/Compute/DataInterfaces/PCGLandscapeDataInterface.cpp b/Engine/Plugins/PCG/Source/PCG/Private/Compute/DataInterfaces/PCGLandscapeDataInterface.cpp
--- a/Engine/Plugins/PCG/Source/PCG/Private/Compute/DataInterfaces/PCGLandscapeDataInterface.cpp
+++ b/Engine/Plugins/PCG/Source/PCG/Private/Compute/DataInterfaces/PCGLandscapeDataInterface.cpp
@@ -201,6 +201,7 @@ BEGIN_SHADER_PARAMETER_STRUCT(FPCGLandscapeDataInterfaceParameters,)
 	SHADER_PARAMETER_TEXTURE(Texture2D<uint4>,	NormalVirtualTexturePageTable)
 	SHADER_PARAMETER_TEXTURE(Texture2D<uint>,	NormalVirtualTexturePageTableIndirection)
 	SHADER_PARAMETER(uint32,					NormalVirtualTextureAdaptive)
+	SHADER_PARAMETER(FUintVector2,				NormalVirtualTexturePageTableLayers)
 	SHADER_PARAMETER_SAMPLER(SamplerState,		NormalVirtualTexture0Sampler)
 	SHADER_PARAMETER_SAMPLER(SamplerState,		NormalVirtualTexture1Sampler)
 	SHADER_PARAMETER(FVector3f,					NormalVirtualTextureLWCTile)
@@ -622,20 +623,20 @@ void FPCGLandscapeDataProviderProxy::AllocateResources(FRDGBuilder& GraphBuilder
 	if (TStrongObjectPtr<const URuntimeVirtualTexture> BaseColorVT = BaseColorVirtualTexture.Pin())
 	{
 		BaseColorVirtualLayer.Initialize(BaseColorVT.Get(), /*PageTableIndex=*/0, bBaseColorSRGB);
-		BaseColorVirtualPage.Initialize(BaseColorVT.Get(), /*PageTableIndex=*/0, /*bIncludeWorldToUV=*/true, /*bIncludeHeightUnpack=*/false);
+		BaseColorVirtualPage.Initialize(BaseColorVT.Get());
 	}
 
 	if (TStrongObjectPtr<const URuntimeVirtualTexture> HeightVT = HeightVirtualTexture.Pin())
 	{
 		HeightVirtualLayer.Initialize(HeightVT.Get(), /*PageTableIndex=*/0, /*bSRGB=*/false);
-		HeightVirtualPage.Initialize(HeightVT.Get(), /*PageTableIndex=*/0, /*bIncludeWorldToUV=*/true, /*bIncludeHeightUnpack=*/true);
+		HeightVirtualPage.Initialize(HeightVT.Get());
 	}
 
 	if (TStrongObjectPtr<const URuntimeVirtualTexture> NormalVT = NormalVirtualTexture.Pin())
 	{
 		const URuntimeVirtualTexture* NormalVTPtr = NormalVT.Get();
 
-		NormalVirtualPage.Initialize(NormalVTPtr, /*PageTableIndex=*/0, /*bIncludeWorldToUV=*/true, /*bIncludeHeightUnpack=*/false);
+		NormalVirtualPage.Initialize(NormalVTPtr);
 
 		switch (NormalVirtualTextureMode)
 		{
@@ -643,23 +644,27 @@ void FPCGLandscapeDataProviderProxy::AllocateResources(FRDGBuilder& GraphBuilder
 			NormalVirtualLayer0.Initialize(NormalVTPtr, /*PageTableIndex=*/0, /*bSRGB=*/false);
 			NormalVirtualLayer1.Initialize(NormalVTPtr, /*PageTableIndex=*/1, /*bSRGB=*/false);
 			NormalUnpackMode = PCGVirtualTextureCommon::ENormalUnpackType::B5G6R5;
+			NormalPageTableLayers = NormalVirtualPage.bSinglePhysicalSpace ? FUintVector2(0, 0) : FUintVector2(0, 1);
 			break;
 		case ERuntimeVirtualTextureMaterialType::BaseColor_Normal_Specular:
 			NormalVirtualLayer0.Initialize(NormalVTPtr, /*PageTableIndex=*/0, /*bSRGB=*/false);
 			NormalVirtualLayer1.Initialize(NormalVTPtr, /*PageTableIndex=*/1, /*bSRGB=*/false);
 			NormalUnpackMode = PCGVirtualTextureCommon::ENormalUnpackType::BC3BC3;
+			NormalPageTableLayers = NormalVirtualPage.bSinglePhysicalSpace ? FUintVector2(0, 0) : FUintVector2(0, 1);
 			break;
 		case ERuntimeVirtualTextureMaterialType::BaseColor_Normal_Specular_YCoCg:
 		case ERuntimeVirtualTextureMaterialType::BaseColor_Normal_Specular_Mask_YCoCg:
 			NormalVirtualLayer0.Initialize(NormalVTPtr, /*PageTableIndex=*/1, /*bSRGB=*/false);
 			NormalVirtualLayer1.Initialize(NormalVTPtr, /*PageTableIndex=*/2, /*bSRGB=*/false);
 			NormalUnpackMode = PCGVirtualTextureCommon::ENormalUnpackType::BC5BC1;
+			NormalPageTableLayers = NormalVirtualPage.bSinglePhysicalSpace ? FUintVector2(0, 0) : FUintVector2(1, 2);
 			break;
 		default:
 			NormalVirtualPage.Reset();
 			NormalVirtualLayer0.Reset();
 			NormalVirtualLayer1.Reset();
 			NormalUnpackMode = PCGVirtualTextureCommon::ENormalUnpackType::None;
+			NormalPageTableLayers = FUintVector2(0, 0);
 			break;
 		}
 	}
@@ -801,7 +806,7 @@ bool FPCGLandscapeDataProviderProxy::SetNormalVirtualTextureParameters(FParamete
 	OutShaderParameters.NormalVirtualTexturePackedUniform1 = NormalVirtualPage.PageTableUniforms[1];
 	OutShaderParameters.NormalVirtualTextureUniforms0 = NormalVirtualLayer0.TextureUniforms;
 	OutShaderParameters.NormalVirtualTextureUniforms1 = NormalVirtualLayer1.TextureUniforms;
-
+	OutShaderParameters.NormalVirtualTexturePageTableLayers = NormalPageTableLayers;
 	return true;
 }
 
@@ -870,6 +875,7 @@ void FPCGLandscapeDataProviderProxy::SetNormalVirtualTextureParameters_Default(F
 	OutShaderParameters.NormalVirtualTexturePageTable = GBlackUintTexture->TextureRHI;
 	OutShaderParameters.NormalVirtualTexturePageTableIndirection = GBlackUintTexture->TextureRHI;
 	OutShaderParameters.NormalVirtualTextureAdaptive = 0;
+	OutShaderParameters.NormalVirtualTexturePageTableLayers = FUintVector2(0, 0);
 	OutShaderParameters.NormalVirtualTextureLWCTile = FVector3f::ZeroVector;
 	OutShaderParameters.NormalVirtualTextureWorldToUVTransform = FMatrix44f::Identity;
 	OutShaderParameters.NormalVirtualTextureEnabled = 0;
diff --git a/Engine/Plugins/PCG/Source/PCG/Private/Compute/DataInterfaces/PCGLandscapeDataInterface.h b/Engine/Plugins/PCG/Source/PCG/Private/Compute/DataInterfaces/PCGLandscapeDataInterface.h
--- a/Engine/Plugins/PCG/Source/PCG/Private/Compute/DataInterfaces/PCGLandscapeDataInterface.h
+++ b/Engine/Plugins/PCG/Source/PCG/Private/Compute/DataInterfaces/PCGLandscapeDataInterface.h
@@ -154,6 +154,7 @@ protected:
 	PCGVirtualTextureCommon::FVirtualTextureLayer NormalVirtualLayer0;
 	PCGVirtualTextureCommon::FVirtualTextureLayer NormalVirtualLayer1;
 	PCGVirtualTextureCommon::ENormalUnpackType NormalUnpackMode;
+	FUintVector2 NormalPageTableLayers;
 
 	FVector2D LandscapeGridSize = FVector2D(1.0f, 1.0f);
 };
diff --git a/Engine/Plugins/PCG/Source/PCG/Private/Compute/DataInterfaces/PCGVirtualTextureDataInterface.cpp b/Engine/Plugins/PCG/Source/PCG/Private/Compute/DataInterfaces/PCGVirtualTextureDataInterface.cpp
--- a/Engine/Plugins/PCG/Source/PCG/Private/Compute/DataInterfaces/PCGVirtualTextureDataInterface.cpp
+++ b/Engine/Plugins/PCG/Source/PCG/Private/Compute/DataInterfaces/PCGVirtualTextureDataInterface.cpp
@@ -42,29 +42,20 @@ void UPCGVirtualTextureDataInterface::GetSupportedInputs(TArray<FShaderFunctionD
 }
 
 BEGIN_SHADER_PARAMETER_STRUCT(FPCGVirtualTextureDataInterfaceParameters, )
-	// Layer 0
-	SHADER_PARAMETER_SRV_ARRAY(Texture2D, VirtualTexture0, [PCGVirtualTextureDIConstants::MAX_NUM_BINDINGS])
-	SHADER_PARAMETER_TEXTURE_ARRAY(Texture2D<uint4>, PageTable0, [PCGVirtualTextureDIConstants::MAX_NUM_BINDINGS])
-	SHADER_PARAMETER_TEXTURE_ARRAY(Texture2D<uint>, PageTableIndirection0, [PCGVirtualTextureDIConstants::MAX_NUM_BINDINGS])
+	// Per binding info
+	SHADER_PARAMETER_TEXTURE_ARRAY(Texture2D<uint4>, PageTable, [PCGVirtualTextureDIConstants::MAX_NUM_BINDINGS])
+	SHADER_PARAMETER_TEXTURE_ARRAY(Texture2D<uint>, PageTableIndirection, [PCGVirtualTextureDIConstants::MAX_NUM_BINDINGS])
+	SHADER_PARAMETER_ARRAY(FUintVector4, PageTableUniforms, [PCGVirtualTextureDIConstants::MAX_NUM_BINDINGS * 2])
+	SHADER_PARAMETER_ARRAY(FUintVector4, PageTableInfo, [PCGVirtualTextureDIConstants::MAX_NUM_BINDINGS])
+	SHADER_PARAMETER_ARRAY(FVector4f, LWCTiles, [PCGVirtualTextureDIConstants::MAX_NUM_BINDINGS])
+	SHADER_PARAMETER_ARRAY(FMatrix44f, WorldToUVTransforms, [PCGVirtualTextureDIConstants::MAX_NUM_BINDINGS])
 
-	// Layer 1
+	// Per layer info
+	SHADER_PARAMETER_SRV_ARRAY(Texture2D, VirtualTexture0, [PCGVirtualTextureDIConstants::MAX_NUM_BINDINGS])
 	SHADER_PARAMETER_SRV_ARRAY(Texture2D, VirtualTexture1, [PCGVirtualTextureDIConstants::MAX_NUM_BINDINGS])
-	SHADER_PARAMETER_TEXTURE_ARRAY(Texture2D<uint4>, PageTable1, [PCGVirtualTextureDIConstants::MAX_NUM_BINDINGS])
-	SHADER_PARAMETER_TEXTURE_ARRAY(Texture2D<uint>, PageTableIndirection1, [PCGVirtualTextureDIConstants::MAX_NUM_BINDINGS])
-
-	// Layer 2
 	SHADER_PARAMETER_SRV_ARRAY(Texture2D, VirtualTexture2, [PCGVirtualTextureDIConstants::MAX_NUM_BINDINGS])
-	SHADER_PARAMETER_TEXTURE_ARRAY(Texture2D<uint4>, PageTable2, [PCGVirtualTextureDIConstants::MAX_NUM_BINDINGS])
-	SHADER_PARAMETER_TEXTURE_ARRAY(Texture2D<uint>, PageTableIndirection2, [PCGVirtualTextureDIConstants::MAX_NUM_BINDINGS])
-
-	// Layer info
-	SHADER_PARAMETER_ARRAY(FUintVector4, MaterialTypes, [PCGVirtualTextureDIConstants::MAX_NUM_BINDINGS])
-	SHADER_PARAMETER_ARRAY(FUintVector4, ValidLayerMasks, [PCGVirtualTextureDIConstants::MAX_NUM_BINDINGS])
-	SHADER_PARAMETER_ARRAY(FVector4f, LWCTiles, [PCGVirtualTextureDIConstants::MAX_NUM_BINDINGS])
-	SHADER_PARAMETER_ARRAY(FMatrix44f, WorldToUVTransforms, [PCGVirtualTextureDIConstants::MAX_NUM_BINDINGS])
-	SHADER_PARAMETER_ARRAY(FUintVector4, Adaptive, [PCGVirtualTextureDIConstants::MAX_NUM_BINDINGS * PCGVirtualTextureDIConstants::MAX_NUM_LAYERS])
 	SHADER_PARAMETER_ARRAY(FUintVector4, Uniforms, [PCGVirtualTextureDIConstants::MAX_NUM_BINDINGS * PCGVirtualTextureDIConstants::MAX_NUM_LAYERS])
-	SHADER_PARAMETER_ARRAY(FUintVector4, PageTableUniforms, [PCGVirtualTextureDIConstants::MAX_NUM_BINDINGS * 2])
+	SHADER_PARAMETER_ARRAY(FUintVector4, ValidLayerMasks, [PCGVirtualTextureDIConstants::MAX_NUM_BINDINGS])
 
 	// Shared
 	SHADER_PARAMETER(uint32, NumVirtualTextures)
@@ -184,30 +175,31 @@ void UPCGVirtualTextureDataProvider::Reset()
 
 FPCGVirtualTextureDataProviderProxy::FPCGVirtualTextureDataProviderProxy(const TArray<TObjectPtr<const URuntimeVirtualTexture>>& InVirtualTextures)
 {
+	for (int32 BindingIndex = 0; BindingIndex < PCGVirtualTextureDIConstants::MAX_NUM_BINDINGS; ++BindingIndex)
+	{
+		MaterialTypes[BindingIndex] = ERuntimeVirtualTextureMaterialType::Count;
+	}
+
 	NumVirtualTextures = InVirtualTextures.Num();
 	for (int32 VirtualTextureIndex = 0; VirtualTextureIndex < NumVirtualTextures; ++VirtualTextureIndex)
 	{
-		ERuntimeVirtualTextureMaterialType MaterialType = ERuntimeVirtualTextureMaterialType::Count;
-
 		if (const URuntimeVirtualTexture* VirtualTexture = InVirtualTextures[VirtualTextureIndex])
 		{
-			MaterialType = VirtualTexture->GetMaterialType();
+			const ERuntimeVirtualTextureMaterialType MaterialType = VirtualTexture->GetMaterialType();
+			MaterialTypes[VirtualTextureIndex] = MaterialType;
+			
+			PageTables[VirtualTextureIndex].Initialize(VirtualTexture);
+
 			const int32 NumVirtualTextureLayers = FMath::Min(URuntimeVirtualTexture::GetLayerCount(MaterialType), PCGVirtualTextureDIConstants::MAX_NUM_LAYERS);
 			const int32 PackedLayerBaseIndex = VirtualTextureIndex * PCGVirtualTextureDIConstants::MAX_NUM_LAYERS;
 
 			for (int32 LayerIndex = 0; LayerIndex < NumVirtualTextureLayers; ++LayerIndex)
 			{
 				const int32 PackedLayerIndex = PackedLayerBaseIndex + LayerIndex;
-
-				// Note, we always use the first layer to get the world to uv transform because every material type uses at least the first layer.
-				const bool bIncludeWorldToUVTransform = LayerIndex == 0;
-
 				Layers[PackedLayerIndex].Initialize(VirtualTexture, LayerIndex, VirtualTexture->IsLayerSRGB(LayerIndex));
-				PageTables[PackedLayerIndex].Initialize(VirtualTexture, LayerIndex, bIncludeWorldToUVTransform, bIncludeWorldToUVTransform);
 			}
 		}
 
-		MaterialTypes[VirtualTextureIndex] = MaterialType;
 	}
 }
 
@@ -248,87 +240,76 @@ void FPCGVirtualTextureDataProviderProxy::GatherDispatchData(FDispatchData const
 
 		for (int32 BindingIndex = 0; BindingIndex < PCGVirtualTextureDIConstants::MAX_NUM_BINDINGS; ++BindingIndex)
 		{
+			const uint32 MaterialType = (uint32)MaterialTypes[BindingIndex];
+			const PCGVirtualTextureCommon::FVirtualTexturePageTable& PageTable = PageTables[BindingIndex];
+
+			if (PageTable.IsValid())
+			{
+				Parameters.PageTable[BindingIndex] = PageTable.PageTableRef;
+				Parameters.PageTableIndirection[BindingIndex] = PageTable.PageTableIndirectionRef;
+				Parameters.PageTableInfo[BindingIndex] = FUintVector4(MaterialType, PageTable.bIsAdaptive ? 1 : 0, PageTable.bSinglePhysicalSpace ? 1 : 0, 0);
+				Parameters.PageTableUniforms[BindingIndex * 2 + 0] = PageTable.PageTableUniforms[0];
+				Parameters.PageTableUniforms[BindingIndex * 2 + 1] = PageTable.PageTableUniforms[1];
+
+				const FLargeWorldRenderPosition TextureOrigin(PageTable.WorldToUVParameters[0]);
+				Parameters.LWCTiles[BindingIndex] = TextureOrigin.GetTile();
+				Parameters.WorldToUVTransforms[BindingIndex] = FMatrix44f(
+					TextureOrigin.GetOffset(),
+					FVector3f((FVector4f)PageTable.WorldToUVParameters[1]),
+					FVector3f((FVector4f)PageTable.WorldToUVParameters[2]),
+					FVector3f((FVector4f)PageTable.WorldToUVParameters[3]));
+			}
+			else
+			{
+				Parameters.PageTable[BindingIndex] = GBlackUintTexture->TextureRHI;
+				Parameters.PageTableIndirection[BindingIndex] = GBlackUintTexture->TextureRHI;
+				Parameters.PageTableInfo[BindingIndex] = FUintVector4(MaterialType, 0, 0, 0);
+				Parameters.PageTableUniforms[BindingIndex * 2 + 0] = FUintVector4::ZeroValue;
+				Parameters.PageTableUniforms[BindingIndex * 2 + 1] = FUintVector4::ZeroValue;
+				Parameters.LWCTiles[BindingIndex] = FVector3f::ZeroVector;
+				Parameters.WorldToUVTransforms[BindingIndex] = FMatrix44f::Identity;
+			}
+
 			const int32 PackedLayerBaseIndex = BindingIndex * PCGVirtualTextureDIConstants::MAX_NUM_LAYERS;
 			const PCGVirtualTextureCommon::FVirtualTextureLayer& Layer0 = Layers[PackedLayerBaseIndex + 0];
 			const PCGVirtualTextureCommon::FVirtualTextureLayer& Layer1 = Layers[PackedLayerBaseIndex + 1];
 			const PCGVirtualTextureCommon::FVirtualTextureLayer& Layer2 = Layers[PackedLayerBaseIndex + 2];
-			const PCGVirtualTextureCommon::FVirtualTexturePageTable& PageTable0 = PageTables[PackedLayerBaseIndex + 0];
-			const PCGVirtualTextureCommon::FVirtualTexturePageTable& PageTable1 = PageTables[PackedLayerBaseIndex + 1];
-			const PCGVirtualTextureCommon::FVirtualTexturePageTable& PageTable2 = PageTables[PackedLayerBaseIndex + 2];
 
-			Parameters.MaterialTypes[BindingIndex] = FUintVector4(static_cast<uint32>(MaterialTypes[BindingIndex]), 0, 0, 0);
 			Parameters.ValidLayerMasks[BindingIndex] = FUintVector4::ZeroValue;
 
-			if (Layer0.IsValid() && PageTable0.IsValid())
+			if (Layer0.IsValid())
 			{
 				Parameters.ValidLayerMasks[BindingIndex].X |= 1;
 				Parameters.VirtualTexture0[BindingIndex] = Layer0.TextureSRV;
-				Parameters.PageTable0[BindingIndex] = PageTable0.PageTableRef;
-				Parameters.PageTableIndirection0[BindingIndex] = PageTable0.PageTableIndirectionRef;
-				Parameters.Adaptive[PackedLayerBaseIndex + 0] = static_cast<FUintVector4>(PageTable0.bIsAdaptive);
 				Parameters.Uniforms[PackedLayerBaseIndex + 0] = Layer0.TextureUniforms;
-
-				// Layer 0 should always be valid, so we can use it to set the shared parameters.
-				FLargeWorldRenderPosition TextureOrigin(PageTable0.WorldToUVParameters[0]);
-				Parameters.LWCTiles[BindingIndex] = TextureOrigin.GetTile();
-
-				Parameters.WorldToUVTransforms[BindingIndex] = FMatrix44f(
-					TextureOrigin.GetOffset(),
-					FVector3f((FVector4f)PageTable0.WorldToUVParameters[1]),
-					FVector3f((FVector4f)PageTable0.WorldToUVParameters[2]),
-					FVector3f((FVector4f)PageTable0.WorldToUVParameters[3]));
-
-				Parameters.PageTableUniforms[BindingIndex * 2 + 0] = PageTable0.PageTableUniforms[0];
-				Parameters.PageTableUniforms[BindingIndex * 2 + 1] = PageTable0.PageTableUniforms[1];
 			}
 			else
 			{
 				Parameters.VirtualTexture0[BindingIndex] = GBlackTextureWithSRV->ShaderResourceViewRHI;
-				Parameters.PageTable0[BindingIndex] = GBlackUintTexture->TextureRHI;
-				Parameters.PageTableIndirection0[BindingIndex] = GBlackUintTexture->TextureRHI;
-				Parameters.Adaptive[PackedLayerBaseIndex + 0] = FUintVector4::ZeroValue;
 				Parameters.Uniforms[PackedLayerBaseIndex + 0] = FUintVector4::ZeroValue;
-
-				// If Layer 0 was invalid, default initialize the shared parameters as well.
-				Parameters.LWCTiles[BindingIndex] = FVector3f::ZeroVector;
-				Parameters.WorldToUVTransforms[BindingIndex] = FMatrix44f::Identity;
-				Parameters.PageTableUniforms[BindingIndex * 2 + 0] = FUintVector4::ZeroValue;
-				Parameters.PageTableUniforms[BindingIndex * 2 + 1] = FUintVector4::ZeroValue;
 			}
 
-			if (Layer1.IsValid() && PageTable1.IsValid())
+			if (Layer1.IsValid())
 			{
 				Parameters.ValidLayerMasks[BindingIndex].X |= 2;
 				Parameters.VirtualTexture1[BindingIndex] = Layer1.TextureSRV;
-				Parameters.PageTable1[BindingIndex] = PageTable1.PageTableRef;
-				Parameters.PageTableIndirection1[BindingIndex] = PageTable1.PageTableIndirectionRef;
-				Parameters.Adaptive[PackedLayerBaseIndex + 1] = static_cast<FUintVector4>(PageTable1.bIsAdaptive);
 				Parameters.Uniforms[PackedLayerBaseIndex + 1] = Layer1.TextureUniforms;
 			}
 			else
 			{
 				Parameters.VirtualTexture1[BindingIndex] = GBlackTextureWithSRV->ShaderResourceViewRHI;
-				Parameters.PageTable1[BindingIndex] = GBlackUintTexture->TextureRHI;
-				Parameters.PageTableIndirection1[BindingIndex] = GBlackUintTexture->TextureRHI;
-				Parameters.Adaptive[PackedLayerBaseIndex + 1] = FUintVector4::ZeroValue;
 				Parameters.Uniforms[PackedLayerBaseIndex + 1] = FUintVector4::ZeroValue;
 			}
 
-			if (Layer2.IsValid() && PageTable2.IsValid())
+			if (Layer2.IsValid())
 			{
 				Parameters.ValidLayerMasks[BindingIndex].X |= 4;
 				Parameters.VirtualTexture2[BindingIndex] = Layer2.TextureSRV;
-				Parameters.PageTable2[BindingIndex] = PageTable2.PageTableRef;
-				Parameters.PageTableIndirection2[BindingIndex] = PageTable2.PageTableIndirectionRef;
-				Parameters.Adaptive[PackedLayerBaseIndex + 2] = static_cast<FUintVector4>(PageTable2.bIsAdaptive);
 				Parameters.Uniforms[PackedLayerBaseIndex + 2] = Layer2.TextureUniforms;
 			}
 			else
 			{
 				Parameters.VirtualTexture2[BindingIndex] = GBlackTextureWithSRV->ShaderResourceViewRHI;
-				Parameters.PageTable2[BindingIndex] = GBlackUintTexture->TextureRHI;
-				Parameters.PageTableIndirection2[BindingIndex] = GBlackUintTexture->TextureRHI;
-				Parameters.Adaptive[PackedLayerBaseIndex + 2] = FUintVector4::ZeroValue;
 				Parameters.Uniforms[PackedLayerBaseIndex + 2] = FUintVector4::ZeroValue;
 			}
 		}
diff --git a/Engine/Plugins/PCG/Source/PCG/Private/Compute/DataInterfaces/PCGVirtualTextureDataInterface.h b/Engine/Plugins/PCG/Source/PCG/Private/Compute/DataInterfaces/PCGVirtualTextureDataInterface.h
--- a/Engine/Plugins/PCG/Source/PCG/Private/Compute/DataInterfaces/PCGVirtualTextureDataInterface.h
+++ b/Engine/Plugins/PCG/Source/PCG/Private/Compute/DataInterfaces/PCGVirtualTextureDataInterface.h
@@ -83,6 +83,6 @@ protected:
 
 	ERuntimeVirtualTextureMaterialType MaterialTypes[PCGVirtualTextureDIConstants::MAX_NUM_BINDINGS];
 	PCGVirtualTextureCommon::FVirtualTextureLayer Layers[PCGVirtualTextureDIConstants::MAX_NUM_LAYERS * PCGVirtualTextureDIConstants::MAX_NUM_BINDINGS];
-	PCGVirtualTextureCommon::FVirtualTexturePageTable PageTables[PCGVirtualTextureDIConstants::MAX_NUM_LAYERS * PCGVirtualTextureDIConstants::MAX_NUM_BINDINGS];
+	PCGVirtualTextureCommon::FVirtualTexturePageTable PageTables[PCGVirtualTextureDIConstants::MAX_NUM_BINDINGS];
 	int32 NumVirtualTextures = 0;
 };
