[PCG] Fix GetResolution API on texture array data interface.

#rnx
#Synced-CL 49645631

diff --git a/Engine/Plugins/PCG/Shaders/Private/PCGTexture2DArrayDataInterface.ush b/Engine/Plugins/PCG/Shaders/Private/PCGTexture2DArrayDataInterface.ush
--- a/Engine/Plugins/PCG/Shaders/Private/PCGTexture2DArrayDataInterface.ush
+++ b/Engine/Plugins/PCG/Shaders/Private/PCGTexture2DArrayDataInterface.ush
@@ -19,7 +19,7 @@ uint GetNumData_{DataInterfaceName}()
 
 uint2 GetResolution_{DataInterfaceName}(uint InDataIndex)
 {
-	return {DataInterfaceName}__Resolution;
+	return {DataInterfaceName}_Resolution;
 }
 
 uint GetArraySize_{DataInterfaceName}(uint InDataIndex)
diff --git a/Engine/Plugins/PCG/Source/PCG/Private/Compute/Elements/PCGCustomHLSL.cpp b/Engine/Plugins/PCG/Source/PCG/Private/Compute/Elements/PCGCustomHLSL.cpp
--- a/Engine/Plugins/PCG/Source/PCG/Private/Compute/Elements/PCGCustomHLSL.cpp
+++ b/Engine/Plugins/PCG/Source/PCG/Private/Compute/Elements/PCGCustomHLSL.cpp
@@ -650,7 +650,7 @@ void UPCGCustomHLSLSettings::UpdateInputDeclarations()
 			"uint {0}_GetNumData();\n"
 			"// NumElements: Size in XY, ArraySize in Z.\n"
 			"uint3 {0}_GetNumElements(uint DataIndex);\n"
-			"uint2 {0}_GetSize(uint DataIndex);\n"
+			"uint2 {0}_GetResolution(uint DataIndex);\n"
 			"uint {0}_GetArraySize(uint DataIndex);\n"
 			"// Computes an unclamped texture coordinate.\n"
 			"float2 {0}_GetTexCoords(uint DataIndex, float3 WorldPos);\n"
