[PCG] Added option to force normal selection in the create polygon 2d node.

#Synced-CL 49649709

diff --git a/Engine/Plugins/PCG/Source/PCG/Private/Elements/Polygon/PCGCreatePolygon2D.cpp b/Engine/Plugins/PCG/Source/PCG/Private/Elements/Polygon/PCGCreatePolygon2D.cpp
--- a/Engine/Plugins/PCG/Source/PCG/Private/Elements/Polygon/PCGCreatePolygon2D.cpp
+++ b/Engine/Plugins/PCG/Source/PCG/Private/Elements/Polygon/PCGCreatePolygon2D.cpp
@@ -148,9 +148,15 @@ bool FPCGCreatePolygonElement::ExecuteInternal(FPCGContext* Context) const
 		}
 
 		const double LargestDimension = PolygonBounds.GetSize().GetMax();
+		const bool bDegeneratePolygonCase = PolygonArea <= UE_SMALL_NUMBER * FMath::Square(LargestDimension);
 
-		// If the polygon is a point/line or otherwise has a very small area, use the provided default up axis.
-		if (PolygonArea <= UE_SMALL_NUMBER * FMath::Square(LargestDimension))
+		const bool bCreateOpenPolygon = (Settings->InputType == EPCGCreatePolygonInputType::ForceOpen || (bInputIsOpen && Settings->InputType == EPCGCreatePolygonInputType::Automatic));
+
+		// If the polygon is a point/line or otherwise has a very small area, always use the provided default up axis.
+		// Alternatively, use the default polygon normal based on the normal behavior and the polygon type.
+		if (bDegeneratePolygonCase || 
+			(Settings->DefaultPolygonNormalBehavior == EPCGCreatePolygonDefaultNormalBehavior::OnlyForOpenPolygons && bCreateOpenPolygon) ||
+			Settings->DefaultPolygonNormalBehavior == EPCGCreatePolygonDefaultNormalBehavior::Always)
 		{
 			PlaneNormal = Settings->DefaultPolygonNormal.GetSafeNormal(UE_SMALL_NUMBER, FVector::ZAxisVector);
 		}
@@ -173,8 +179,6 @@ bool FPCGCreatePolygonElement::ExecuteInternal(FPCGContext* Context) const
 			Positions2D.Emplace(LocalPosition.X, LocalPosition.Y);
 		}
 
-		const bool bCreateOpenPolygon = (Settings->InputType == EPCGCreatePolygonInputType::ForceOpen || (bInputIsOpen && Settings->InputType == EPCGCreatePolygonInputType::Automatic));
-
 		// Create starting polygon, which will be used as-is if this isn't an open polygon.
 		// 1. Split positions into outer/holes...
 		UE::Geometry::FGeneralPolygon2d SourcePolygon;
diff --git a/Engine/Plugins/PCG/Source/PCG/Public/Elements/Polygon/PCGCreatePolygon2D.h b/Engine/Plugins/PCG/Source/PCG/Public/Elements/Polygon/PCGCreatePolygon2D.h
--- a/Engine/Plugins/PCG/Source/PCG/Public/Elements/Polygon/PCGCreatePolygon2D.h
+++ b/Engine/Plugins/PCG/Source/PCG/Public/Elements/Polygon/PCGCreatePolygon2D.h
@@ -15,6 +15,14 @@ enum class EPCGCreatePolygonInputType : uint8
 	ForceClosed UMETA(Tooltip="Always considers inputs as closed polygons.")
 };
 
+UENUM()
+enum class EPCGCreatePolygonDefaultNormalBehavior : uint8
+{
+	OnlyForDegenerateCases UMETA(Tooltip="Uses the default normal only for polygons that are degenerate (collinear or too small)."),
+	OnlyForOpenPolygons UMETA(Tooltip="Uses the default normal for all open polygons (including forced open polygons) and degenerate polygons."),
+	Always UMETA(Tooltip="Uses the default normal for all polygons.")
+};
+
 class FPCGCreatePolygonElement : public IPCGElement
 {
 protected:
@@ -81,6 +89,10 @@ public:
 	UPROPERTY(EditAnywhere, BlueprintReadWrite, Category = Settings, AdvancedDisplay, meta = (PCG_Overridable))
 	FVector DefaultPolygonNormal = FVector::ZAxisVector;
 
+	/** Controls whether we will use the default polygon normal : always/only for open/only for degenerate data */
+	UPROPERTY(EditAnywhere, BlueprintReadWrite, Category = Settings, AdvancedDisplay, meta = (PCG_Overridable))
+	EPCGCreatePolygonDefaultNormalBehavior DefaultPolygonNormalBehavior = EPCGCreatePolygonDefaultNormalBehavior::OnlyForDegenerateCases;
+
 	/** Controls behavior of how joining points are resolved. See GeometryScript for more details. */
 	UPROPERTY(EditAnywhere, BlueprintReadWrite, Category = Settings, AdvancedDisplay, meta = (EditCondition = "InputType != EPCGCreatePolygonInputType::ForceClosed", PCG_Overridable))
 	EPCGPolygonJoinType JoinType = EPCGPolygonJoinType::Square;
