[PCG] Fix static mesh API decls in Custom HLSL node

#rnx
#Synced-CL 49633484

diff --git a/Engine/Plugins/PCG/Source/PCG/Private/Compute/Elements/PCGCustomHLSL.cpp b/Engine/Plugins/PCG/Source/PCG/Private/Compute/Elements/PCGCustomHLSL.cpp
--- a/Engine/Plugins/PCG/Source/PCG/Private/Compute/Elements/PCGCustomHLSL.cpp
+++ b/Engine/Plugins/PCG/Source/PCG/Private/Compute/Elements/PCGCustomHLSL.cpp
@@ -733,14 +733,14 @@ void UPCGCustomHLSLSettings::UpdateInputDeclarations()
 			"int {0}_GetNumVertices(int DataIndex);\n"
 			"void {0}_GetVertex(int DataIndex, int VertexIndex, out float3 OutPosition, out float3 OutNormal, out float3 OutTangent, out float3 OutBitangent);\n"
 			"float4 {0}_GetVertexColor(int DataIndex, int VertexIndex);\n"
-			"float2 {0}_GetVertexUVs(int DataIndex, int VertexIndex, int UVSet);\n"
+			"float2 {0}_GetVertexUV(int DataIndex, int VertexIndex, int UVSet);\n"
 			"\n"
 			"// Triangle functions\n"
 			"int {0}_GetNumTriangles(int DataIndex);\n"
 			"void {0}_GetTriangleIndices(int DataIndex, int TriangleIndex, out int OutIndex0, out int OutIndex1, out int OutIndex2);\n"
 			"void {0}_SampleTriangle(int DataIndex, int TriangleIndex, float3 BaryCoords, out float3 OutPosition, out float3 OutNormal, out float3 OutTangent, out float3 OutBitangent);\n"
 			"float4 {0}_SampleTriangleColor(int DataIndex, int TriangleIndex, float3 BaryCoords);\n"
-			"float2 {0}_SampleTriangleUVs(int DataIndex, int TriangleIndex, float3 BaryCoords, int UVSet);\n"
+			"float2 {0}_SampleTriangleUV(int DataIndex, int TriangleIndex, float3 BaryCoords, int UVSet);\n"
 			"\n"
 			"// Get bounds extents of the static mesh.\n"
 			"float3 {0}_GetMeshBoundsExtents(int DataIndex);\n"),
