[PCG] Optional PostExecuteConsumerThread (Mpsc) : scheduling contention improvement on the PostTaskExecute code path
CVar: pcg.Editor.EnablePostExecuteConsumerThread / pcg.Runtime.EnablePostExecuteConsumerThread (disabled by default)
CVar: pcg.Editor.PostExecuteConsumerThreadTime/pcg.Runtime.PostExecuteConsumerThreadTime (when consumer thread is enabled)

#Synced-CL 49625772

diff --git a/Engine/Plugins/PCG/Source/PCG/Private/Graph/PCGGraphExecutor.cpp b/Engine/Plugins/PCG/Source/PCG/Private/Graph/PCGGraphExecutor.cpp
--- a/Engine/Plugins/PCG/Source/PCG/Private/Graph/PCGGraphExecutor.cpp
+++ b/Engine/Plugins/PCG/Source/PCG/Private/Graph/PCGGraphExecutor.cpp
@@ -67,6 +67,17 @@ namespace PCGGraphExecutor
 		8,
 		TEXT("Runtime: Hard cap on number of threads for processing scheduled pcg elements calculated from pcg.Runtime.MaxPercentageOfExecutingThreads"));
 
+	// Disabled by default because on lower end platforms it is a performance choice (adding a dedicated thread considering lower number of cores might cause less scheduling contention which this consumer thread tries to alleviate)
+	TAutoConsoleVariable<bool> CVarRuntimeEnablePostExecuteConsumerThread(
+		TEXT("pcg.Runtime.EnablePostExecuteConsumerThread"),
+		false,
+		TEXT("Runtime: When enabled, a single thread will process all post executes and scheduling, instead of doing it directly in the async the executing threads"));
+
+	TAutoConsoleVariable<float> CVarRuntimePostExecuteConsumerThreadTime(
+		TEXT("pcg.Runtime.PostExecuteConsumerThreadTime"),
+		5.0f,
+		TEXT("Runtime: Allocated time the consumer thread will use before yielding (when pcg.Runtime.EnablePostExecuteConsumerThread is enabled)"));
+
 	TAutoConsoleVariable<float> CVarTimePerFrame(
 		TEXT("pcg.FrameTime"),
 		5.0f,
@@ -98,6 +109,16 @@ namespace PCGGraphExecutor
 		8,
 		TEXT("Editor: Hard cap on number of threads for processing scheduled pcg elements calculated from pcg.Editor.MaxPercentageOfExecutingThreads"));
 
+	TAutoConsoleVariable<bool> CVarEditorEnablePostExecuteConsumerThread(
+		TEXT("pcg.Editor.EnablePostExecuteConsumerThread"),
+		false,
+		TEXT("Editor: When enabled, a single thread will process all post executes and scheduling, instead of doing it directly in the async the executing threads"));
+
+	TAutoConsoleVariable<float> CVarEditorPostExecuteConsumerThreadTime(
+		TEXT("pcg.Editor.PostExecuteConsumerThreadTime"),
+		15.0f,
+		TEXT("Editor: Allocated time the consumer thread will use before yielding (when pcg.Editor.EnablePostExecuteConsumerThread is enabled)"));
+
 	TAutoConsoleVariable<float> CVarEditorTimePerFrame(
 		TEXT("pcg.EditorFrameTime"),
 		15.0f,
@@ -358,6 +379,9 @@ FPCGGraphExecutor::~FPCGGraphExecutor()
 	ClearAllTasks();
 	ReleaseGenerationNotification();
 #endif
+
+	// Shutdown thread
+	PostTaskExecuteConsumerThread.Reset();
 }
 
 void FPCGGraphExecutor::Compile(UPCGGraph* Graph)
@@ -706,6 +730,15 @@ TSet<IPCGGraphExecutionSource*> FPCGGraphExecutor::Cancel(TFunctionRef<bool(TWea
 
 				TSharedPtr<FPCGGraphActiveTask> ActiveTask = CancelledActiveTask.Get<TSharedPtr<FPCGGraphActiveTask>>();
 				check(ActiveTask->Element);
+
+				// Even if we've set SchedulingDisabled = true, it is possible some tasks slipped into the queue and will get PostTaskExecute called through consumer thread
+				// if this is the case the bIsExecutingTask would be true and we can wait for the consumer thread to process.
+				while (PostTaskExecuteConsumerThread.IsValid() && ActiveTask->bIsExecutingTask)
+				{
+					// Wait for consumer thread to process...
+					FPlatformProcess::YieldThread();
+				}
+
 				ActiveTask->Element->Abort(ActiveTask->Context.Get());
 
 				bStableCancellationSet &= !CancelNextTasks(ActiveTask->NodeId, CancelledExecutionSources);
@@ -1677,6 +1710,15 @@ double FPCGGraphExecutor::GetTickBudgetInSeconds() const
 #endif
 }
 
+double FPCGGraphExecutor::GetPostExecuteConsumerThreadBudgetInSeconds() const
+{
+#if WITH_EDITOR
+	return PCGGraphExecutor::GetCVarValueForWorld(World, PCGGraphExecutor::CVarRuntimePostExecuteConsumerThreadTime, PCGGraphExecutor::CVarEditorPostExecuteConsumerThreadTime) / 1000.0;
+#else
+	return PCGGraphExecutor::CVarRuntimePostExecuteConsumerThreadTime.GetValueOnAnyThread() / 1000.0;
+#endif
+}
+
 int32 FPCGGraphExecutor::GetMaxNumProcessingTasks() const
 {
 #if WITH_EDITOR
@@ -1713,6 +1755,89 @@ int32 FPCGGraphExecutor::GetMaxExecutingThreads() const
 #endif
 }
 
+bool FPCGGraphExecutor::IsPostExecuteConsumerThreadEnabled() const
+{
+	if (!FApp::ShouldUseThreadingForPerformance())
+	{
+		return false;
+	}
+
+#if WITH_EDITOR
+	return PCGGraphExecutor::GetCVarValueForWorld(World, PCGGraphExecutor::CVarRuntimeEnablePostExecuteConsumerThread, PCGGraphExecutor::CVarEditorEnablePostExecuteConsumerThread);
+#else
+	return PCGGraphExecutor::CVarRuntimeEnablePostExecuteConsumerThread.GetValueOnAnyThread();
+#endif
+}
+
+void FPCGGraphExecutor::ProcessAsyncExecuteResult(TSharedPtr<FPCGGraphActiveTask> Task, bool bIsDone, bool bIsPaused, double EndTime)
+{
+	PostTaskExecute(Task, bIsDone);
+
+	if (bIsDone || bIsPaused)
+	{
+		// If this task is paused it is possible it scheduled some new tasks, process those new tasks before calling ExecuteScheduling (so we are not one frame behind)
+		if (bIsPaused)
+		{
+			ProcessScheduledTasks(EndTime);
+		}
+
+		if (FPlatformTime::Seconds() < EndTime)
+		{
+			// Budget is provided by CVarAsyncOutOfTickBudgetInMilliseconds
+			ExecuteScheduling(EndTime, nullptr,  /*bForceCheckPausedTasks=*/bIsPaused);
+		}
+	}
+}
+
+uint32 FPCGGraphExecutor::FPostTaskExecuteConsumerThread::Run()
+{
+	check(Owner);
+	while (Event->Wait() && !bQuit)
+	{
+		TRACE_CPUPROFILER_EVENT_SCOPE(FPCGGraphExecutor::FPostTaskExecuteConsumerThread::Run);
+		const double EndTime = FPlatformTime::Seconds() + Owner->GetPostExecuteConsumerThreadBudgetInSeconds();
+
+		FPostTaskExecuteItem Item;
+		while (Queue.Dequeue(Item))
+		{
+			Owner->ProcessAsyncExecuteResult(Item.Task, Item.bIsDone, Item.bIsPaused, EndTime);
+
+			// Yield everytime we've exhausted our budget
+			if (FPlatformTime::Seconds() > EndTime)
+			{
+				FPlatformProcess::YieldThread();
+			}
+		}
+	}
+
+	return 0;
+}
+
+FPCGGraphExecutor::FPostTaskExecuteConsumerThread::FPostTaskExecuteConsumerThread(FPCGGraphExecutor* InOwner)
+	: Owner(InOwner)
+{
+	Thread = FRunnableThread::Create(this, TEXT("FPostTaskExecuteConsumerThread"));
+}
+
+FPCGGraphExecutor::FPostTaskExecuteConsumerThread::~FPostTaskExecuteConsumerThread()
+{
+	bQuit = true;
+	Event->Trigger();
+	Thread->Kill(true);
+	delete Thread;
+	Thread = nullptr;
+}
+
+void FPCGGraphExecutor::StartPostExecuteConsumerThread()
+{
+	if (PostTaskExecuteConsumerThread.IsValid() || !IsPostExecuteConsumerThreadEnabled())
+	{
+		return;
+	}
+
+	PostTaskExecuteConsumerThread = MakeUnique<FPostTaskExecuteConsumerThread>(this);
+}
+
 bool FPCGGraphExecutor::ExecuteScheduling(double EndTime, TSharedPtr<FPCGGraphActiveTask>* OutMainThreadTask, bool bForceCheckPausedTasks)
 {
 	// Can currently be disabled by cancellation
@@ -1944,24 +2069,18 @@ bool FPCGGraphExecutor::ExecuteScheduling(double EndTime, TSharedPtr<FPCGGraphAc
 							const bool bIsDone = ActiveTask->bWasCancelled || ActiveTask->Element->Execute(ActiveTask->Context.Get());
 							const bool bIsPaused = ActiveTask->Context->bIsPaused;
 
-							PostTaskExecute(ActiveTask, bIsDone);
-
-							if (bIsDone || bIsPaused)
+							// If Scheduling is disabled it means we are cancelling tasks so in that case avoid using the consumer thread so that the Wait calls in Cancel
+							// can rely on the task having been fully executed (PostTaskExecute included)
+							if (!SchedulingDisabled && PostTaskExecuteConsumerThread.IsValid())
+							{
+								PostTaskExecuteConsumerThread->Queue.Enqueue(ActiveTask, bIsDone, bIsPaused);
+								PostTaskExecuteConsumerThread->Event->Trigger();
+							}
+							else
 							{
 								const double EndTime = FPlatformTime::Seconds() + SchedulingBudgetInSeconds;
-								// If this task is paused it is possible it scheduled some new tasks, process those new tasks before calling ExecuteScheduling (so we are not one frame behind)
-								if (bIsPaused)
-								{
-									ProcessScheduledTasks(EndTime);
-								}
-
-								if (FPlatformTime::Seconds() < EndTime)
-								{
-									// Budget is provided by CVarAsyncOutOfTickBudgetInMilliseconds
-									ExecuteScheduling(EndTime, nullptr,  /*bForceCheckPausedTasks=*/bIsPaused);
-								}
+								ProcessAsyncExecuteResult(ActiveTask, bIsDone, bIsPaused, EndTime);
 							}
-
 							return bIsDone;
 						});
 						bStateChanged = true;
@@ -2053,6 +2172,9 @@ void FPCGGraphExecutor::Execute(double& InOutEndTime)
 {
 	TRACE_CPUPROFILER_EVENT_SCOPE(FPCGGraphExecutor::Execute);
 
+	// Start Thread if enabled
+	StartPostExecuteConsumerThread();
+
 #if WITH_EDITOR
 
 	{
diff --git a/Engine/Plugins/PCG/Source/PCG/Private/Graph/PCGGraphExecutor.h b/Engine/Plugins/PCG/Source/PCG/Private/Graph/PCGGraphExecutor.h
--- a/Engine/Plugins/PCG/Source/PCG/Private/Graph/PCGGraphExecutor.h
+++ b/Engine/Plugins/PCG/Source/PCG/Private/Graph/PCGGraphExecutor.h
@@ -13,6 +13,8 @@
 #include "Graph/PCGGraphTask.h"
 #include "Graph/PCGStackContext.h"
 
+#include "Containers/MpscQueue.h"
+#include "HAL/Runnable.h"
 #include "Tasks/Task.h"
 #include "Templates/UniquePtr.h"
 #include "UObject/GCObject.h"
@@ -190,6 +192,8 @@ private:
 	float GetMaxPercentageOfThreadsToUse() const;
 	float GetMaxPercentageOfExecutingThreads() const;
 	int32 GetMaxExecutingThreads() const;
+	bool IsPostExecuteConsumerThreadEnabled() const;
+	double GetPostExecuteConsumerThreadBudgetInSeconds() const;
 
 	struct FCachedResult
 	{
@@ -217,6 +221,8 @@ private:
 	FPCGTaskId ScheduleSingleElement(TSharedRef<IPCGElement> Element, bool bHasAbortCallbacks, IPCGGraphExecutionSource* InExecutionSource,
 		const TArray<FPCGTaskId>& TaskExecutionDependencies, const TArray<FPCGTaskId>& TaskDataDependencies);
 
+	void ProcessAsyncExecuteResult(TSharedPtr<FPCGGraphActiveTask> Task, bool bIsDone, bool bIsPaused, double EndTime);
+	void StartPostExecuteConsumerThread();
 	void PostTaskExecute(TSharedPtr<FPCGGraphActiveTask> ActiveTask, bool bIsDone);
 	bool ProcessScheduledTasks(double EndTime);
 	void ExecuteTasksEnded();
@@ -370,6 +376,30 @@ private:
 	std::atomic<bool> bNeedToExecuteTasksEnded = false;
 	std::atomic<int32> NewReadyTaskCount = 0;
 
+	struct FPostTaskExecuteItem
+	{
+		TSharedPtr<FPCGGraphActiveTask> Task;
+		bool bIsDone = false;
+		bool bIsPaused = false;
+	};
+
+	struct FPostTaskExecuteConsumerThread : public FRunnable
+	{
+		FPCGGraphExecutor* Owner = nullptr;
+
+		FRunnableThread* Thread = nullptr;
+		TMpscQueue<FPostTaskExecuteItem> Queue;
+		FEventRef Event{ EEventMode::AutoReset };
+		bool bQuit = false;
+
+		FPostTaskExecuteConsumerThread(FPCGGraphExecutor* InOwner);
+		virtual ~FPostTaskExecuteConsumerThread();
+
+		virtual uint32 Run() override;
+	};
+
+	TUniquePtr<FPostTaskExecuteConsumerThread> PostTaskExecuteConsumerThread;
+
 #if WITH_EDITOR
 	TWeakPtr<IPCGEditorProgressNotification> GenerationProgressNotification;
 	double GenerationProgressNotificationStartTime = 0.0;
