[PCG] PCGComponent: Always call UnregisterExecutionSource if we are being unregistered through unloading

#rnx
#Synced-CL 49625772

diff --git a/Engine/Plugins/PCG/Source/PCG/Private/PCGComponent.cpp b/Engine/Plugins/PCG/Source/PCG/Private/PCGComponent.cpp
--- a/Engine/Plugins/PCG/Source/PCG/Private/PCGComponent.cpp
+++ b/Engine/Plugins/PCG/Source/PCG/Private/PCGComponent.cpp
@@ -1707,8 +1707,8 @@ void UPCGComponent::OnUnregister()
 
 		Subsystem->OnPCGComponentUnregistered.Broadcast(IsValidChecked(this) ? this : nullptr);
 
-		// If Original Component is managed by runtime and being unloaded, unregister immediatly so that runtime generated Local components also get cleaned up. (BeginDestroy is too late)
-		if (bUnregisteredThroughLoading && IsManagedByRuntimeGenSystem() && !IsLocalComponent())
+		// Always unregister here if we are being unloaded (we are not being re-registered)
+		if (bUnregisteredThroughLoading)
 		{
 			Subsystem->UnregisterExecutionSource(this, /*bForce=*/true);
 		}
