[PCG] Change checks to compare metadata domain roots to checkSlow (was consuming 15% of total execution time in test map ~20 seconds / 140 seconds)

#Synced-CL 49695653

diff --git a/Engine/Plugins/PCG/Source/PCG/Private/Metadata/PCGMetadata.cpp b/Engine/Plugins/PCG/Source/PCG/Private/Metadata/PCGMetadata.cpp
--- a/Engine/Plugins/PCG/Source/PCG/Private/Metadata/PCGMetadata.cpp
+++ b/Engine/Plugins/PCG/Source/PCG/Private/Metadata/PCGMetadata.cpp
@@ -941,8 +941,11 @@ FPCGMetadataAttributeBase* UPCGMetadata::CopyAttribute(const FPCGMetadataAttribu
 	check(OriginalAttribute);
 	const FPCGMetadataDomain* OriginalMetadataDomain = OriginalAttribute->GetMetadataDomain();
 	check(OriginalMetadataDomain);
+
+#if DO_GUARD_SLOW
 	const UPCGMetadata* RootMetadata = GetRoot();
-	check((RootMetadata && RootMetadata->MetadataDomains.Contains(OriginalMetadataDomain->DomainID) && RootMetadata->MetadataDomains[OriginalMetadataDomain->DomainID].Get() == OriginalMetadataDomain->GetRoot()) || !bKeepParent);
+	checkSlow((RootMetadata && RootMetadata->MetadataDomains.Contains(OriginalMetadataDomain->DomainID) && RootMetadata->MetadataDomains[OriginalMetadataDomain->DomainID].Get() == OriginalMetadataDomain->GetRoot()) || !bKeepParent);
+#endif
 
 	FPCGMetadataAttributeBase* Attribute = WithMetadataDomain_Lambda(NewAttributeName.MetadataDomain, [&](FPCGMetadataDomain* MetadataDomain) -> FPCGMetadataAttributeBase*
 	{
diff --git a/Engine/Plugins/PCG/Source/PCG/Private/Metadata/PCGMetadataDomain.cpp b/Engine/Plugins/PCG/Source/PCG/Private/Metadata/PCGMetadataDomain.cpp
--- a/Engine/Plugins/PCG/Source/PCG/Private/Metadata/PCGMetadataDomain.cpp
+++ b/Engine/Plugins/PCG/Source/PCG/Private/Metadata/PCGMetadataDomain.cpp
@@ -824,7 +824,7 @@ FPCGMetadataAttributeBase* FPCGMetadataDomain::CopyAttribute(FName AttributeToCo
 FPCGMetadataAttributeBase* FPCGMetadataDomain::CopyAttribute(const FPCGMetadataAttributeBase* OriginalAttribute, FName NewAttributeName, bool bKeepParent, bool bCopyEntries, bool bCopyValues)
 {
 	check(OriginalAttribute);
-	check(OriginalAttribute->GetMetadataDomain()->GetRoot() == GetRoot() || !bKeepParent);
+	checkSlow(OriginalAttribute->GetMetadataDomain()->GetRoot() == GetRoot() || !bKeepParent);
 	FPCGMetadataAttributeBase* NewAttribute = OriginalAttribute->Copy(NewAttributeName, this, bKeepParent, bCopyEntries, bCopyValues);
 
 	if (NewAttribute)
diff --git a/Engine/Plugins/PCG/Source/PCG/Public/Metadata/PCGMetadataAttributeTpl.h b/Engine/Plugins/PCG/Source/PCG/Public/Metadata/PCGMetadataAttributeTpl.h
--- a/Engine/Plugins/PCG/Source/PCG/Public/Metadata/PCGMetadataAttributeTpl.h
+++ b/Engine/Plugins/PCG/Source/PCG/Public/Metadata/PCGMetadataAttributeTpl.h
@@ -239,7 +239,7 @@ protected:
 		// In that case, we will copy all the data from this attribute and all its ancestors.
 
 		// We can't keep the parent if we don't have the same root.
-		check(!bKeepParent || PCGMetadataHelpers::HasSameRoot(Metadata, InMetadata));
+		checkSlow(!bKeepParent || PCGMetadataHelpers::HasSameRoot(Metadata, InMetadata));
 
 		// Validate that the new name is valid
 		if (!IsValidName(NewName))
