[PCG] SplineData: add missing fast crc path

#Synced-CL 49695653

diff --git a/Engine/Plugins/PCG/Source/PCG/Private/Data/PCGSplineData.cpp b/Engine/Plugins/PCG/Source/PCG/Private/Data/PCGSplineData.cpp
--- a/Engine/Plugins/PCG/Source/PCG/Private/Data/PCGSplineData.cpp
+++ b/Engine/Plugins/PCG/Source/PCG/Private/Data/PCGSplineData.cpp
@@ -118,6 +118,13 @@ void UPCGSplineData::AddToCrc(FArchiveCrc32& Ar, bool bFullDataCrc) const
 {
 	Super::AddToCrc(Ar, bFullDataCrc);
 
+	if (!bFullDataCrc)
+	{
+		// Fallback to UID
+		AddUIDToCrc(Ar);
+		return;
+	}
+
 	if (Metadata)
 	{
 		Metadata->AddToCrc(Ar, bFullDataCrc);
