[PCG] Small performance levers:
- Removed CRC computation when caching is disabled
- Added a state variable to prepare data in the execute to prevent some reentrance
- Added a cvar to turn off data validation

#rnx
#Synced-CL 49682793

diff --git a/Engine/Plugins/PCG/Source/PCG/Private/Elements/PCGStaticMeshSpawner.cpp b/Engine/Plugins/PCG/Source/PCG/Private/Elements/PCGStaticMeshSpawner.cpp
--- a/Engine/Plugins/PCG/Source/PCG/Private/Elements/PCGStaticMeshSpawner.cpp
+++ b/Engine/Plugins/PCG/Source/PCG/Private/Elements/PCGStaticMeshSpawner.cpp
@@ -213,7 +213,18 @@ bool FPCGStaticMeshSpawnerElement::PrepareDataInternal(FPCGContext* InContext) c
 		// Compute CRC if it has not been computed (it likely isn't, but this is to futureproof this)
 		if (!Context->DependenciesCrc.IsValid())
 		{
-			GetDependenciesCrc(FPCGGetDependenciesCrcParams(&Context->InputData, Settings, Context->ExecutionSource.Get()), Context->DependenciesCrc);
+			FPCGDataCollection TempCollection;
+			const FPCGDataCollection* InputDataCollectionPtr = &Context->InputData;
+
+			// If the collection didn't have any Data crc information yet, we need to compute it, but can't do it on the input directly.
+			if (Context->InputData.DataCrcs.Num() != Context->InputData.TaggedData.Num())
+			{
+				TempCollection = Context->InputData;
+				TempCollection.ComputeCrcs(ShouldComputeFullOutputDataCrc(Context));
+				InputDataCollectionPtr = &TempCollection;
+			}
+
+			GetDependenciesCrc(FPCGGetDependenciesCrcParams(InputDataCollectionPtr, Settings, Context->ExecutionSource.Get()), Context->DependenciesCrc);
 		}
 		
 		if (Context->DependenciesCrc.IsValid())
diff --git a/Engine/Plugins/PCG/Source/PCG/Private/Graph/PCGGraphExecutor.cpp b/Engine/Plugins/PCG/Source/PCG/Private/Graph/PCGGraphExecutor.cpp
--- a/Engine/Plugins/PCG/Source/PCG/Private/Graph/PCGGraphExecutor.cpp
+++ b/Engine/Plugins/PCG/Source/PCG/Private/Graph/PCGGraphExecutor.cpp
@@ -2594,16 +2594,13 @@ void FPCGGraphExecutor::BuildTaskInput(FPCGGraphTask& Task, FTaskIds& ResultsToM
 				if (!InputDataOnPin.IsEmpty())
 				{
 					// Proceed carefully when adding data items - if pin is single-data, only add first item.
-					if (!ensure(InputDataOnPin.Num() == InputDataCrcsOnPin.Num()))
-					{
-						InputDataCrcsOnPin.SetNumZeroed(InputDataOnPin.Num());
-					}
+					check(InputDataCrcsOnPin.IsEmpty() || InputDataOnPin.Num() == InputDataCrcsOnPin.Num());
 
 					const int NumberDataItemsToTake = bAllowMultipleData ? InputDataOnPin.Num() : 1;
 
 					Task.TaskInput.AddDataForPin(
 						MakeArrayView(InputDataOnPin.GetData(), NumberDataItemsToTake),
-						MakeArrayView(InputDataCrcsOnPin.GetData(), NumberDataItemsToTake),
+						MakeArrayView(InputDataCrcsOnPin.GetData(), FMath::Min(NumberDataItemsToTake, InputDataCrcsOnPin.Num())),
 						InputPinLabelCrc);
 
 					if (NumberDataItemsToTake < InputDataOnPin.Num())
diff --git a/Engine/Plugins/PCG/Source/PCG/Private/PCGElement.cpp b/Engine/Plugins/PCG/Source/PCG/Private/PCGElement.cpp
--- a/Engine/Plugins/PCG/Source/PCG/Private/PCGElement.cpp
+++ b/Engine/Plugins/PCG/Source/PCG/Private/PCGElement.cpp
@@ -25,6 +25,11 @@
 
 #define LOCTEXT_NAMESPACE "PCGElement"
 
+static TAutoConsoleVariable<bool> CVarPCGValidatePinAndTypes(
+	TEXT("pcg.debug.ValidateDataPinsAndTypes"),
+	true,
+	TEXT("Controls whether we validate that output data from nodes are assigned to valid pins and have proper types for their pins."));
+
 static TAutoConsoleVariable<bool> CVarPCGValidatePointMetadata(
 	TEXT("pcg.debug.ValidatePointMetadata"),
 	true,
@@ -378,6 +383,12 @@ void IPCGElement::PreExecutePrimaryLoopElement(FPCGContext* Context, const UPCGS
 		return;
 	}
 
+	// No need to split data if we're not using the cache.
+	if (!Context->IsCacheEnabled())
+	{
+		return;
+	}
+
 	// Mark inputs in the order they're presented so we can appropriately find the relation from output to input after the execution
 	// TODO: this is not sufficient to do a proper mapping from output to input when we have a cartesian loop
 	for (int32 DataIndex = 0; DataIndex < Context->InputData.TaggedData.Num(); ++DataIndex)
@@ -385,12 +396,6 @@ void IPCGElement::PreExecutePrimaryLoopElement(FPCGContext* Context, const UPCGS
 		Context->InputData.TaggedData[DataIndex].OriginalIndex = DataIndex;
 	}
 
-	// No need to split data if we're not using the cache.
-	if (!Context->IsCacheEnabled())
-	{
-		return;
-	}
-
 	TArray<FPCGDataCollection> PrimaryDataCollections;
 	FPCGDataCollection OtherData;
 	if (!PCGElementHelpers::SplitDataPerPrimaryPin(Settings, Context->InputData, ExecutionLoopMode(Settings), PrimaryDataCollections, OtherData))
@@ -451,8 +456,14 @@ bool IPCGElement::PrepareData(FPCGContext* Context) const
 
 bool IPCGElement::ConvertInputsIfNeeded(FPCGContext* Context) const
 {
+	if (Context->bConvertInputsIfNeededDone)
+	{
+		return true;
+	}
+
 	if (Context->InputData.TaggedData.IsEmpty())
 	{
+		Context->bConvertInputsIfNeededDone = true;
 		return true;
 	}
 
@@ -536,6 +547,7 @@ bool IPCGElement::ConvertInputsIfNeeded(FPCGContext* Context) const
 		}
 	}
 
+	Context->bConvertInputsIfNeededDone = true;
 	return true;
 }
 
@@ -569,47 +581,50 @@ void IPCGElement::PostExecute(FPCGContext* Context) const
 		}
 	}
 
+	if (Context->IsCacheEnabled())
 	{
-		TRACE_CPUPROFILER_EVENT_SCOPE(IPCGElement::PostExecute::CRC);
+		{
+			TRACE_CPUPROFILER_EVENT_SCOPE(IPCGElement::PostExecute::CRC);
 
-		// Some nodes benefit from computing an actual CRC from the data. This can halt the propagation of change/executions through the graph. For
-		// data like landscapes we will never have a full accurate data crc for it so we'll tend to assume changed which triggers downstream
-		// execution. Performing a detailed CRC of output data can detect real change in the data and halt the cascade of execution.
-		const bool bShouldComputeFullOutputDataCrc = ShouldComputeFullOutputDataCrc(Context);
+			// Some nodes benefit from computing an actual CRC from the data. This can halt the propagation of change/executions through the graph. For
+			// data like landscapes we will never have a full accurate data crc for it so we'll tend to assume changed which triggers downstream
+			// execution. Performing a detailed CRC of output data can detect real change in the data and halt the cascade of execution.
+			const bool bShouldComputeFullOutputDataCrc = ShouldComputeFullOutputDataCrc(Context);
 
-		// Compute Crc from output data which will include output pin labels.
-		Context->OutputData.ComputeCrcs(bShouldComputeFullOutputDataCrc);
-	}
+			// Compute Crc from output data which will include output pin labels.
+			Context->OutputData.ComputeCrcs(bShouldComputeFullOutputDataCrc);
+		}
 
 #if WITH_EDITOR
-	const bool bHasErrorsOrWarnings = Context->Node && Context->HasVisualLogs();
+		const bool bHasErrorsOrWarnings = Context->Node && Context->HasVisualLogs();
 #else
-	const bool bHasErrorsOrWarnings = false;
+		const bool bHasErrorsOrWarnings = false;
 #endif
 
-	{
-		TRACE_CPUPROFILER_EVENT_SCOPE(IPCGElement::PostExecute::StoreInCache);
-		// Store result in cache
-		// TODO - There is a potential mismatch here between using the Settings (incl. overrides) and the input settings interface (pre-overrides), as done in the graph executor.
-		// TODO - The dependencies CRC here should always be valid except in the indirection case, which we should normalize to allow caching (tested here, otherwise it can ensure in the graph cache)
-		if (!Context->OutputData.bCancelExecution && !bHasErrorsOrWarnings && Context->DependenciesCrc.IsValid() && IsCacheableInstance(Settings))
 		{
-			bool bCacheable = true;
-
-			// GPU proxies are never cached. Caching them would hold on to precious graphics memory.
-			for (const FPCGTaggedData& TaggedData : Context->OutputData.TaggedData)
+			TRACE_CPUPROFILER_EVENT_SCOPE(IPCGElement::PostExecute::StoreInCache);
+			// Store result in cache
+			// TODO - There is a potential mismatch here between using the Settings (incl. overrides) and the input settings interface (pre-overrides), as done in the graph executor.
+			// TODO - The dependencies CRC here should always be valid except in the indirection case, which we should normalize to allow caching (tested here, otherwise it can ensure in the graph cache)
+			if (!Context->OutputData.bCancelExecution && !bHasErrorsOrWarnings && Context->DependenciesCrc.IsValid() && IsCacheableInstance(Settings))
 			{
-				if (TaggedData.Data && !TaggedData.Data->IsCacheable())
+				bool bCacheable = true;
+
+				// GPU proxies are never cached. Caching them would hold on to precious graphics memory.
+				for (const FPCGTaggedData& TaggedData : Context->OutputData.TaggedData)
 				{
-					bCacheable = false;
-					break;
+					if (TaggedData.Data && !TaggedData.Data->IsCacheable())
+					{
+						bCacheable = false;
+						break;
+					}
 				}
-			}
 
-			if (bCacheable)
-			{
-				FPCGStoreInCacheParams Params = { .Element = this, .Crc = Context->DependenciesCrc };
-				Context->StoreInCache(Params, Context->OutputData);
+				if (bCacheable)
+				{
+					FPCGStoreInCacheParams Params = { .Element = this, .Crc = Context->DependenciesCrc };
+					Context->StoreInCache(Params, Context->OutputData);
+				}
 			}
 		}
 	}
@@ -1021,55 +1036,63 @@ void IPCGElement::CleanupAndValidateOutput(FPCGContext* Context) const
 				if (!TaggedData.bPinlessData)
 				{
 					TaggedData.Pin = OutputPinProperties[0].Label;
-				}				
+				}
 			}
 		}
 
 		// Validate all out data for errors in labels
 #if WITH_EDITOR
-		if (SettingsInterface->bEnabled)
-		{
-			// remove null outputs
-			Context->OutputData.TaggedData.RemoveAll([this, Context](const FPCGTaggedData& TaggedData){
+		const bool bNeedsPinsAndTypesValidation = CVarPCGValidatePinAndTypes.GetValueOnAnyThread();
+		const bool bNeedsMetadataValidation =  CVarPCGValidatePointMetadata.GetValueOnAnyThread();
 
-				if (TaggedData.Data == nullptr)
-				{
-					PCGE_LOG(Error, GraphAndLog, FText::Format(LOCTEXT("NullPinOutputData", "Invalid output(s) generated for pin '{0}'"), FText::FromName(TaggedData.Pin)));
-					return true;
-				}
+		if (SettingsInterface->bEnabled && (bNeedsPinsAndTypesValidation || bNeedsMetadataValidation))
+		{
+			if (bNeedsPinsAndTypesValidation)
+			{
+				// remove null outputs
+				Context->OutputData.TaggedData.RemoveAll([this, Context](const FPCGTaggedData& TaggedData) {
 
-				return false;
-			});
+					if (TaggedData.Data == nullptr)
+					{
+						PCGE_LOG(Error, GraphAndLog, FText::Format(LOCTEXT("NullPinOutputData", "Invalid output(s) generated for pin '{0}'"), FText::FromName(TaggedData.Pin)));
+						return true;
+					}
 
+					return false;
+					});
+			}
 
 			for (FPCGTaggedData& TaggedData : Context->OutputData.TaggedData)
 			{
-				const int32 MatchIndex = OutputPinProperties.IndexOfByPredicate([&TaggedData](const FPCGPinProperties& InProp) { return TaggedData.Pin == InProp.Label; });
-				if (MatchIndex == INDEX_NONE)
+				if (bNeedsPinsAndTypesValidation)
 				{
-					// Only display an error if we expected this data to have a pin.
-					if (!TaggedData.bPinlessData)
+					const int32 MatchIndex = OutputPinProperties.IndexOfByPredicate([&TaggedData](const FPCGPinProperties& InProp) { return TaggedData.Pin == InProp.Label; });
+					if (MatchIndex == INDEX_NONE)
 					{
-						PCGE_LOG(Error, GraphAndLog, FText::Format(LOCTEXT("OutputCannotBeRouted", "Output data generated for non-existent output pin '{0}'"), FText::FromName(TaggedData.Pin)));
+						// Only display an error if we expected this data to have a pin.
+						if (!TaggedData.bPinlessData)
+						{
+							PCGE_LOG(Error, GraphAndLog, FText::Format(LOCTEXT("OutputCannotBeRouted", "Output data generated for non-existent output pin '{0}'"), FText::FromName(TaggedData.Pin)));
+						}
 					}
-				}
-				else if (ensure(TaggedData.Data))
-				{
-					// Try to get dynamic current pin types, otherwise settle for static types
-					const UPCGPin* OutputPin = Context->Node ? Context->Node->GetOutputPin(OutputPinProperties[MatchIndex].Label) : nullptr;
-					const FPCGDataTypeIdentifier PinTypes = OutputPin ? OutputPin->GetCurrentTypesID() : OutputPinProperties[MatchIndex].AllowedTypes;
-
-					FPCGDataTypeIdentifier DataType = TaggedData.Data->GetUnderlyingDataTypeId();
-					const bool bTypesOverlap = !!(PinTypes & DataType);
-					const bool bTypeIsSubset = !DataType.IsWider(PinTypes);
-					// TODO: Temporary fix for Settings directly from InputData (ie. from elements with code and not PCG nodes)
-					if ((!bTypesOverlap || !bTypeIsSubset) && DataType != EPCGDataType::Settings)
+					else if (ensure(TaggedData.Data))
 					{
-						PCGE_LOG(Warning, GraphAndLog, FText::Format(LOCTEXT("OutputIncompatibleType", "Output data generated for pin '{0}' does not have a compatible type: '{1}'. Consider using more specific/narrower input pin types, or more general/wider output pin types."), FText::FromName(TaggedData.Pin), DataType.ToDisplayText()));
+						// Try to get dynamic current pin types, otherwise settle for static types
+						const UPCGPin* OutputPin = Context->Node ? Context->Node->GetOutputPin(OutputPinProperties[MatchIndex].Label) : nullptr;
+						const FPCGDataTypeIdentifier PinTypes = OutputPin ? OutputPin->GetCurrentTypesID() : OutputPinProperties[MatchIndex].AllowedTypes;
+
+						FPCGDataTypeIdentifier DataType = TaggedData.Data->GetUnderlyingDataTypeId();
+						const bool bTypesOverlap = !!(PinTypes & DataType);
+						const bool bTypeIsSubset = !DataType.IsWider(PinTypes);
+						// TODO: Temporary fix for Settings directly from InputData (ie. from elements with code and not PCG nodes)
+						if ((!bTypesOverlap || !bTypeIsSubset) && DataType != EPCGDataType::Settings)
+						{
+							PCGE_LOG(Warning, GraphAndLog, FText::Format(LOCTEXT("OutputIncompatibleType", "Output data generated for pin '{0}' does not have a compatible type: '{1}'. Consider using more specific/narrower input pin types, or more general/wider output pin types."), FText::FromName(TaggedData.Pin), DataType.ToDisplayText()));
+						}
 					}
 				}
 
-				if (CVarPCGValidatePointMetadata.GetValueOnAnyThread())
+				if (bNeedsMetadataValidation)
 				{
 					if (const UPCGBasePointData* PointData = Cast<UPCGBasePointData>(TaggedData.Data))
 					{
diff --git a/Engine/Plugins/PCG/Source/PCG/Public/PCGContext.h b/Engine/Plugins/PCG/Source/PCG/Public/PCGContext.h
--- a/Engine/Plugins/PCG/Source/PCG/Public/PCGContext.h
+++ b/Engine/Plugins/PCG/Source/PCG/Public/PCGContext.h
@@ -134,6 +134,9 @@ PRAGMA_ENABLE_DEPRECATION_WARNINGS
 	TSet<FPCGTaskId> DynamicDependencies;
 
 	EPCGExecutionPhase CurrentPhase = EPCGExecutionPhase::NotExecuted;
+	/** Tracks whether the input conversion step has been done (inside of the PrepareData phase). */
+	bool bConvertInputsIfNeededDone : 1 = false;
+
 	TArray<TPair<FPCGDataCollection, FPCGDataCollection>> CachedInputToOutputInternalResults;
 
 	/** Get the current call stack. */
diff --git a/Engine/Plugins/PCG/Source/PCG/Public/PCGData.h b/Engine/Plugins/PCG/Source/PCG/Public/PCGData.h
--- a/Engine/Plugins/PCG/Source/PCG/Public/PCGData.h
+++ b/Engine/Plugins/PCG/Source/PCG/Public/PCGData.h
@@ -252,7 +252,7 @@ struct FPCGDataCollection
 	template<typename AllocatorType1, typename AllocatorType2>
 	void GetInputsAndCrcsByPin(const FName& InPinLabel, TArray<FPCGTaggedData, AllocatorType1>& OutData, TArray<FPCGCrc, AllocatorType2>& OutDataCrcs) const
 	{
-		if (!ensure(TaggedData.Num() == DataCrcs.Num()))
+		if (!DataCrcs.IsEmpty() && !ensure(TaggedData.Num() == DataCrcs.Num()))
 		{
 			// CRCs are not up to date. Error recovery - add 0 CRCs.
 			const_cast<TArray<FPCGCrc>&>(DataCrcs).SetNumZeroed(TaggedData.Num());
@@ -263,7 +263,11 @@ struct FPCGDataCollection
 			if (ensure(TaggedData[I].Data) && TaggedData[I].Pin == InPinLabel)
 			{
 				OutData.Add(TaggedData[I]);
-				OutDataCrcs.Add(DataCrcs[I]);
+
+				if (!DataCrcs.IsEmpty())
+				{
+					OutDataCrcs.Add(DataCrcs[I]);
+				}
 			}
 		}
 	}
