[PCG] Fixed proper off-GT allocation for user parameter data.

#rnx
#Synced-CL 49649709

diff --git a/Engine/Plugins/PCG/Source/PCG/Private/Helpers/PCGPropertyHelpers.cpp b/Engine/Plugins/PCG/Source/PCG/Private/Helpers/PCGPropertyHelpers.cpp
--- a/Engine/Plugins/PCG/Source/PCG/Private/Helpers/PCGPropertyHelpers.cpp
+++ b/Engine/Plugins/PCG/Source/PCG/Private/Helpers/PCGPropertyHelpers.cpp
@@ -449,7 +449,7 @@ UPCGParamData* PCGPropertyHelpers::ExtractPropertyAsAttributeSet(const PCGProper
 		if (!ParamData)
 		{
 			// From there, we should be able to create the data.
-			ParamData = NewObject<UPCGParamData>();
+			ParamData = FPCGContext::NewObject_AnyThread<UPCGParamData>(OptionalContext);
 			check(!Metadata);
 			Metadata = ParamData->MutableMetadata();
 			check(Metadata);
