[PCG] Fix edge case in the Spline Sampler Interior where the ray tracing intersect on the spline vertex

#rnx
#Synced-CL 49758339

diff --git a/Engine/Plugins/PCG/Source/PCG/Private/Elements/PCGSplineSampler.cpp b/Engine/Plugins/PCG/Source/PCG/Private/Elements/PCGSplineSampler.cpp
--- a/Engine/Plugins/PCG/Source/PCG/Private/Elements/PCGSplineSampler.cpp
+++ b/Engine/Plugins/PCG/Source/PCG/Private/Elements/PCGSplineSampler.cpp
@@ -36,8 +36,11 @@ namespace PCGSplineSamplerHelpers
 
 		TArray<FVector2D> IntersectionPoints;
 
-		// Taken from TPolygon2::FindIntersections
+		// Taken from TPolygon2::FindIntersections, with a slight modification to handle the edge case where the segment intersects the polygon
+		// exactly at one of the polygon vertices
 		const UE::Geometry::FSegment2d InSegment{SegmentStart, SegmentEnd};
+		TArray<FVector2D> PolygonVerticesIntersects;
+		
 		for (UE::Geometry::FSegment2d PolygonSegment : InPolygon.Segments())
 		{
 			if (PolygonSegment.Intersects(InSegment))
@@ -45,10 +48,42 @@ namespace PCGSplineSamplerHelpers
 				UE::Geometry::FIntrSegment2Segment2d Intersection(PolygonSegment, InSegment);
 				if (Intersection.Find())
 				{
-					IntersectionPoints.Add(Intersection.Point0);
+					TArray<FVector2D, TFixedAllocator<2>> TentativePoints = {Intersection.Point0};
 					if (Intersection.Quantity == 2)
 					{
-						IntersectionPoints.Add(Intersection.Point1);
+						TentativePoints.Add(Intersection.Point1);
+					}
+
+					// For each tentative point, check if the intersection is a vertex. If so, discard it if we already found it.
+					// Make sure to handle collocated vertices.
+					TArray<FVector2D, TFixedAllocator<2>> PolygonVertices = {PolygonSegment.StartPoint()};
+					if (!FMath::IsNearlyZero(PolygonSegment.Extent))
+					{
+						PolygonVertices.Add(PolygonSegment.EndPoint());
+					}
+					
+					for (const FVector2D& TentativePoint : TentativePoints)
+					{
+						bool bDiscardPoint = false;
+						for (const FVector2D& PolygonVertex : PolygonVertices)
+						{
+							if (PolygonVertex.Equals(TentativePoint))
+							{
+								if (PolygonVerticesIntersects.Contains(PolygonVertex))
+								{
+									bDiscardPoint = true;
+								}
+								else
+								{
+									PolygonVerticesIntersects.Add(PolygonVertex);
+								}
+							}
+						}
+
+						if (!bDiscardPoint)
+						{
+							IntersectionPoints.Add(TentativePoint);
+						}
 					}
 				}
 			}
