[PCG] Properly managed parent assignment when copying metadata domains. Added validation on post-execute.

#rnx
#Synced-CL 49761097

diff --git a/Engine/Plugins/PCG/Source/PCG/Private/Metadata/PCGMetadata.cpp b/Engine/Plugins/PCG/Source/PCG/Private/Metadata/PCGMetadata.cpp
--- a/Engine/Plugins/PCG/Source/PCG/Private/Metadata/PCGMetadata.cpp
+++ b/Engine/Plugins/PCG/Source/PCG/Private/Metadata/PCGMetadata.cpp
@@ -568,6 +568,10 @@ void UPCGMetadata::InitializeAsCopy(const FPCGMetadataInitializeParams& InParams
 			{
 				bShouldParent = true;
 			}
+			else
+			{
+				CurrentMetadataDomain.Parent = nullptr;
+			}
 		};
 		
 		FindFixOrCreateDomainInitializeParams(InParams, DomainID, OtherMetadataDomain, std::move(InitializeFunc));
diff --git a/Engine/Plugins/PCG/Source/PCG/Private/PCGElement.cpp b/Engine/Plugins/PCG/Source/PCG/Private/PCGElement.cpp
--- a/Engine/Plugins/PCG/Source/PCG/Private/PCGElement.cpp
+++ b/Engine/Plugins/PCG/Source/PCG/Private/PCGElement.cpp
@@ -1112,6 +1112,25 @@ void IPCGElement::CleanupAndValidateOutput(FPCGContext* Context) const
 							PCGE_LOG(Warning, GraphAndLog, FText::Format(LOCTEXT("OutputMissingPointMetadata", "Output generated for pin '{0}' does not have valid point metadata"), FText::FromName(TaggedData.Pin)));
 						}
 					}
+
+					if (const UPCGMetadata* Metadata = TaggedData.Data.Get() ? TaggedData.Data->ConstMetadata() : nullptr)
+					{
+						const bool bMetadataHasParent = (Metadata->GetParent() != nullptr);
+
+						TArray<FPCGMetadataDomainID> DomainIDs = TaggedData.Data->GetAllSupportedMetadataDomainIDs();
+						for (FPCGMetadataDomainID DomainID : DomainIDs)
+						{
+							if (const FPCGMetadataDomain* Domain = Metadata->GetConstMetadataDomain(DomainID))
+							{
+								const bool bDomainHasParent = (Domain->GetParent() != nullptr);
+
+								if (bDomainHasParent && !bMetadataHasParent)
+								{
+									PCGE_LOG(Error, GraphAndLog, FText::Format(LOCTEXT("InvalidDomainParenting", "Invalid parenting on the domain '{0}'."), FText::FromName(DomainID.DebugName)));
+								}
+							}
+						}
+					}
 				}
 			}
 		}
