[PCG] Static analysis fix

#rnx

#Synced-CL 49758339

diff --git a/Engine/Plugins/PCG/Source/PCG/Public/Metadata/Accessors/PCGSplineAccessor.h b/Engine/Plugins/PCG/Source/PCG/Public/Metadata/Accessors/PCGSplineAccessor.h
--- a/Engine/Plugins/PCG/Source/PCG/Public/Metadata/Accessors/PCGSplineAccessor.h
+++ b/Engine/Plugins/PCG/Source/PCG/Public/Metadata/Accessors/PCGSplineAccessor.h
@@ -389,7 +389,7 @@ consteval EPCGSplineDataProperties DeprecationSplineAccessorTarget()
 	else
 	{
 		// Pitfall assert
-		static_assert(Target != Target);
+		static_assert(!std::is_same_v<decltype(Target), EPCGSplineAccessorTarget>);
 		return {};
 	}
 }
