[PCG] Generate Landscape Textures: Respect the grass type inclusion list if one is provided.

* Only run grass exporter on the selected grass types rather than all.
* Always produce an output for each grass type, so that output data has well definied indices.

[RN] [PCG] Added setting on Generate Landscape Textures and Generate Grass Maps nodes to output texture data for every grass type in the inclusion list regardless of whether it is present on the landscape or not. Helps produce consistent output to allow easy indexing into the resulting data.
#Synced-CL 49761898

diff --git a/Engine/Plugins/PCG/Shaders/Private/PCGGrassMapUnpackerCS.usf b/Engine/Plugins/PCG/Shaders/Private/PCGGrassMapUnpackerCS.usf
--- a/Engine/Plugins/PCG/Shaders/Private/PCGGrassMapUnpackerCS.usf
+++ b/Engine/Plugins/PCG/Shaders/Private/PCGGrassMapUnpackerCS.usf
@@ -51,7 +51,8 @@ RWTexture2D<float> OutUnpackedHeight;
 
 StructuredBuffer<int> InComponentHeightOffsets;
 StructuredBuffer<int2> InComponentGrassmapOffsetAndChannels;
-int InNumGlobalGrassTypes;
+StructuredBuffer<int> InOutputIndexToExportIndex;
+int InNumExportedGrassTypes;
 
 int4 InLinearTileIndexToComponentIndex[PCG_MAX_NUM_LANDSCAPE_COMPONENTS]; // Only x component used, 16 byte alignment on arrays
 uint2 InNumTiles;
@@ -63,16 +64,16 @@ uint4 InWindow; // xy == Window start coord, zw = Window end coord
 
 // Dispatch group count is ((InWindow.zw - InWindow.xy + 1) / 8, Max(1, NumGrassMaps)).
 [numthreads(THREADGROUPSIZE_X, THREADGROUPSIZE_Y, THREADGROUPSIZE_Z)]
-void PCGGrassMapUnpacker_CS(const uint3 WindowCoord_Slice : SV_DispatchThreadID)
+void PCGGrassMapUnpacker_CS(const uint3 WindowCoord_OutputSlice : SV_DispatchThreadID)
 {
-	const uint2 WindowCoord = WindowCoord_Slice.xy;
+	const uint2 WindowCoord = WindowCoord_OutputSlice.xy;
 	const uint2 FullCoord = InWindow.xy + WindowCoord;
 	if (any(FullCoord > InWindow.zw))
 	{
 		return;
 	}
 
-	const uint Slice = WindowCoord_Slice.z;
+	const uint OutputSlice = WindowCoord_OutputSlice.z;
 
 	// The rendered landscape components in the input are not in any particular order, so lookup which component index for this output tile.
 	// Clamp because we are dividing by resolution-1 and last col/row of texture will overflow.
@@ -85,19 +86,19 @@ void PCGGrassMapUnpacker_CS(const uint3 WindowCoord_Slice : SV_DispatchThreadID)
 	if (ComponentIndex == -1)
 	{
 		// First layer of threads also deals with height.
-		if (InOutputHeight == 1 && Slice == 0)
+		if (InOutputHeight == 1 && OutputSlice == 0)
 		{
 			OutUnpackedHeight[WindowCoord] = 0.0f;
 		}
 
-		OutUnpackedGrassMaps[uint3(WindowCoord, Slice)] = 0.0f;
+		OutUnpackedGrassMaps[uint3(WindowCoord, OutputSlice)] = 0.0f;
 		return;
 	}
 
 	const uint2 LocalTileCoords = uint2(FullCoord.x - TileCoords.x * (InLandscapeComponentResolution - 1), FullCoord.y - TileCoords.y * (InLandscapeComponentResolution - 1));
 
 	// Heights. The first layer threads also deal with height.
-	if (InOutputHeight == 1 && Slice == 0)
+	if (InOutputHeight == 1 && OutputSlice == 0)
 	{
 		const int ComponentXOffset = InComponentHeightOffsets[ComponentIndex];
 
@@ -108,7 +109,8 @@ void PCGGrassMapUnpacker_CS(const uint3 WindowCoord_Slice : SV_DispatchThreadID)
 	}
 
 	// Coordinates into input height & grass maps (see packing above).
-	const int2 ComponentAndGrassXOffsetAndChannel = InComponentGrassmapOffsetAndChannels[ComponentIndex * InNumGlobalGrassTypes + Slice];
+	const int ExportSlice = InOutputIndexToExportIndex[OutputSlice];
+	const int2 ComponentAndGrassXOffsetAndChannel = (ExportSlice > -1) ? InComponentGrassmapOffsetAndChannels[ComponentIndex * InNumExportedGrassTypes + ExportSlice] : -1;
 	if (ComponentAndGrassXOffsetAndChannel.x != -1)
 	{
 		const uint ComponentXOffset = ComponentAndGrassXOffsetAndChannel.x;
@@ -118,14 +120,19 @@ void PCGGrassMapUnpacker_CS(const uint3 WindowCoord_Slice : SV_DispatchThreadID)
 		const uint2 InputCoords = uint2(LocalTileCoords.x + ComponentXOffset, LocalTileCoords.y);
 
 		// Grass maps
-		OutUnpackedGrassMaps[uint3(WindowCoord, Slice)] = InPackedGrassMaps[InputCoords][ExportChannelIndex];
+		OutUnpackedGrassMaps[uint3(WindowCoord, OutputSlice)] = InPackedGrassMaps[InputCoords][ExportChannelIndex];
 
 		// TODO: position of height data can move around based on slices, so also needs to be in mapping somehow
 		// Heights. The first layer threads also deal with height.
-		if (InOutputHeight == 1 && Slice == 0)
+		if (InOutputHeight == 1 && OutputSlice == 0)
 		{
 			const float RelativeHeight = (UnpackHeight(InPackedGrassMaps[InputCoords].xy) - LANDSCAPE_MID_VALUE) * InLandscapeGridScale.z * TERRAIN_ZSCALE;
 			OutUnpackedHeight[WindowCoord] = RelativeHeight + InLandscapeLocationZ;
 		}
 	}
+	else
+	{
+		// Grass map was not exported, write 0. Happens if for example user specifies a grass type name that isn't present on any of the landscape components.
+		OutUnpackedGrassMaps[uint3(WindowCoord, OutputSlice)] = 0.0f;
+	}
 }
diff --git a/Engine/Plugins/PCG/Source/PCG/Private/Elements/PCGGenerateLandscapeTextures.cpp b/Engine/Plugins/PCG/Source/PCG/Private/Elements/PCGGenerateLandscapeTextures.cpp
--- a/Engine/Plugins/PCG/Source/PCG/Private/Elements/PCGGenerateLandscapeTextures.cpp
+++ b/Engine/Plugins/PCG/Source/PCG/Private/Elements/PCGGenerateLandscapeTextures.cpp
@@ -68,6 +68,7 @@ UPCGGenerateLandscapeTexturesSettings::UPCGGenerateLandscapeTexturesSettings()
 	if (PCGHelpers::IsNewObjectAndNotDefault(this))
 	{
 		bGenerateHeightMap = true;
+		bOutputDataForEachSelectedGrassType = true;
 	}
 }
 
@@ -157,10 +158,12 @@ EPCGChangeType UPCGGenerateLandscapeTexturesSettings::GetChangeTypeForProperty(c
 }
 #endif // WITH_EDITOR
 
+PRAGMA_DISABLE_DEPRECATION_WARNINGS
 FPCGGenerateLandscapeTexturesContext::~FPCGGenerateLandscapeTexturesContext()
 {
 	delete LandscapeGrassWeightExporter;
 }
+PRAGMA_ENABLE_DEPRECATION_WARNINGS
 
 void FPCGGenerateLandscapeTexturesContext::AddExtraStructReferencedObjects(FReferenceCollector& Collector)
 {
@@ -169,26 +172,20 @@ void FPCGGenerateLandscapeTexturesContext::AddExtraStructReferencedObjects(FRefe
 		Collector.AddReferencedObject(HeightTextureData);
 	}
 
-	for (TObjectPtr<UPCGTextureData>& Data : GrassMapTextureDatas)
-	{
-		if (Data)
-		{
-			Collector.AddReferencedObject(Data);
-		}
-	}
-
 	if (GrassMapTextureArrayData)
 	{
 		Collector.AddReferencedObject(GrassMapTextureArrayData);
 	}
 
-	for (TObjectPtr<UTexture>& Texture : TexturesToStream)
+	if (BlackTexture)
 	{
-		if (Texture)
-		{
-			Collector.AddReferencedObject(Texture);
-		}
+		Collector.AddReferencedObject(BlackTexture);
 	}
+
+	Collector.AddReferencedObjects(GrassMapTextureDatas);
+	Collector.AddReferencedObjects(TexturesToStream);
+	Collector.AddReferencedObjects(GrassTypesToExport);
+	Collector.AddReferencedObjects(GrassTypesToOutput);
 }
 
 void FPCGGenerateLandscapeTexturesElement::GetDependenciesCrc(const FPCGGetDependenciesCrcParams& InParams, FPCGCrc& OutCrc) const
@@ -286,8 +283,19 @@ bool FPCGGenerateLandscapeTexturesElement::ExecuteInternal(FPCGContext* InContex
 			}
 		}
 
+		// If we're using an inclusion list, initialize the grass types to output to nulls, one for each selection entry. The entries will be replaced as grass types are encountered below.
+		const bool bOutputDataForEachSelectedGrassType = !Settings->bExcludeSelectedGrassTypes && Settings->bOutputDataForEachSelectedGrassType;
+		if (bOutputDataForEachSelectedGrassType)
+		{
+			Context->GrassTypesToOutput.SetNum(SelectedGrassTypes.Num());
+
+			Context->OutputIndexToExportIndex.SetNumUninitialized(SelectedGrassTypes.Num());
+			FMemory::Memset(Context->OutputIndexToExportIndex.GetData(), -1, Context->OutputIndexToExportIndex.Num() * Context->OutputIndexToExportIndex.GetTypeSize());
+		}
+
 		const FBox GenerationBounds = ExecutionSource->GetExecutionState().GetBounds();
 		TArray<ULandscapeComponent*> OverlappingComponents;
+		TSet<TObjectPtr<ULandscapeGrassType>> ProcessedGrassTypes;
 
 		for (const TSoftObjectPtr<ALandscapeProxy>& LandscapeProxyPtr : LandscapeData->Landscapes)
 		{
@@ -334,9 +342,44 @@ bool FPCGGenerateLandscapeTexturesElement::ExecuteInternal(FPCGContext* InContex
 							break;
 						}
 
-						for (auto& GrassType : LandscapeComponent->GetGrassTypes())
+						FString GrassName; // Used in loop but hoisted for efficiency
+
+						for (const TObjectPtr<ULandscapeGrassType>& GrassType : LandscapeComponent->GetGrassTypes())
 						{
-							Context->GlobalGrassTypes.AddUnique(GrassType);
+							if (ProcessedGrassTypes.Contains(GrassType))
+							{
+								continue;
+							}
+
+							ProcessedGrassTypes.Add(GrassType);
+
+							GrassType->GetName(GrassName);
+							const int32 IndexInSelection = SelectedGrassTypes.IndexOfByKey(GrassName);
+							const bool bIsSelectedLayer = IndexInSelection != INDEX_NONE;
+
+							// Register the grass type if it is included or not excluded.
+							if (bIsSelectedLayer != Settings->bExcludeSelectedGrassTypes)
+							{
+								const int32 SizeBefore = Context->GrassTypesToExport.Num();
+								Context->GrassTypesToExport.AddUnique(GrassType);
+
+								// Update grass type lists if this grass type is newly added.
+								if (Context->GrassTypesToExport.Num() > SizeBefore)
+								{
+									if (bOutputDataForEachSelectedGrassType)
+									{
+										// If we need to follow the prescribed list of grass types, fill in the entry here.
+										Context->GrassTypesToOutput[IndexInSelection] = GrassType;
+										Context->OutputIndexToExportIndex[IndexInSelection] = Context->GrassTypesToExport.Num() - 1;
+									}
+									else
+									{
+										// Otherwise build list of grass types as we encounter them, if they are not filtered out by the include/exclude list.
+										Context->GrassTypesToOutput.Add(GrassType);
+										Context->OutputIndexToExportIndex.Add(Context->GrassTypesToExport.Num() - 1);
+									}
+								}
+							}
 						}
 
 						LandscapeComponent->UpdateGrassTypes();
@@ -385,26 +428,7 @@ bool FPCGGenerateLandscapeTexturesElement::ExecuteInternal(FPCGContext* InContex
 			}
 		}
 
-		// Filter grass types based on selected/excluded.
-		FString GrassTypeName;
-		GrassTypeName.Reserve(256);
-
-		for (int GrassTypeIndex = 0; GrassTypeIndex < Context->GlobalGrassTypes.Num(); ++GrassTypeIndex)
-		{
-			if (const TObjectPtr<ULandscapeGrassType>& GrassType = Context->GlobalGrassTypes[GrassTypeIndex])
-			{
-				GrassType->GetName(GrassTypeName);
-
-				const bool bIsSelectedLayer = SelectedGrassTypes.Contains(GrassTypeName);
-
-				if (bIsSelectedLayer != Settings->bExcludeSelectedGrassTypes)
-				{
-					Context->SelectedGrassTypes.Emplace(GrassType, GrassTypeIndex);
-				}
-			}
-		}
-
-		if (!Context->Landscape.Get() || Context->LandscapeComponents.IsEmpty() || (Context->SelectedGrassTypes.IsEmpty() && !Settings->bGenerateHeightMap))
+		if (!Context->Landscape.Get() || Context->LandscapeComponents.IsEmpty() || (Context->GrassTypesToOutput.IsEmpty() && !Settings->bGenerateHeightMap))
 		{
 			return true;
 		}
@@ -427,7 +451,7 @@ bool FPCGGenerateLandscapeTexturesElement::ExecuteInternal(FPCGContext* InContex
 						Context->TexturesToStream.Add(HeightMap);
 					}
 
-					if (Context->GlobalGrassTypes.Num() > 0)
+					if (Context->GrassTypesToOutput.Num() > 0)
 					{
 						const ERHIFeatureLevel::Type FeatureLevel = World->GetFeatureLevel();
 						for (UTexture2D* WeightmapTexture : LandscapeComponent->GetRenderedWeightmapTexturesForFeatureLevel(FeatureLevel))
@@ -551,7 +575,7 @@ bool FPCGGenerateLandscapeTexturesElement::ExecuteInternal(FPCGContext* InContex
 		}
 
 		UE::Landscape::Grass::EGrassWeightExporterFlags Flags = UE::Landscape::Grass::EGrassWeightExporterFlags::None;
-		if (!Context->GlobalGrassTypes.IsEmpty())
+		if (!Context->GrassTypesToExport.IsEmpty())
 		{
 			Flags |= UE::Landscape::Grass::EGrassWeightExporterFlags::NeedsGrassmap;
 		}
@@ -564,7 +588,7 @@ bool FPCGGenerateLandscapeTexturesElement::ExecuteInternal(FPCGContext* InContex
 			Landscape, 
 			LandscapeComponents,
 			Flags,
-			Context->GlobalGrassTypes, 
+			Context->GrassTypesToExport,
 			/*InRequestedHeightMips = */{});
 
 		const FIntPoint& OutputTextureSize = Context->LandscapeGrassWeightExporter->GetTargetSize();
@@ -639,6 +663,7 @@ bool FPCGGenerateLandscapeTexturesElement::ExecuteInternal(FPCGContext* InContex
 			, LandscapeComponentExtent=Context->LandscapeComponentExtent
 			, LandscapeTileCoords = MoveTemp(LandscapeTileCoords)
 			, LandscapeComponentKeys = MoveTemp(LandscapeComponentKeys)
+			, OutputIndexToExportIndex = MoveTemp(Context->OutputIndexToExportIndex)
 			, LandscapeComponentResolution
 			, LandscapeGridScale
 			, LandscapeZ
@@ -673,7 +698,7 @@ bool FPCGGenerateLandscapeTexturesElement::ExecuteInternal(FPCGContext* InContex
 			Context->LandscapeGrassWeightExporter->RenderLandscapeComponentToTexture_RenderThread(GraphBuilder, GrassMapTexture);
 
 			// Must have at least one slice to be a valid resource.
-			const uint32 NumOutputSlices = FMath::Max(Context->GlobalGrassTypes.Num(), 1);
+			const uint32 NumOutputSlices = FMath::Max(Context->GrassTypesToOutput.Num(), 1);
 
 			FRDGTextureDesc GrassHeightMapDesc = FRDGTextureDesc::Create2D(
 				bGenerateHeightMap ? OutputResolution : FIntPoint(1, 1), // No default system texture usable for UAV, so always create but with 1x1 resolution
@@ -706,7 +731,7 @@ bool FPCGGenerateLandscapeTexturesElement::ExecuteInternal(FPCGContext* InContex
 			Parameters->InLandscapeComponentResolution = LandscapeComponentResolution;
 			Parameters->InOutputHeight = bGenerateHeightMap ? 1u : 0u;
 			Parameters->InWindow = FUintVector4(WindowStart.X, WindowStart.Y, WindowEnd.X, WindowEnd.Y);
-			Parameters->InNumGlobalGrassTypes = Context->GlobalGrassTypes.Num();
+			Parameters->InNumExportedGrassTypes = Context->GrassTypesToExport.Num();
 
 			// Initialize to invalid component indices.
 			for (int Index = 0; Index < FPCGGrassMapUnpackerCS::MaxNumLandscapeComponents; ++Index)
@@ -748,20 +773,20 @@ bool FPCGGenerateLandscapeTexturesElement::ExecuteInternal(FPCGContext* InContex
 			Parameters->InComponentHeightOffsets = GraphBuilder.CreateSRV(ComponentHeightOffsetsBuffer);
 
 			FRDGBufferRef ComponentGrassmapOffsetAndChannelsBuffer;
-			if (!Context->GlobalGrassTypes.IsEmpty())
+			if (!Context->GrassTypesToExport.IsEmpty())
 			{
 				// For each component, for each grass type, X offset for where to find the grass data in the grass texture (.x) and on which channel (.y)
 				//  It's possible a grass map is absent from a given component, in which case, the offset and channel will be set to -1
 				TArray<FIntVector2> ComponentGrassmapOffsetAndChannels;
-				ComponentGrassmapOffsetAndChannels.Reserve(LandscapeComponentKeys.Num() * Context->GlobalGrassTypes.Num());
+				ComponentGrassmapOffsetAndChannels.Reserve(LandscapeComponentKeys.Num() * Context->GrassTypesToExport.Num());
 				for (int32 ComponentIndex = 0; ComponentIndex < LandscapeComponentKeys.Num(); ++ComponentIndex)
 				{
-					for (int32 GrassmapIndex = 0; GrassmapIndex < Context->GlobalGrassTypes.Num(); ++GrassmapIndex)
+					for (int32 GrassmapIndex = 0; GrassmapIndex < Context->GrassTypesToExport.Num(); ++GrassmapIndex)
 					{
 						FIntVector2& ComponentGrassmapOffsetAndChannel = ComponentGrassmapOffsetAndChannels.Add_GetRef(FIntVector2(-1, -1));
 
 						TValueOrError<TPair<FIntRect /*OutputRect*/, uint8 /*ChannelIndex*/>, FLandscapeGrassWeightExporter_RenderThread::EGrassmapTextureInfoError> TextureInfo =
-							Context->LandscapeGrassWeightExporter->GetTextureInfoForGrass(LandscapeComponentKeys[ComponentIndex], Context->GlobalGrassTypes[GrassmapIndex]);
+							Context->LandscapeGrassWeightExporter->GetTextureInfoForGrass(LandscapeComponentKeys[ComponentIndex], Context->GrassTypesToExport[GrassmapIndex]);
 						if (TextureInfo.HasError())
 						{
 							// It can fail, but only if that's because the grass map is absent from this component
@@ -776,7 +801,7 @@ bool FPCGGenerateLandscapeTexturesElement::ExecuteInternal(FPCGContext* InContex
 						}
 					}
 				}
-				check(ComponentGrassmapOffsetAndChannels.Num() == LandscapeComponentKeys.Num() * Context->GlobalGrassTypes.Num());
+				check(ComponentGrassmapOffsetAndChannels.Num() == LandscapeComponentKeys.Num() * Context->GrassTypesToExport.Num());
 				ComponentGrassmapOffsetAndChannelsBuffer = CreateStructuredBuffer(GraphBuilder, TEXT("PCGGenerateLandscapeTexturesComponentGrassmapOffsetAndChannels"), ComponentGrassmapOffsetAndChannels);
 			}
 			else
@@ -785,6 +810,17 @@ bool FPCGGenerateLandscapeTexturesElement::ExecuteInternal(FPCGContext* InContex
 			}
 			Parameters->InComponentGrassmapOffsetAndChannels = GraphBuilder.CreateSRV(ComponentGrassmapOffsetAndChannelsBuffer);
 
+			FRDGBufferRef OutputIndexToExportIndexBuffer;
+			if (!OutputIndexToExportIndex.IsEmpty())
+			{
+				OutputIndexToExportIndexBuffer = CreateStructuredBuffer(GraphBuilder, TEXT("PCGGenerateLandscapeTexturesOutputIndexToExportIndex"), OutputIndexToExportIndex);
+			}
+			else
+			{
+				OutputIndexToExportIndexBuffer = GSystemTextures.GetDefaultStructuredBuffer(GraphBuilder, sizeof(int32));
+			}
+			Parameters->InOutputIndexToExportIndex = GraphBuilder.CreateSRV(OutputIndexToExportIndexBuffer);
+
 			TShaderMapRef<FPCGGrassMapUnpackerCS> Shader(GetGlobalShaderMap(GMaxRHIFeatureLevel));
 			const int GroupCountX = FMath::DivideAndRoundUp<int>(OutputResolution.X, FPCGGrassMapUnpackerCS::ThreadGroupDim);
 			const int GroupCountY = FMath::DivideAndRoundUp<int>(OutputResolution.Y, FPCGGrassMapUnpackerCS::ThreadGroupDim);
@@ -847,7 +883,7 @@ bool FPCGGenerateLandscapeTexturesElement::ExecuteInternal(FPCGContext* InContex
 		else
 		{
 			// Create the texture data objects if they haven't been created already. There should be one per selected grass type.
-			for (int DataIndex = 0; DataIndex < Context->SelectedGrassTypes.Num(); ++DataIndex)
+			for (int DataIndex = 0; DataIndex < Context->GrassTypesToOutput.Num(); ++DataIndex)
 			{
 				Context->GrassMapTextureDatas.Add(FPCGContext::NewObject_AnyThread<UPCGTextureData>(Context));
 			}
@@ -880,18 +916,37 @@ bool FPCGGenerateLandscapeTexturesElement::ExecuteInternal(FPCGContext* InContex
 		}
 		else
 		{
-			for (int DataIndex = 0; DataIndex < Context->SelectedGrassTypes.Num(); ++DataIndex)
+			for (int DataIndex = 0; DataIndex < Context->GrassTypesToOutput.Num(); ++DataIndex)
 			{
 				check(Context->GrassMapTextureDatas.IsValidIndex(DataIndex));
 				UPCGTextureData* TextureData = Context->GrassMapTextureDatas[DataIndex];
 				check(TextureData);
 
-				// Poll initialize (fine to be called even when initialization was already complete).
-				bAllTexturesInitialized &= TextureData->Initialize(
-					Context->GrassMapHandle,
-					/*TextureIndex=*/Context->SelectedGrassTypes[DataIndex].Get<1>(),
-					Context->OutputTextureTransform,
-					Settings->bSkipReadbackToCPU);
+				const int32 ExportIndex = Context->GrassTypesToOutput[DataIndex] ? Context->GrassTypesToExport.IndexOfByKey(Context->GrassTypesToOutput[DataIndex]) : INDEX_NONE;
+
+				if (ExportIndex != INDEX_NONE)
+				{
+					// Poll initialize (fine to be called even when initialization was already complete).
+					bAllTexturesInitialized &= TextureData->Initialize(
+						Context->GrassMapHandle,
+						/*TextureIndex=*/DataIndex,
+						Context->OutputTextureTransform,
+						Settings->bSkipReadbackToCPU);
+				}
+				else
+				{
+					if (!Context->BlackTexture)
+					{
+						// Expectation is this case should be rare, and black texture generally always loaded.
+						Context->BlackTexture = LoadObject<UTexture2D>(nullptr, TEXT("/Engine/EngineResources/Black.Black"));
+					}
+
+					bAllTexturesInitialized &= TextureData->Initialize(
+						Context->BlackTexture,
+						/*TextureIndex=*/0,
+						Context->OutputTextureTransform,
+						Settings->bSkipReadbackToCPU);
+				}
 			}
 		}
 	}
@@ -946,7 +1001,7 @@ bool FPCGGenerateLandscapeTexturesElement::ExecuteInternal(FPCGContext* InContex
 	}
 	else
 	{
-		for (int DataIndex = 0; DataIndex < Context->SelectedGrassTypes.Num(); ++DataIndex)
+		for (int DataIndex = 0; DataIndex < Context->GrassTypesToOutput.Num(); ++DataIndex)
 		{
 			check(Context->GrassMapTextureDatas.IsValidIndex(DataIndex));
 			UPCGTextureData* TextureData = Context->GrassMapTextureDatas[DataIndex];
@@ -964,7 +1019,7 @@ bool FPCGGenerateLandscapeTexturesElement::ExecuteInternal(FPCGContext* InContex
 
 			FString GrassTypeName = TEXT("");
 
-			if (ULandscapeGrassType* GrassType = Context->SelectedGrassTypes[DataIndex].Get<0>().Get())
+			if (ULandscapeGrassType* GrassType = Context->GrassTypesToOutput[DataIndex].Get())
 			{
 				GrassType->GetName(GrassTypeName);
 			}
@@ -973,6 +1028,11 @@ bool FPCGGenerateLandscapeTexturesElement::ExecuteInternal(FPCGContext* InContex
 			{
 				OutTaggedData.Tags.Add(PCGComputeHelpers::GetPrefixedDataLabel(GrassTypeName));
 			}
+			else if (!Settings->bExcludeSelectedGrassTypes && Settings->SelectedGrassTypes.IsValidIndex(DataIndex))
+			{
+				// We are using inclusion list, look up name
+				OutTaggedData.Tags.Add(PCGComputeHelpers::GetPrefixedDataLabel(Settings->SelectedGrassTypes[DataIndex]));
+			}
 			else
 			{
 				PCGE_LOG(Warning, LogOnly, LOCTEXT("MissingGrassName",
diff --git a/Engine/Plugins/PCG/Source/PCG/Public/Elements/PCGGenerateLandscapeTextures.h b/Engine/Plugins/PCG/Source/PCG/Public/Elements/PCGGenerateLandscapeTextures.h
--- a/Engine/Plugins/PCG/Source/PCG/Public/Elements/PCGGenerateLandscapeTextures.h
+++ b/Engine/Plugins/PCG/Source/PCG/Public/Elements/PCGGenerateLandscapeTextures.h
@@ -63,10 +63,14 @@ public:
 	UPROPERTY(BlueprintReadWrite, EditAnywhere, Category = Settings, meta = (EditCondition = "bOverrideFromInput", PCG_DiscardPropertySelection, PCG_DiscardExtraSelection))
 	FPCGAttributePropertyInputSelector GrassTypesAttribute;
 
-	/** If toggled, will only generate grass types which are not selected. */
+	/** Only generate grass types which are not selected. */
 	UPROPERTY(BlueprintReadWrite, EditAnywhere, Category = Settings, meta = (PCG_Overridable))
 	bool bExcludeSelectedGrassTypes = true;
 
+	/** Output texture data for each selected grass type. Outputs black for grass types that are not present. */
+	UPROPERTY(BlueprintReadWrite, EditAnywhere, Category = Settings, meta = (PCG_Overridable, EditCondition = "!bExcludeSelectedGrassTypes"))
+	bool bOutputDataForEachSelectedGrassType = false;
+
 	/** Skip CPU readback of emitted textures during initialization of the texture datas. */
 	UPROPERTY(BlueprintReadWrite, EditAnywhere, Category = Settings)
 	bool bSkipReadbackToCPU = false;
@@ -99,12 +103,14 @@ public:
 	TRefCountPtr<IPooledRenderTarget> HeightHandle = nullptr;
 	TRefCountPtr<IPooledRenderTarget> GrassMapHandle = nullptr;
 
-	/** List of the grass types selected for generation. We have to hold their texture index as well, since
-	 * this array could be sparse, but the actual texture array we produce will always have all the grass types.
-	 */
-	TArray<TTuple<TWeakObjectPtr<ULandscapeGrassType>, /*TextureIndex=*/int32>> SelectedGrassTypes;
+	/** The grass types passed to the grass exporter. Each of these grass types is present on at least one of the landscape components. */
+	TArray<TObjectPtr<ULandscapeGrassType>> GrassTypesToExport;
 
-	TArray<TObjectPtr<ULandscapeGrassType>> GlobalGrassTypes;
+	/** The grass types to output from the node. Grass types that were not present on the overlapping landscape components will emit empty/black texture data. */
+	TArray<TObjectPtr<ULandscapeGrassType>> GrassTypesToOutput;
+
+	/** Maps output grass type indices to exported grass type indices. */
+	TArray<int32> OutputIndexToExportIndex;
 
 	TWeakObjectPtr<ALandscape> Landscape;
 	TArray<TWeakObjectPtr<ULandscapeComponent>> LandscapeComponents;
@@ -134,6 +140,8 @@ public:
 
 	FTransform OutputTextureTransform;
 
+	TObjectPtr<UTexture> BlackTexture;
+
 protected:
 	virtual void AddExtraStructReferencedObjects(FReferenceCollector& Collector) override;
 
@@ -143,6 +151,10 @@ public:
 	int32 NumGrassTypes = 0;
 	UE_DEPRECATED(5.8, "Use Landscape")
 	TWeakObjectPtr<ALandscapeProxy> LandscapeProxy = nullptr;
+	UE_DEPRECATED(5.8, "No longer used")
+	TArray<TTuple<TWeakObjectPtr<ULandscapeGrassType>, int32>> SelectedGrassTypes;
+	UE_DEPRECATED(5.8, "No longer used")
+	TArray<TObjectPtr<ULandscapeGrassType>> GlobalGrassTypes;
 };
 
 class FPCGGenerateLandscapeTexturesElement : public IPCGElement
diff --git a/Engine/Plugins/PCG/Source/PCGCompute/Internal/PCGGrassMapUnpackerCS.h b/Engine/Plugins/PCG/Source/PCGCompute/Internal/PCGGrassMapUnpackerCS.h
--- a/Engine/Plugins/PCG/Source/PCGCompute/Internal/PCGGrassMapUnpackerCS.h
+++ b/Engine/Plugins/PCG/Source/PCGCompute/Internal/PCGGrassMapUnpackerCS.h
@@ -22,12 +22,13 @@ public:
 		SHADER_PARAMETER_RDG_TEXTURE_SRV(Texture2D<float4>, InPackedGrassMaps)
 		SHADER_PARAMETER_RDG_BUFFER_SRV(StructuredBuffer<int32>, InComponentHeightOffsets)
 		SHADER_PARAMETER_RDG_BUFFER_SRV(StructuredBuffer<FIntVector2>, InComponentGrassmapOffsetAndChannels)
+		SHADER_PARAMETER_RDG_BUFFER_SRV(StructuredBuffer<int32>, InOutputIndexToExportIndex)
 		SHADER_PARAMETER_RDG_TEXTURE_UAV(RWTexture2D<float>, OutUnpackedHeight)
 		SHADER_PARAMETER_RDG_TEXTURE_UAV(RWTexture2DArray<float>, OutUnpackedGrassMaps)
 		SHADER_PARAMETER_ARRAY(FIntVector4, InLinearTileIndexToComponentIndex, [MaxNumLandscapeComponents])
 		SHADER_PARAMETER(FUintVector2, InNumTiles)
 		SHADER_PARAMETER(uint32, InLandscapeComponentResolution)
-		SHADER_PARAMETER(int32, InNumGlobalGrassTypes)
+		SHADER_PARAMETER(int32, InNumExportedGrassTypes)
 		SHADER_PARAMETER(FVector3f, InLandscapeGridScale)
 		SHADER_PARAMETER(float, InLandscapeLocationZ)
 		SHADER_PARAMETER(uint32, InOutputHeight)
