[PCG] Fix crash with trying to rename Point to Attribute Set node

* Updated base settings to disallow renaming for compact nodes

#rnx

#Synced-CL 49818457

diff --git a/Engine/Plugins/PCG/Source/PCG/Private/PCGSettings.cpp b/Engine/Plugins/PCG/Source/PCG/Private/PCGSettings.cpp
--- a/Engine/Plugins/PCG/Source/PCG/Private/PCGSettings.cpp
+++ b/Engine/Plugins/PCG/Source/PCG/Private/PCGSettings.cpp
@@ -798,6 +798,11 @@ bool UPCGSettings::CanUseExecutionDependencyPinForCulling() const
 	return HasExecutionDependencyPin();
 }
 
+bool UPCGSettings::CanUserEditTitle() const
+{
+	return !ShouldDrawNodeCompact();
+}
+
 void UPCGSettings::ApplyDeprecationBeforeUpdatePins(UPCGNode* InOutNode, TArray<TObjectPtr<UPCGPin>>& InputPins, TArray<TObjectPtr<UPCGPin>>& OutputPins)
 {
 	check(InOutNode);
diff --git a/Engine/Plugins/PCG/Source/PCG/Public/PCGSettings.h b/Engine/Plugins/PCG/Source/PCG/Public/PCGSettings.h
--- a/Engine/Plugins/PCG/Source/PCG/Public/PCGSettings.h
+++ b/Engine/Plugins/PCG/Source/PCG/Public/PCGSettings.h
@@ -307,7 +307,7 @@ public:
 	/** Returns the icon to use instead of text in compact node form */
 	virtual bool GetCompactNodeIcon(FName& OutCompactNodeIcon) const { return false; }
 	/** Returns whether the user can directly interact with the node name */
-	virtual bool CanUserEditTitle() const { return true; }
+	UE_API virtual bool CanUserEditTitle() const;
 
 	/** UpdatePins will kick off invalid edges, so this is useful for moving edges around in case of pin changes. */
 	UE_API virtual void ApplyDeprecationBeforeUpdatePins(UPCGNode* InOutNode, TArray<TObjectPtr<UPCGPin>>& InputPins, TArray<TObjectPtr<UPCGPin>>& OutputPins);
