Add additional RDG scopes to surface GPU work.

#rnx

#Synced-CL 49817612

diff --git a/Engine/Plugins/PCG/Source/PCG/Private/Elements/PCGDownsampleTexture.cpp b/Engine/Plugins/PCG/Source/PCG/Private/Elements/PCGDownsampleTexture.cpp
--- a/Engine/Plugins/PCG/Source/PCG/Private/Elements/PCGDownsampleTexture.cpp
+++ b/Engine/Plugins/PCG/Source/PCG/Private/Elements/PCGDownsampleTexture.cpp
@@ -134,6 +134,8 @@ bool FPCGDownsampleTextureElement::ExecuteInternal(FPCGContext* InContext) const
 				TRefCountPtr<IPooledRenderTarget> OutputTextureExported;
 
 				{
+					RDG_EVENT_SCOPE(GraphBuilder, "PCGDownsampleTexture");
+
 #if !UE_BUILD_SHIPPING
 					RenderCaptureInterface::FScopedCapture RenderCapture(PCGDownsampleTexture::TriggerGPUCaptureDispatchIndex > 0, GraphBuilder, TEXT("PCGDownsampleTexture"));
 					PCGDownsampleTexture::TriggerGPUCaptureDispatchIndex = FMath::Max(PCGDownsampleTexture::TriggerGPUCaptureDispatchIndex - 1, 0);
diff --git a/Engine/Source/Runtime/Landscape/Private/LandscapeGrassWeightExporter.cpp b/Engine/Source/Runtime/Landscape/Private/LandscapeGrassWeightExporter.cpp
--- a/Engine/Source/Runtime/Landscape/Private/LandscapeGrassWeightExporter.cpp
+++ b/Engine/Source/Runtime/Landscape/Private/LandscapeGrassWeightExporter.cpp
@@ -503,6 +503,8 @@ void FLandscapeGrassWeightExporter_RenderThread::RenderLandscapeComponentToTextu
 
 void FLandscapeGrassWeightExporter_RenderThread::RenderLandscapeComponentToTexture_RenderThread(FRDGBuilder& GraphBuilder, FRDGTextureRef OutputTexture)
 {
+	RDG_EVENT_SCOPE(GraphBuilder, "RenderLandscapeComponentToTexture");
+
 	FSceneViewFamilyContext ViewFamily(FSceneViewFamily::ConstructionValues(nullptr, SceneInterface, FEngineShowFlags(ESFIM_Game))
 		.SetTime(FGameTime::GetTimeSinceAppStart()));
 
