[PCG] Fix crash in UPCGLandscapeSplineData::GetTransformAtDistance()

#rnx

#Synced-CL 49854193

#Synced-CL 49854193

diff --git a/Engine/Plugins/PCG/Source/PCG/Private/Data/PCGLandscapeSplineData.cpp b/Engine/Plugins/PCG/Source/PCG/Private/Data/PCGLandscapeSplineData.cpp
--- a/Engine/Plugins/PCG/Source/PCG/Private/Data/PCGLandscapeSplineData.cpp
+++ b/Engine/Plugins/PCG/Source/PCG/Private/Data/PCGLandscapeSplineData.cpp
@@ -110,6 +110,13 @@ FTransform UPCGLandscapeSplineData::GetTransformAtDistance(int SegmentIndex, FVe
 	check(Segment);
 
 	const TArray<FLandscapeSplineInterpPoint>& InterpPoints = Segment->GetPoints();
+	const int32 SecondToLastPointIndex = InterpPoints.Num() - 2;
+
+	if (SecondToLastPointIndex < 0)
+	{
+		// Not enough points in the segment to interpolate a transform.
+		return GetTransform();
+	}
 
 	int32 PointIndex = 0;
 	FVector::FReal Alpha = 0.0f;
@@ -119,15 +126,15 @@ FTransform UPCGLandscapeSplineData::GetTransformAtDistance(int SegmentIndex, FVe
 		GetInterpPointAtDistance(ClampedSegmentIndex, Distance, PointIndex, /*bComputeAlpha=*/true, Alpha);
 
 		// If our RightPoint ends up being on the next segment, clamp it back to the current segment.
-		PointIndex = FMath::Min(PointIndex, InterpPoints.Num() - 2);
+		PointIndex = FMath::Clamp(PointIndex, 0, SecondToLastPointIndex);
 	}
 	else if(ClampedSegmentIndex < SegmentIndex)
 	{
+		// If we've landed beyond the last segment, just clamp to the end of the last segment.
 		Alpha = 1.0f;
-		PointIndex = FMath::Max(0, InterpPoints.Num() - 2);
+		PointIndex = SecondToLastPointIndex;
 	}
 
-	check(InterpPoints.IsValidIndex(PointIndex));
 	const FLandscapeSplineInterpPoint& PreviousPoint = InterpPoints[PointIndex];
 	const FLandscapeSplineInterpPoint& CurrentPoint = InterpPoints[PointIndex + 1];
 
