[PCG] Fix crash when using virtual texture data interface in debug builds

#rnx

#Synced-CL 49854193

diff --git a/Engine/Plugins/PCG/Source/PCG/Private/Compute/DataInterfaces/PCGVirtualTextureDataInterface.cpp b/Engine/Plugins/PCG/Source/PCG/Private/Compute/DataInterfaces/PCGVirtualTextureDataInterface.cpp
--- a/Engine/Plugins/PCG/Source/PCG/Private/Compute/DataInterfaces/PCGVirtualTextureDataInterface.cpp
+++ b/Engine/Plugins/PCG/Source/PCG/Private/Compute/DataInterfaces/PCGVirtualTextureDataInterface.cpp
@@ -174,21 +174,28 @@ void UPCGVirtualTextureDataProvider::Reset()
 }
 
 FPCGVirtualTextureDataProviderProxy::FPCGVirtualTextureDataProviderProxy(const TArray<TObjectPtr<const URuntimeVirtualTexture>>& InVirtualTextures)
+{
+	for (const TObjectPtr<const URuntimeVirtualTexture>& VirtualTexture : InVirtualTextures)
+	{
+		VirtualTextures.Add(VirtualTexture);
+	}
+}
+
+void FPCGVirtualTextureDataProviderProxy::AllocateResources(FRDGBuilder& GraphBuilder, FAllocationData const& InAllocationData)
 {
 	for (int32 BindingIndex = 0; BindingIndex < PCGVirtualTextureDIConstants::MAX_NUM_BINDINGS; ++BindingIndex)
 	{
 		MaterialTypes[BindingIndex] = ERuntimeVirtualTextureMaterialType::Count;
 	}
 
-	NumVirtualTextures = InVirtualTextures.Num();
-	for (int32 VirtualTextureIndex = 0; VirtualTextureIndex < NumVirtualTextures; ++VirtualTextureIndex)
+	for (int32 VirtualTextureIndex = 0; VirtualTextureIndex < VirtualTextures.Num(); ++VirtualTextureIndex)
 	{
-		if (const URuntimeVirtualTexture* VirtualTexture = InVirtualTextures[VirtualTextureIndex])
+		if (TStrongObjectPtr<const URuntimeVirtualTexture> VirtualTexture = VirtualTextures[VirtualTextureIndex].Pin())
 		{
 			const ERuntimeVirtualTextureMaterialType MaterialType = VirtualTexture->GetMaterialType();
 			MaterialTypes[VirtualTextureIndex] = MaterialType;
 			
-			PageTables[VirtualTextureIndex].Initialize(VirtualTexture);
+			PageTables[VirtualTextureIndex].Initialize(VirtualTexture.Get());
 
 			const int32 NumVirtualTextureLayers = FMath::Min(URuntimeVirtualTexture::GetLayerCount(MaterialType), PCGVirtualTextureDIConstants::MAX_NUM_LAYERS);
 			const int32 PackedLayerBaseIndex = VirtualTextureIndex * PCGVirtualTextureDIConstants::MAX_NUM_LAYERS;
@@ -196,7 +203,7 @@ FPCGVirtualTextureDataProviderProxy::FPCGVirtualTextureDataProviderProxy(const T
 			for (int32 LayerIndex = 0; LayerIndex < NumVirtualTextureLayers; ++LayerIndex)
 			{
 				const int32 PackedLayerIndex = PackedLayerBaseIndex + LayerIndex;
-				Layers[PackedLayerIndex].Initialize(VirtualTexture, LayerIndex, VirtualTexture->IsLayerSRGB(LayerIndex));
+				Layers[PackedLayerIndex].Initialize(VirtualTexture.Get(), LayerIndex, VirtualTexture->IsLayerSRGB(LayerIndex));
 			}
 		}
 
@@ -225,7 +232,7 @@ void FPCGVirtualTextureDataProviderProxy::GatherPermutations(FPermutationData& I
 	FPCGVirtualTextureDataInterfacePermutationIds PermutationIds(InOutPermutationData.PermutationVector);
 	for (int32 InvocationIndex = 0; InvocationIndex < InOutPermutationData.NumInvocations; ++InvocationIndex)
 	{
-		InOutPermutationData.PermutationIds[InvocationIndex] |= (NumVirtualTextures > 1 ? PermutationIds.EnableMultipleVirtualTextures : 0);
+		InOutPermutationData.PermutationIds[InvocationIndex] |= (VirtualTextures.Num() > 1 ? PermutationIds.EnableMultipleVirtualTextures : 0);
 	}
 }
 
@@ -235,7 +242,7 @@ void FPCGVirtualTextureDataProviderProxy::GatherDispatchData(FDispatchData const
 	for (int32 InvocationIndex = 0; InvocationIndex < ParameterArray.Num(); ++InvocationIndex)
 	{
 		FParameters& Parameters = ParameterArray[InvocationIndex];
-		Parameters.NumVirtualTextures = NumVirtualTextures;
+		Parameters.NumVirtualTextures = VirtualTextures.Num();
 		Parameters.Sampler = TStaticSamplerState<SF_Bilinear, AM_Clamp, AM_Clamp, AM_Clamp>::GetRHI();
 
 		for (int32 BindingIndex = 0; BindingIndex < PCGVirtualTextureDIConstants::MAX_NUM_BINDINGS; ++BindingIndex)
diff --git a/Engine/Plugins/PCG/Source/PCG/Private/Compute/DataInterfaces/PCGVirtualTextureDataInterface.h b/Engine/Plugins/PCG/Source/PCG/Private/Compute/DataInterfaces/PCGVirtualTextureDataInterface.h
--- a/Engine/Plugins/PCG/Source/PCG/Private/Compute/DataInterfaces/PCGVirtualTextureDataInterface.h
+++ b/Engine/Plugins/PCG/Source/PCG/Private/Compute/DataInterfaces/PCGVirtualTextureDataInterface.h
@@ -76,13 +76,15 @@ public:
 	bool IsValid(FValidationData const& InValidationData) const override;
 	void GatherPermutations(FPermutationData& InOutPermutationData) const override;
 	void GatherDispatchData(FDispatchData const& InDispatchData) override;
+	void AllocateResources(FRDGBuilder& GraphBuilder, FAllocationData const& InAllocationData) override;
 	//~ End FComputeDataProviderRenderProxy Interface
 
 protected:
 	using FParameters = FPCGVirtualTextureDataInterfaceParameters;
 
+	TArray<TWeakObjectPtr<const URuntimeVirtualTexture>> VirtualTextures;
+
 	ERuntimeVirtualTextureMaterialType MaterialTypes[PCGVirtualTextureDIConstants::MAX_NUM_BINDINGS];
 	PCGVirtualTextureCommon::FVirtualTextureLayer Layers[PCGVirtualTextureDIConstants::MAX_NUM_LAYERS * PCGVirtualTextureDIConstants::MAX_NUM_BINDINGS];
 	PCGVirtualTextureCommon::FVirtualTexturePageTable PageTables[PCGVirtualTextureDIConstants::MAX_NUM_BINDINGS];
-	int32 NumVirtualTextures = 0;
 };
