[PCG] PCGEditor worfklow support
- Add a default application mode (FPCGEditorDefaultMode)
- Add a new FPCGGenericTabFactory class to implement tabs
- Add a FDocumentTracker which will be used to add document tab factories (sub graphs)
- Deprecate: EPCGEditorPanel enum

#Synced-CL 49922246

diff --git a/Engine/Plugins/Experimental/ProceduralVegetationEditor/Source/ProceduralVegetationEditor/Private/PVEditor.cpp b/Engine/Plugins/Experimental/ProceduralVegetationEditor/Source/ProceduralVegetationEditor/Private/PVEditor.cpp
--- a/Engine/Plugins/Experimental/ProceduralVegetationEditor/Source/ProceduralVegetationEditor/Private/PVEditor.cpp
+++ b/Engine/Plugins/Experimental/ProceduralVegetationEditor/Source/ProceduralVegetationEditor/Private/PVEditor.cpp
@@ -230,8 +230,8 @@ TSharedRef<FTabManager::FLayout> FPVEditor::GetDefaultLayout() const
 				(
 					FTabManager::NewStack()
 					->SetSizeCoefficient(0.95f)
-					->AddTab(GetPanelID(EPCGEditorPanel::Viewport1), ETabState::OpenedTab)
-					->SetForegroundTab(GetPanelID(EPCGEditorPanel::Viewport1))
+					->AddTab(PCGEditorTabs::ViewportID[0], ETabState::OpenedTab)
+					->SetForegroundTab(FTabId(PCGEditorTabs::ViewportID[0]))
 					->SetHideTabWell(true)
 				)
 			)
@@ -244,15 +244,15 @@ TSharedRef<FTabManager::FLayout> FPVEditor::GetDefaultLayout() const
 				(
 					FTabManager::NewStack()
 					->SetSizeCoefficient(0.6f)
-					->AddTab(GetPanelID(EPCGEditorPanel::GraphEditor), ETabState::OpenedTab)
-					->AddTab(GetPanelID(EPCGEditorPanel::UserParams), ETabState::ClosedTab)
+					->AddTab(PCGEditorTabs::GraphEditorID, ETabState::OpenedTab)
+					->AddTab(PCGEditorTabs::UserParamsID, ETabState::ClosedTab)
 					->SetHideTabWell(true)
 				)
 				->Split
 				(
 					FTabManager::NewStack()
 					->SetSizeCoefficient(0.4f)
-					->AddTab(GetPanelID(EPCGEditorPanel::PropertyDetails1), ETabState::OpenedTab)
+					->AddTab(PCGEditorTabs::PropertyDetailsID[0], ETabState::OpenedTab)
 				)
 			)
 		)
@@ -262,12 +262,12 @@ TSharedRef<FTabManager::FLayout> FPVEditor::GetDefaultLayout() const
 bool FPVEditor::IsPanelAvailable(const FName PanelID) const
 {
 	static TSet<FName> SupportedIDs = {
-		GetPanelID(EPCGEditorPanel::Find),
-		GetPanelID(EPCGEditorPanel::GraphEditor),
-		GetPanelID(EPCGEditorPanel::Log),
-		GetPanelID(EPCGEditorPanel::PropertyDetails1),
-		GetPanelID(EPCGEditorPanel::UserParams),
-		GetPanelID(EPCGEditorPanel::Viewport1)
+		PCGEditorTabs::FindID,
+		PCGEditorTabs::GraphEditorID,
+		PCGEditorTabs::LogID,
+		PCGEditorTabs::PropertyDetailsID[0],
+		PCGEditorTabs::UserParamsID,
+		PCGEditorTabs::ViewportID[0]
 	};
 	return SupportedIDs.Contains(PanelID);
 }
@@ -314,12 +314,14 @@ void FPVEditor::ChangeNodeInspection(UPCGEditorGraphNodeBase* InNode)
 {
 	if (UPVBaseSettings* Settings = GetNodeSettings(InNode))
 	{
-		if (ViewportWidgets.Num() == 0)
+		TSharedPtr<SPVEditorViewport> Viewport = StaticCastSharedPtr<SPVEditorViewport>(GetViewportWidget(0));
+		
+		if (!Viewport.IsValid())
 		{
 			return;
 		}
 
-		ViewportWidgets[0]->OnNodeInspectionChanged(Settings);
+		Viewport->OnNodeInspectionChanged(Settings);
 		
 		NodeBeingInspected = InNode;
 		
@@ -795,18 +797,23 @@ FReply FPVEditor::ExportInternal()
 
 void FPVEditor::UpdateStats()
 {
-	if (ViewportWidgets.IsEmpty() || NodeBeingInspected == nullptr)
+	TSharedPtr<SPVEditorViewport> Viewport = StaticCastSharedPtr<SPVEditorViewport>(GetViewportWidget(0));
+	if (!Viewport.IsValid() || NodeBeingInspected == nullptr)
 	{
 		return;
 	}
 
 	const UPVBaseSettings* NodeBeingInspectedSettings = GetNodeSettings(NodeBeingInspected);
+	if (!NodeBeingInspectedSettings)
+	{
+		return;
+	}
 
 	TArray<FText> OverlayStats;
 	GatherStats(OverlayStats);
-	ViewportWidgets[0]->PopulateStatsOverlayText(OverlayStats);
+	Viewport->PopulateStatsOverlayText(OverlayStats);
 	
-	ViewportWidgets[0]->SetOverlayText(
+	Viewport->SetOverlayText(
 		NodeBeingInspectedSettings->bLockInspection
 		? NodeBeingInspected->GetNodeTitle(ENodeTitleType::FullTitle)
 		: FText::GetEmpty()
@@ -837,14 +844,7 @@ void FPVEditor::OnLockNodeSelection()
 
 TSharedRef<SPCGEditorViewport> FPVEditor::CreateViewportWidget()
 {
-	TSharedPtr<SPVEditorViewport> ViewportWidget;
-	
-	SAssignNew(ViewportWidget, SPVEditorViewport);
-
-	ViewportWidgets.Add(ViewportWidget.ToSharedRef());
-
-	return ViewportWidget.ToSharedRef();
-	
+	return SNew(SPVEditorViewport);
 }
 
 FPCGStack FPVEditor::GetPVDataStack(const UPCGNode* InNode, const UPCGPin*& OutPin) const
diff --git a/Engine/Plugins/Experimental/ProceduralVegetationEditor/Source/ProceduralVegetationEditor/Private/PVEditor.h b/Engine/Plugins/Experimental/ProceduralVegetationEditor/Source/ProceduralVegetationEditor/Private/PVEditor.h
--- a/Engine/Plugins/Experimental/ProceduralVegetationEditor/Source/ProceduralVegetationEditor/Private/PVEditor.h
+++ b/Engine/Plugins/Experimental/ProceduralVegetationEditor/Source/ProceduralVegetationEditor/Private/PVEditor.h
@@ -62,7 +62,6 @@ private:
 	TSharedPtr<FManagedArrayCollection> CollectionBeingInspected = nullptr;
 	TObjectPtr<UPCGEditorGraphNodeBase> NodeBeingInspected = nullptr;
 	
-	TArray<TSharedRef<SPVEditorViewport>> ViewportWidgets;
 	TSharedPtr<SCollectionSpreadSheetWidget> CollectionSpreadSheetWidget;
 
 	TWeakPtr<SDockTab> CollectionSpreadSheetTab;
diff --git a/Engine/Plugins/PCG/Source/PCG/Private/PCGGraph.cpp b/Engine/Plugins/PCG/Source/PCG/Private/PCGGraph.cpp
--- a/Engine/Plugins/PCG/Source/PCG/Private/PCGGraph.cpp
+++ b/Engine/Plugins/PCG/Source/PCG/Private/PCGGraph.cpp
@@ -159,6 +159,19 @@ namespace PCGGraphUtils
 * UPCGGraphInterface
 ****************************/
 
+FText UPCGGraphInterface::GetDisplayName() const
+{
+#if WITH_EDITOR
+	TOptional<FText> OverrideTitle = GetTitleOverride();
+	if (OverrideTitle.IsSet())
+	{
+		return OverrideTitle.GetValue();
+	}
+#endif
+
+	return FText::FromString(FName::NameToDisplayString(GetName(), /*bIsBool=*/false));
+}
+
 #if WITH_EDITOR
 bool UPCGGraphInterface::IsStandaloneGraph() const
 {
diff --git a/Engine/Plugins/PCG/Source/PCG/Private/PCGSubgraph.cpp b/Engine/Plugins/PCG/Source/PCG/Private/PCGSubgraph.cpp
--- a/Engine/Plugins/PCG/Source/PCG/Private/PCGSubgraph.cpp
+++ b/Engine/Plugins/PCG/Source/PCG/Private/PCGSubgraph.cpp
@@ -320,8 +320,7 @@ FString UPCGSubgraphSettings::GetAdditionalTitleInformation() const
 
 	if (UPCGGraph* TargetSubgraph = GetSubgraph())
 	{
-		// Use the same transformation than in the palette view to add spaces between uppercase characters
-		return FName::NameToDisplayString(TargetSubgraph->GetName(), /*bIsBool=*/false);
+		return TargetSubgraph->GetDisplayName().ToString();
 	}
 	else
 	{
diff --git a/Engine/Plugins/PCG/Source/PCG/Public/PCGGraph.h b/Engine/Plugins/PCG/Source/PCG/Public/PCGGraph.h
--- a/Engine/Plugins/PCG/Source/PCG/Public/PCGGraph.h
+++ b/Engine/Plugins/PCG/Source/PCG/Public/PCGGraph.h
@@ -135,6 +135,8 @@ public:
 	/** Construct a data that will contain a copy of the UserParameters, in order to freeze them while executing a graph. */
 	UE_API UPCGUserParametersData* ConstructUserParametersData(FPCGContext* InContext) const;
 
+	UE_API virtual FText GetDisplayName() const;
+
 #if WITH_EDITOR
 	FOnPCGGraphChanged OnGraphChangedDelegate;
 	FOnPCGGraphParametersChanged OnGraphParametersChangedDelegate;
@@ -143,7 +145,7 @@ public:
 	virtual TOptional<FText> GetTitleOverride() const { return bOverrideTitle ? Title : TOptional<FText>(); }
 	virtual TOptional<FLinearColor> GetColorOverride() const { return bOverrideColor ? Color : TOptional<FLinearColor>(); }
 	virtual TOptional<FPCGGraphToolData> GetGraphToolData() const { return {}; }
-
+	
 	UE_API static bool IsStandaloneGraphAsset(const FAssetData& InAssetData);
 
 	UFUNCTION()
diff --git a/Engine/Plugins/PCG/Source/PCGEditor/Private/Details/PCGCustomHLSLSettingsDetails.cpp b/Engine/Plugins/PCG/Source/PCGEditor/Private/Details/PCGCustomHLSLSettingsDetails.cpp
--- a/Engine/Plugins/PCG/Source/PCGEditor/Private/Details/PCGCustomHLSLSettingsDetails.cpp
+++ b/Engine/Plugins/PCG/Source/PCGEditor/Private/Details/PCGCustomHLSLSettingsDetails.cpp
@@ -29,7 +29,7 @@ void FPCGCustomHLSLSettingsDetails::CustomizeDetails(IDetailLayoutBuilder& Detai
 		const TSharedPtr<const FPCGEditor> PCGEditorSharedPtr = PCGEditorWeakPtr.Pin();
 
 		// Add a "open source editor" button if the source editor tab is not already visible.
-		if (PCGEditorSharedPtr.IsValid() && !PCGEditorSharedPtr->IsPanelCurrentlyOpen(EPCGEditorPanel::NodeSource))
+		if (PCGEditorSharedPtr.IsValid() && !PCGEditorSharedPtr->IsPanelCurrentlyOpen(PCGEditorTabs::HLSLSourceID))
 		{
 			SettingsCategory.AddCustomRow(FText::GetEmpty())
 				.Visibility(
@@ -37,7 +37,7 @@ void FPCGCustomHLSLSettingsDetails::CustomizeDetails(IDetailLayoutBuilder& Detai
 				[PCGEditorWeakPtr]
 				{
 					const TSharedPtr<const FPCGEditor> SharedPtr = PCGEditorWeakPtr.Pin();
-					return (SharedPtr.IsValid() && SharedPtr->IsPanelCurrentlyOpen(EPCGEditorPanel::NodeSource)) ? EVisibility::Collapsed : EVisibility::Visible;
+					return (SharedPtr.IsValid() && SharedPtr->IsPanelCurrentlyOpen(PCGEditorTabs::HLSLSourceID)) ? EVisibility::Collapsed : EVisibility::Visible;
 				}))
 				.ValueContent()
 				.MaxDesiredWidth(120.f)
@@ -53,7 +53,7 @@ void FPCGCustomHLSLSettingsDetails::CustomizeDetails(IDetailLayoutBuilder& Detai
 						{
 							if (const TSharedPtr<const FPCGEditor> SharedPtr = PCGEditorWeakPtr.Pin())
 							{
-								SharedPtr->BringFocusToPanel(EPCGEditorPanel::NodeSource);
+								SharedPtr->BringFocusToPanel(PCGEditorTabs::HLSLSourceID);
 							}
 
 							return FReply::Handled();
diff --git a/Engine/Plugins/PCG/Source/PCGEditor/Private/Details/PCGGraphDetails.cpp b/Engine/Plugins/PCG/Source/PCGEditor/Private/Details/PCGGraphDetails.cpp
--- a/Engine/Plugins/PCG/Source/PCGEditor/Private/Details/PCGGraphDetails.cpp
+++ b/Engine/Plugins/PCG/Source/PCGEditor/Private/Details/PCGGraphDetails.cpp
@@ -64,7 +64,7 @@ void FPCGGraphDetails::CustomizeDetails(IDetailLayoutBuilder& DetailBuilder)
 			[PCGEditorWeakPtr]
 			{
 				const TSharedPtr<const FPCGEditor> PCGEditorSharedPtr = PCGEditorWeakPtr.Pin();
-				return (PCGEditorSharedPtr.IsValid() && PCGEditorSharedPtr->IsPanelCurrentlyOpen(EPCGEditorPanel::UserParams)) ? EVisibility::Collapsed : EVisibility::Visible;
+				return (PCGEditorSharedPtr.IsValid() && PCGEditorSharedPtr->IsPanelCurrentlyOpen(PCGEditorTabs::UserParamsID)) ? EVisibility::Collapsed : EVisibility::Visible;
 			}))
 			.ValueContent()
 			.MaxDesiredWidth(120.f)
@@ -81,7 +81,7 @@ void FPCGGraphDetails::CustomizeDetails(IDetailLayoutBuilder& DetailBuilder)
 					{
 						if (const TSharedPtr<const FPCGEditor> SharedPtr = PCGEditorWeakPtr.Pin())
 						{
-							SharedPtr->BringFocusToPanel(EPCGEditorPanel::UserParams);
+							SharedPtr->BringFocusToPanel(PCGEditorTabs::UserParamsID);
 						}
 
 						return FReply::Handled();
diff --git a/Engine/Plugins/PCG/Source/PCGEditor/Private/PCGEditor.cpp b/Engine/Plugins/PCG/Source/PCGEditor/Private/PCGEditor.cpp
--- a/Engine/Plugins/PCG/Source/PCGEditor/Private/PCGEditor.cpp
+++ b/Engine/Plugins/PCG/Source/PCGEditor/Private/PCGEditor.cpp
@@ -20,10 +20,12 @@
 #include "Tests/Determinism/PCGDeterminismTestsCommon.h"
 
 #include "PCGEditorCommands.h"
+#include "PCGEditorDefaultMode.h"
 #include "PCGEditorGraph.h"
 #include "PCGEditorMenuContext.h"
 #include "PCGEditorSettings.h"
 #include "PCGEditorStyle.h"
+#include "PCGEditorTabFactories.h"
 #include "PCGEditorUtils.h"
 #include "Managers/PCGEditorInspectionDataManager.h"
 #include "Nodes/PCGEditorGraphNode.h"
@@ -76,66 +78,41 @@
 #include "Widgets/Docking/SDockTab.h"
 #include "Widgets/Input/SEditableTextBox.h"
 #include "Widgets/Notifications/SNotificationList.h"
+#include "WorkflowOrientedApp/WorkflowTabManager.h"
+#include "WorkflowOrientedApp/WorkflowTabFactory.h"
 
 #define LOCTEXT_NAMESPACE "PCGGraphEditor"
 
 namespace FPCGEditor_private
 {
-	const FName GraphEditorID = FName(TEXT("GraphEditor"));
-	const FName PropertyDetailsID[] = {
-		FName(TEXT("PropertyDetails")),
-		FName(TEXT("PropertyDetails2")),
-		FName(TEXT("PropertyDetails3")),
-		FName(TEXT("PropertyDetails4")) };
-	const FName PaletteID = FName(TEXT("Palette"));
-	const FName DebugObjectID = FName(TEXT("DebugObject"));
-	const FName AttributesID[] = {
-		FName(TEXT("Attributes")),
-		FName(TEXT("Attributes2")),
-		FName(TEXT("Attributes3")),
-		FName(TEXT("Attributes4")) };
-	const FName FindID = FName(TEXT("Find"));
-	const FName DeterminismID = FName(TEXT("Determinism"));
-	const FName ProfilingID = FName(TEXT("Profiling"));
-	const FName LogID = FName(TEXT("Log"));
-	const FName HLSLSourceID = FName(TEXT("HLSLSource"));
-	const FName UserParamsID = FName(TEXT("UserParams"));
-	const FName ViewportID[] = {
-		FName(TEXT("Viewport")),
-		FName(TEXT("Viewport2")),
-		FName(TEXT("Viewport3")),
-		FName(TEXT("Viewport4")) };
-
-	const FText UserParamsTabName = LOCTEXT("UserParamsTab", "Graph Parameters");
-
 	const FName AllDefaultPanels[] =
 	{
-		GraphEditorID,
-		PropertyDetailsID[0],
-		PropertyDetailsID[1],
-		PropertyDetailsID[2],
-		PropertyDetailsID[3],
-		PaletteID,
-		DebugObjectID,
-		AttributesID[0],
-		AttributesID[1],
-		AttributesID[2],
-		AttributesID[3],
-		FindID,
-		DeterminismID,
-		ProfilingID,
-		LogID,
-		HLSLSourceID,
-		UserParamsID,
-		ViewportID[0],
-		ViewportID[1],
-		ViewportID[2],
-		ViewportID[3]
+		PCGEditorTabs::GraphEditorID,
+		PCGEditorTabs::PropertyDetailsID[0],
+		PCGEditorTabs::PropertyDetailsID[1],
+		PCGEditorTabs::PropertyDetailsID[2],
+		PCGEditorTabs::PropertyDetailsID[3],
+		PCGEditorTabs::PaletteID,
+		PCGEditorTabs::DebugObjectID,
+		PCGEditorTabs::AttributesID[0],
+		PCGEditorTabs::AttributesID[1],
+		PCGEditorTabs::AttributesID[2],
+		PCGEditorTabs::AttributesID[3],
+		PCGEditorTabs::FindID,
+		PCGEditorTabs::DeterminismID,
+		PCGEditorTabs::ProfilingID,
+		PCGEditorTabs::LogID,
+		PCGEditorTabs::HLSLSourceID,
+		PCGEditorTabs::UserParamsID,
+		PCGEditorTabs::ViewportID[0],
+		PCGEditorTabs::ViewportID[1],
+		PCGEditorTabs::ViewportID[2],
+		PCGEditorTabs::ViewportID[3]
 	};
 
 	const FName AllRestrictedPanels[] =
 	{
-		HLSLSourceID,
+		PCGEditorTabs::HLSLSourceID,
 	};
 }
 
@@ -190,42 +167,12 @@ void FPCGEditor::Initialize(const EToolkitMode::Type InMode, const TSharedPtr<cl
 	PCGGraphBeingEdited->PCGEditorGraph->SetEditor(SharedThis(this));
 	PCGEditorGraph = PCGGraphBeingEdited->PCGEditorGraph;
 
-	for (int PropertyDetailsIndex = 0; PropertyDetailsIndex < 4; ++PropertyDetailsIndex)
-	{
-		TSharedRef<SPCGEditorGraphDetailsView> PropertyDetailsWidget = SNew(SPCGEditorGraphDetailsView);
-		PropertyDetailsWidget->SetEditor(SharedThis(this));
-		PropertyDetailsWidget->SetObject(PCGGraphBeingEdited);
-		PropertyDetailsWidgets.Add(PropertyDetailsWidget);
-	}
-
-	GraphEditorWidget = CreateGraphEditorWidget();
-	PaletteWidget = CreatePaletteWidget();
-	DebugObjectTreeWidget = CreateDebugObjectTreeWidget();
-	FindWidget = CreateFindWidget();
-
-	for (int AttributesIndex = 0; AttributesIndex < FPCGEditorInspectionDataManager::NumberOfEntries; ++AttributesIndex)
-	{
-		const int ViewportEditorPanel = static_cast<int>(EPCGEditorPanel::Viewport1) + AttributesIndex;
-
-		AttributesWidgets.Add(CreateAttributesWidget(AttributesIndex));
-		AttributesWidgets[AttributesIndex]->SetViewportWidget(CreateViewportWidget(), static_cast<EPCGEditorPanel>(ViewportEditorPanel));
-	}
+	GetOrCreateGraphEditorWidget();
+	GetOrCreateNodeSourceWidget();
 	
-	DeterminismWidget = CreateDeterminismWidget();
-	ProfilingWidget = CreateProfilingWidget();
-	LogWidget = CreateLogWidget();
-	NodeSourceWidget = CreateNodeSourceWidget();
-	UserParamsWidget = CreateGraphParamsWidget();
-
 	BindCommands();
 	RegisterToolbar();
 
-	const TSharedRef<FTabManager::FLayout> StandaloneDefaultLayout = GetDefaultLayout();
-
-	const FName PCGGraphEditorAppName = FName(TEXT("PCGEditorApp"));
-
-	InitAssetEditor(InMode, InToolkitHost, PCGGraphEditorAppName, StandaloneDefaultLayout, /*bCreateDefaultStandaloneMenu=*/ true, /*bCreateDefaultToolbar=*/ true, InObjectToEdit);
-
 	PCGGraphBeingEdited->OnGraphChangedDelegate.AddRaw(this, &FPCGEditor::OnGraphChanged);
 	PCGGraphBeingEdited->OnNodeSourceCompiledDelegate.AddRaw(this, &FPCGEditor::OnNodeSourceCompiled);
 
@@ -258,6 +205,19 @@ void FPCGEditor::Initialize(const EToolkitMode::Type InMode, const TSharedPtr<cl
 			PCGEditorGraphNode->SetInspected(false);
 		}
 	}
+
+	DocumentManager = MakeShared<FDocumentTracker>(PCGEditorTabs::PCGEditorDocumentID);
+
+	const FName PCGGraphEditorAppName = FName(TEXT("PCGEditorApp"));
+	InitAssetEditor(InMode, InToolkitHost, PCGGraphEditorAppName, FTabManager::FLayout::NullLayout, /*bCreateDefaultStandaloneMenu=*/ true, /*bCreateDefaultToolbar=*/ true, InObjectToEdit);
+
+	TSharedPtr<FPCGEditor> ThisPtr = SharedThis<FPCGEditor>(this);
+	DocumentManager->Initialize(AsShared());
+
+	// @todo_pcg: add document factories here
+
+	AddApplicationMode(FPCGEditorDefaultMode::StaticName(), MakeShared<FPCGEditorDefaultMode>(ThisPtr));
+	SetCurrentMode(FPCGEditorDefaultMode::StaticName());
 }
 
 UPCGEditorGraph* FPCGEditor::GetPCGEditorGraph()
@@ -330,7 +290,7 @@ void FPCGEditor::UpdateAfterInspectedStackChanged(const FPCGStack& FullStack)
 			const bool bHasBeenGeneratedThisSession = Component && Component->bGenerated && Component->WasGeneratedThisSession();
 			const bool bWasAborted = LastExecutionStatus.IsSet() && LastExecutionStatus.GetValue().Key == Source && LastExecutionStatus.GetValue().Value == EPCGGenerationStatus::Aborted;
 			const bool bWasInspecting = Source->GetExecutionState().GetInspection().IsInspecting();
-			const bool bNeedsInspection = Algo::AnyOf(AttributesWidgets, [](const TSharedPtr<SPCGEditorGraphAttributeListView>& ALV) { return ALV->GetNodeBeingInspected() != nullptr; });
+			const bool bNeedsInspection = Algo::AnyOf(AttributesWidgets, [](const TSharedPtr<SPCGEditorGraphAttributeListView>& ALV) { return ALV.IsValid() && ALV->GetNodeBeingInspected() != nullptr; });
 
 			if (!bHasBeenGeneratedThisSession || (bNeedsInspection && !bWasInspecting))
 			{
@@ -466,7 +426,10 @@ const FPCGStack* FPCGEditor::GetStackBeingInspected() const
 
 void FPCGEditor::SetSourceEditorTargetObject(UObject* InObject)
 {
-	NodeSourceWidget->SetTextProviderObject(InObject);
+	if (NodeSourceWidget.IsValid())
+	{
+		NodeSourceWidget->SetTextProviderObject(InObject);
+	}
 }
 
 void FPCGEditor::JumpToNode(const UEdGraphNode* InNode)
@@ -506,56 +469,58 @@ void FPCGEditor::JumpToNode(const UPCGNode* InNode)
 	}
 }
 
+PRAGMA_DISABLE_DEPRECATION_WARNINGS
 FName FPCGEditor::GetPanelID(const EPCGEditorPanel Panel) const
 {
 	switch (Panel)
 	{
 		case EPCGEditorPanel::Attributes1:
-			return FPCGEditor_private::AttributesID[0];
+			return PCGEditorTabs::AttributesID[0];
 		case EPCGEditorPanel::Attributes2:
-			return FPCGEditor_private::AttributesID[1];
+			return PCGEditorTabs::AttributesID[1];
 		case EPCGEditorPanel::Attributes3:
-			return FPCGEditor_private::AttributesID[2];
+			return PCGEditorTabs::AttributesID[2];
 		case EPCGEditorPanel::Attributes4:
-			return FPCGEditor_private::AttributesID[3];
+			return PCGEditorTabs::AttributesID[3];
 		case EPCGEditorPanel::DebugObjectTree:
-			return FPCGEditor_private::DebugObjectID;
+			return PCGEditorTabs::DebugObjectID;
 		case EPCGEditorPanel::Determinism:
-			return FPCGEditor_private::DeterminismID;
+			return PCGEditorTabs::DeterminismID;
 		case EPCGEditorPanel::Find:
-			return FPCGEditor_private::FindID;
+			return PCGEditorTabs::FindID;
 		case EPCGEditorPanel::GraphEditor:
-			return FPCGEditor_private::GraphEditorID;
+			return PCGEditorTabs::GraphEditorID;
 		case EPCGEditorPanel::Log:
-			return FPCGEditor_private::LogID;
+			return PCGEditorTabs::LogID;
 		case EPCGEditorPanel::NodePalette:
-			return FPCGEditor_private::PaletteID;
+			return PCGEditorTabs::PaletteID;
 		case EPCGEditorPanel::NodeSource:
-			return FPCGEditor_private::HLSLSourceID;
+			return PCGEditorTabs::HLSLSourceID;
 		case EPCGEditorPanel::Profiling:
-			return FPCGEditor_private::ProfilingID;
+			return PCGEditorTabs::ProfilingID;
 		case EPCGEditorPanel::PropertyDetails1:
-			return FPCGEditor_private::PropertyDetailsID[0];
+			return PCGEditorTabs::PropertyDetailsID[0];
 		case EPCGEditorPanel::PropertyDetails2:
-			return FPCGEditor_private::PropertyDetailsID[1];
+			return PCGEditorTabs::PropertyDetailsID[1];
 		case EPCGEditorPanel::PropertyDetails3:
-			return FPCGEditor_private::PropertyDetailsID[2];
+			return PCGEditorTabs::PropertyDetailsID[2];
 		case EPCGEditorPanel::PropertyDetails4:
-			return FPCGEditor_private::PropertyDetailsID[3];
+			return PCGEditorTabs::PropertyDetailsID[3];
 		case EPCGEditorPanel::UserParams:
-			return FPCGEditor_private::UserParamsID;
+			return PCGEditorTabs::UserParamsID;
 		case EPCGEditorPanel::Viewport1:
-			return FPCGEditor_private::ViewportID[0];
+			return PCGEditorTabs::ViewportID[0];
 		case EPCGEditorPanel::Viewport2:
-			return FPCGEditor_private::ViewportID[1];
+			return PCGEditorTabs::ViewportID[1];
 		case EPCGEditorPanel::Viewport3:
-			return FPCGEditor_private::ViewportID[2];
+			return PCGEditorTabs::ViewportID[2];
 		case EPCGEditorPanel::Viewport4:
-			return FPCGEditor_private::ViewportID[3];
+			return PCGEditorTabs::ViewportID[3];
 		default:
 			return NAME_None;
 	}
 }
+PRAGMA_ENABLE_DEPRECATION_WARNINGS
 
 void FPCGEditor::BringFocusToPanel(FName PanelID) const
 {
@@ -606,87 +571,10 @@ bool FPCGEditor::IsPanelAvailable(const FName PanelID) const
 
 void FPCGEditor::RegisterTabSpawners(const TSharedRef<FTabManager>& InTabManager)
 {
-	WorkspaceMenuCategory = InTabManager->AddLocalWorkspaceMenuCategory(LOCTEXT("WorkspaceMenu_PCGEditor", "PCG Editor"));
-	TSharedRef<FWorkspaceItem> DetailsGroup = WorkspaceMenuCategory->AddGroup(LOCTEXT("WorkspaceMenu_PCGEditor_Details", "Details"));
-	TSharedRef<FWorkspaceItem> AttributesGroup = WorkspaceMenuCategory->AddGroup(LOCTEXT("WorkspaceMenu_PCGEditor_Attributes", "Attributes"));
-	TSharedRef<FWorkspaceItem> ViewportGroup = WorkspaceMenuCategory->AddGroup(LOCTEXT("WorkspaceMenu_PCGEditor_Viewport", "Data Viewport"));
-	const TSharedRef<FWorkspaceItem>& WorkspaceMenuCategoryRef = WorkspaceMenuCategory.ToSharedRef();
-
 	FAssetEditorToolkit::RegisterTabSpawners(InTabManager);
+	FWorkflowCentricApplication::RegisterTabSpawners(InTabManager);
 
-	using SpawnOnTabCallback = decltype(FOnSpawnTab::CreateSP(this, &FPCGEditor::SpawnTab_GraphEditor));
-
-	struct TabInfo
-	{
-		FName ID;
-		SpawnOnTabCallback Callback;
-		FText DisplayName;
-		TSharedRef<FWorkspaceItem> Group;
-	};
-
-	TabInfo DefaultTabInfo[] =
-	{
-		{FPCGEditor_private::GraphEditorID, FOnSpawnTab::CreateSP(this, &FPCGEditor::SpawnTab_GraphEditor), LOCTEXT("GraphTab", "Graph"), WorkspaceMenuCategoryRef},
-		{FPCGEditor_private::PropertyDetailsID[0], FOnSpawnTab::CreateSP(this, &FPCGEditor::SpawnTab_PropertyDetails, 0), LOCTEXT("DetailsTab1", "Details 1"), DetailsGroup},
-		{FPCGEditor_private::PropertyDetailsID[1], FOnSpawnTab::CreateSP(this, &FPCGEditor::SpawnTab_PropertyDetails, 1), LOCTEXT("DetailsTab2", "Details 2"), DetailsGroup},
-		{FPCGEditor_private::PropertyDetailsID[2], FOnSpawnTab::CreateSP(this, &FPCGEditor::SpawnTab_PropertyDetails, 2), LOCTEXT("DetailsTab3", "Details 3"), DetailsGroup},
-		{FPCGEditor_private::PropertyDetailsID[3], FOnSpawnTab::CreateSP(this, &FPCGEditor::SpawnTab_PropertyDetails, 3), LOCTEXT("DetailsTab4", "Details 4"), DetailsGroup},
-		{FPCGEditor_private::PaletteID, FOnSpawnTab::CreateSP(this, &FPCGEditor::SpawnTab_Palette), LOCTEXT("PaletteTab", "Palette"), WorkspaceMenuCategoryRef},
-		{FPCGEditor_private::DebugObjectID, FOnSpawnTab::CreateSP(this, &FPCGEditor::SpawnTab_DebugObjectTree), LOCTEXT("DebugTab", "Debug Object Tree"), WorkspaceMenuCategoryRef},
-		{FPCGEditor_private::AttributesID[0], FOnSpawnTab::CreateSP(this, &FPCGEditor::SpawnTab_Attributes, 0), LOCTEXT("AttributesTab1", "Attributes 1"), AttributesGroup},
-		{FPCGEditor_private::AttributesID[1], FOnSpawnTab::CreateSP(this, &FPCGEditor::SpawnTab_Attributes, 1), LOCTEXT("AttributesTab2", "Attributes 2"), AttributesGroup},
-		{FPCGEditor_private::AttributesID[2], FOnSpawnTab::CreateSP(this, &FPCGEditor::SpawnTab_Attributes, 2), LOCTEXT("AttributesTab3", "Attributes 3"), AttributesGroup},
-		{FPCGEditor_private::AttributesID[3], FOnSpawnTab::CreateSP(this, &FPCGEditor::SpawnTab_Attributes, 3), LOCTEXT("AttributesTab4", "Attributes 4"), AttributesGroup},
-		{FPCGEditor_private::FindID, FOnSpawnTab::CreateSP(this, &FPCGEditor::SpawnTab_Find), LOCTEXT("FindTab", "Find"), WorkspaceMenuCategoryRef},
-		{FPCGEditor_private::DeterminismID, FOnSpawnTab::CreateSP(this, &FPCGEditor::SpawnTab_Determinism), LOCTEXT("DeterminismTab", "Determinism"), WorkspaceMenuCategoryRef},
-		{FPCGEditor_private::ProfilingID, FOnSpawnTab::CreateSP(this, &FPCGEditor::SpawnTab_Profiling), LOCTEXT("ProfilingTab", "Profiling"), WorkspaceMenuCategoryRef},
-		{FPCGEditor_private::LogID, FOnSpawnTab::CreateSP(this, &FPCGEditor::SpawnTab_Log), LOCTEXT("LogCaptureTab", "Log Capture"), WorkspaceMenuCategoryRef},
-		{FPCGEditor_private::HLSLSourceID, FOnSpawnTab::CreateSP(this, &FPCGEditor::SpawnTab_NodeSource), LOCTEXT("NodeSourceTab", "HLSL Source"), WorkspaceMenuCategoryRef},
-		{FPCGEditor_private::UserParamsID, FOnSpawnTab::CreateSP(this, &FPCGEditor::SpawnTab_UserParams), FPCGEditor_private::UserParamsTabName, WorkspaceMenuCategoryRef},
-		{FPCGEditor_private::ViewportID[0], FOnSpawnTab::CreateSP(this, &FPCGEditor::SpawnTab_Viewport, 0), LOCTEXT("ViewportTab1", "Viewport 1"), ViewportGroup},
-		{FPCGEditor_private::ViewportID[1], FOnSpawnTab::CreateSP(this, &FPCGEditor::SpawnTab_Viewport, 1), LOCTEXT("ViewportTab2", "Viewport 2"), ViewportGroup},
-		{FPCGEditor_private::ViewportID[2], FOnSpawnTab::CreateSP(this, &FPCGEditor::SpawnTab_Viewport, 2), LOCTEXT("ViewportTab3", "Viewport 3"), ViewportGroup},
-		{FPCGEditor_private::ViewportID[3], FOnSpawnTab::CreateSP(this, &FPCGEditor::SpawnTab_Viewport, 3), LOCTEXT("ViewportTab4", "Viewport 4"), ViewportGroup},
-	};
-
-	for (const TabInfo& Info : DefaultTabInfo)
-	{
-		if (!IsPanelAvailable(Info.ID))
-		{
-			continue;
-		}
-
-		InTabManager->RegisterTabSpawner(Info.ID, Info.Callback)
-		.SetDisplayName(Info.DisplayName)
-		.SetGroup(Info.Group);
-	}
-}
-
-void FPCGEditor::UnregisterTabSpawners(const TSharedRef<class FTabManager>& InTabManager)
-{
-	InTabManager->UnregisterTabSpawner(FPCGEditor_private::GraphEditorID);
-	InTabManager->UnregisterTabSpawner(FPCGEditor_private::PropertyDetailsID[0]);
-	InTabManager->UnregisterTabSpawner(FPCGEditor_private::PropertyDetailsID[1]);
-	InTabManager->UnregisterTabSpawner(FPCGEditor_private::PropertyDetailsID[2]);
-	InTabManager->UnregisterTabSpawner(FPCGEditor_private::PropertyDetailsID[3]);
-	InTabManager->UnregisterTabSpawner(FPCGEditor_private::PaletteID);
-	InTabManager->UnregisterTabSpawner(FPCGEditor_private::DebugObjectID);
-	InTabManager->UnregisterTabSpawner(FPCGEditor_private::AttributesID[0]);
-	InTabManager->UnregisterTabSpawner(FPCGEditor_private::AttributesID[1]);
-	InTabManager->UnregisterTabSpawner(FPCGEditor_private::AttributesID[2]);
-	InTabManager->UnregisterTabSpawner(FPCGEditor_private::AttributesID[3]);
-	InTabManager->UnregisterTabSpawner(FPCGEditor_private::FindID);
-	InTabManager->UnregisterTabSpawner(FPCGEditor_private::DeterminismID);
-	InTabManager->UnregisterTabSpawner(FPCGEditor_private::ProfilingID);
-	InTabManager->UnregisterTabSpawner(FPCGEditor_private::LogID);
-	InTabManager->UnregisterTabSpawner(FPCGEditor_private::HLSLSourceID);
-	InTabManager->UnregisterTabSpawner(FPCGEditor_private::UserParamsID);
-	InTabManager->UnregisterTabSpawner(FPCGEditor_private::ViewportID[0]);
-	InTabManager->UnregisterTabSpawner(FPCGEditor_private::ViewportID[1]);
-	InTabManager->UnregisterTabSpawner(FPCGEditor_private::ViewportID[2]);
-	InTabManager->UnregisterTabSpawner(FPCGEditor_private::ViewportID[3]);
-
-	FAssetEditorToolkit::UnregisterTabSpawners(InTabManager);
+	DocumentManager->SetTabManager(InTabManager);
 }
 
 void FPCGEditor::AddReferencedObjects(FReferenceCollector& Collector)
@@ -698,7 +586,7 @@ void FPCGEditor::AddReferencedObjects(FReferenceCollector& Collector)
 
 	for (TSharedPtr<SPCGEditorGraphAttributeListView> ALV : AttributesWidgets)
 	{
-		if (ALV)
+		if (ALV.IsValid())
 		{
 			ALV->AddReferencedObjects(Collector);
 		}
@@ -749,6 +637,11 @@ void FPCGEditor::PostUndo(bool bSuccess)
 	}
 }
 
+FText FPCGEditor::GetToolkitName() const
+{
+	return PCGGraphBeingEdited->GetDisplayName();
+}
+
 FName FPCGEditor::GetToolkitFName() const
 {
 	return FName(TEXT("PCGEditor"));
@@ -852,7 +745,7 @@ void FPCGEditor::RegisterToolbarButton(FToolMenuSection& Section, EPCGToolbarBut
 		}
 	case EPCGToolbarButtons::OpenDebugObjectTreeTab:
 		{
-			if (IsPanelAvailable(FPCGEditor_private::DebugObjectID))
+			if (IsPanelAvailable(PCGEditorTabs::DebugObjectID))
 			{
 				Section.AddEntry(FToolMenuEntry::InitToolBarButton(
 					FPCGEditorCommands::Get().OpenDebugObjectTreeTab,
@@ -1014,7 +907,7 @@ void FPCGEditor::OnFind()
 {
 	if (TabManager.IsValid() && FindWidget.IsValid())
 	{
-		TabManager->TryInvokeTab(FPCGEditor_private::FindID);
+		TabManager->TryInvokeTab(FTabId(PCGEditorTabs::FindID));
 		FindWidget->FocusForUse();
 	}
 }
@@ -1025,14 +918,14 @@ void FPCGEditor::OpenDetailsView()
 	{
 		auto InvokeFirstUnlockedTab = [this](bool bVisibleOnly) -> bool
 		{
-			for (int DetailsViewIndex = 0; DetailsViewIndex < PropertyDetailsWidgets.Num(); ++DetailsViewIndex)
+			for (int DetailsViewIndex = 0; DetailsViewIndex < PCGEditorTabs::NumPropertyDetailTabs; ++DetailsViewIndex)
 			{
 				TSharedPtr<SPCGEditorGraphDetailsView> DetailsView = PropertyDetailsWidgets[DetailsViewIndex];
 				if (DetailsView.IsValid() && !DetailsView->IsLocked())
 				{
-					if (!bVisibleOnly || TabManager->FindExistingLiveTab(FPCGEditor_private::PropertyDetailsID[DetailsViewIndex]))
+					if (!bVisibleOnly || TabManager->FindExistingLiveTab(FTabId(PCGEditorTabs::PropertyDetailsID[DetailsViewIndex])))
 					{
-						TabManager->TryInvokeTab(FPCGEditor_private::PropertyDetailsID[DetailsViewIndex]);
+						TabManager->TryInvokeTab(FTabId(PCGEditorTabs::PropertyDetailsID[DetailsViewIndex]));
 						return true;
 					}
 				}
@@ -1049,17 +942,14 @@ void FPCGEditor::OpenDetailsView()
 		// Default to first if they are all locked
 		if (PropertyDetailsWidgets[0].IsValid())
 		{
-			TabManager->TryInvokeTab(FPCGEditor_private::PropertyDetailsID[0]);
+			TabManager->TryInvokeTab(FTabId(PCGEditorTabs::PropertyDetailsID[0]));
 		}
 	}
 }
 
 void FPCGEditor::OnDetailsViewTabClosed(TSharedRef<SDockTab> DockTab, int Index)
 {
-	if (!PropertyDetailsWidgets.IsValidIndex(Index))
-	{
-		return;
-	}
+	check(Index >= 0 && Index < PCGEditorTabs::NumPropertyDetailTabs);
 
 	TSharedPtr<SPCGEditorGraphDetailsView> DetailsView = PropertyDetailsWidgets[Index];
 	if (DetailsView.IsValid() && DetailsView->IsLocked())
@@ -1070,10 +960,7 @@ void FPCGEditor::OnDetailsViewTabClosed(TSharedRef<SDockTab> DockTab, int Index)
 
 void FPCGEditor::OnAttributeListViewTabClosed(TSharedRef<SDockTab> DockTab, int Index)
 {
-	if (!AttributesWidgets.IsValidIndex(Index))
-	{
-		return;
-	}
+	check(Index >= 0 && Index < PCGEditorTabs::NumAttributeTabs);
 
 	TSharedPtr<SPCGEditorGraphAttributeListView> AttributeListView = AttributesWidgets[Index];
 	if (AttributeListView.IsValid())
@@ -1089,10 +976,10 @@ void FPCGEditor::OnAttributeListViewTabClosed(TSharedRef<SDockTab> DockTab, int
 		if (NodeInspected)
 		{
 			bool bIsStillInspectedOnVisibleTabs = false;
-			for (int OtherTabIndex = 0; OtherTabIndex < AttributesWidgets.Num(); ++OtherTabIndex)
+			for (int OtherTabIndex = 0; OtherTabIndex < PCGEditorTabs::NumAttributeTabs; ++OtherTabIndex)
 			{
 				TSharedPtr<SPCGEditorGraphAttributeListView> ALV = AttributesWidgets[OtherTabIndex];
-				if (ALV.IsValid() && ALV->GetNodeBeingInspected() == NodeInspected && TabManager->FindExistingLiveTab(FPCGEditor_private::AttributesID[OtherTabIndex]))
+				if (ALV.IsValid() && ALV->GetNodeBeingInspected() == NodeInspected && TabManager->FindExistingLiveTab(FTabId(PCGEditorTabs::AttributesID[OtherTabIndex])))
 				{
 					bIsStillInspectedOnVisibleTabs = true;
 					break;
@@ -1117,7 +1004,12 @@ void FPCGEditor::OnAttributeListViewTabClosed(TSharedRef<SDockTab> DockTab, int
 
 void FPCGEditor::OnViewportViewTabClosed(TSharedRef<SDockTab> DockTab, int Index)
 {
-	AttributesWidgets[Index]->ResetViewport();
+	check(Index >= 0 && Index < PCGEditorTabs::NumAttributeTabs);
+
+	if (AttributesWidgets[Index].IsValid())
+	{
+		AttributesWidgets[Index]->ResetViewport();
+	}
 }
 
 void FPCGEditor::OnPauseAutomaticRegeneration_Clicked()
@@ -1180,12 +1072,12 @@ bool FPCGEditor::IsCurrentlyGenerating() const
 
 bool FPCGEditor::IsDebugObjectTreeTabClosed() const
 {
-	return !TabManager.IsValid() || !TabManager->FindExistingLiveTab(FPCGEditor_private::DebugObjectID).IsValid();
+	return !TabManager.IsValid() || !TabManager->FindExistingLiveTab(FTabId(PCGEditorTabs::DebugObjectID)).IsValid();
 }
 
 void FPCGEditor::OnOpenDebugObjectTreeTab_Clicked()
 {
-	TabManager->TryInvokeTab(FPCGEditor_private::DebugObjectID);
+	TabManager->TryInvokeTab(FTabId(PCGEditorTabs::DebugObjectID));
 }
 
 bool FPCGEditor::CanRunDeterminismNodeTest() const
@@ -1299,7 +1191,7 @@ void FPCGEditor::OnDeterminismNodeTest() const
 	// Give focus to the Determinism Output Tab
 	if (TabManager.IsValid())
 	{
-		TabManager->TryInvokeTab(FPCGEditor_private::DeterminismID);
+		TabManager->TryInvokeTab(FTabId(PCGEditorTabs::DeterminismID));
 	}
 }
 
@@ -1346,7 +1238,7 @@ void FPCGEditor::OnDeterminismGraphTest() const
 	// Give focus to the Determinism Output Tab
 	if (TabManager.IsValid())
 	{
-		TabManager->TryInvokeTab(FPCGEditor_private::DeterminismID);
+		TabManager->TryInvokeTab(FTabId(PCGEditorTabs::DeterminismID));
 	}
 }
 
@@ -1372,7 +1264,10 @@ void FPCGEditor::OnEditGraphSettings()
 
 	for (TSharedPtr<SPCGEditorGraphDetailsView> PropertyDetailsWidget : PropertyDetailsWidgets)
 	{
-		PropertyDetailsWidget->SetObject(PCGGraphBeingEdited);
+		if (PropertyDetailsWidget.IsValid())
+		{
+			PropertyDetailsWidget->SetObject(PCGGraphBeingEdited);
+		}
 	}
 
 	OpenDetailsView();
@@ -1385,13 +1280,18 @@ bool FPCGEditor::IsEditGraphSettingsToggled() const
 		return false;
 	}
 
-	for (int32 WidgetIndex = 0; WidgetIndex < PropertyDetailsWidgets.Num(); ++WidgetIndex)
+	for (int32 WidgetIndex = 0; WidgetIndex < PCGEditorTabs::NumPropertyDetailTabs; ++WidgetIndex)
 	{
+		if (!PropertyDetailsWidgets[WidgetIndex].IsValid())
+		{
+			continue;
+		}
+
 		const TArray<TWeakObjectPtr<UObject>>& SelectedObjects = PropertyDetailsWidgets[WidgetIndex]->GetSelectedObjects();
 		// The only object selected should be the graph. If there is no details view panel open, leave it disabled.
 		if (SelectedObjects.Num() == 1 && SelectedObjects[0] == PCGGraphBeingEdited.Get())
 		{
-			if (const TSharedPtr<SDockTab>& Tab = TabManager->FindExistingLiveTab(FPCGEditor_private::PropertyDetailsID[WidgetIndex]))
+			if (const TSharedPtr<SDockTab>& Tab = TabManager->FindExistingLiveTab(FTabId(PCGEditorTabs::PropertyDetailsID[WidgetIndex])))
 			{
 				if (Tab.IsValid() && Tab->IsForeground())
 				{
@@ -1406,19 +1306,19 @@ bool FPCGEditor::IsEditGraphSettingsToggled() const
 
 void FPCGEditor::OnToggleGraphParamsPanel() const
 {
-	if (IsPanelCurrentlyForeground(EPCGEditorPanel::UserParams))
+	if (IsPanelCurrentlyForeground(PCGEditorTabs::UserParamsID))
 	{
-		CloseGraphPanel(EPCGEditorPanel::UserParams);
+		CloseGraphPanel(PCGEditorTabs::UserParamsID);
 	}
 	else
 	{
-		BringFocusToPanel(EPCGEditorPanel::UserParams);
+		BringFocusToPanel(PCGEditorTabs::UserParamsID);
 	}
 }
 
 bool FPCGEditor::IsToggleGraphParamsToggled() const
 {
-	return IsPanelCurrentlyOpen(EPCGEditorPanel::UserParams);
+	return IsPanelCurrentlyOpen(PCGEditorTabs::UserParamsID);
 }
 
 bool FPCGEditor::CanCollapseNodesInSubgraph() const
@@ -1995,7 +1895,10 @@ void FPCGEditor::SetNodeInspected(UPCGEditorGraphNodeBase* InspectedNode, bool b
 	TArray<UPCGEditorGraphNodeBase*, TInlineAllocator<4>> InspectedNodesBefore;
 	for (TSharedPtr<SPCGEditorGraphAttributeListView> AttributeListView : AttributesWidgets)
 	{
-		InspectedNodesBefore.Add(AttributeListView->GetNodeBeingInspected());
+		if (AttributeListView.IsValid())
+		{
+			InspectedNodesBefore.Add(AttributeListView->GetNodeBeingInspected());
+		}
 	}
 
 	bool bIsInspecting = false;
@@ -2007,7 +1910,7 @@ void FPCGEditor::SetNodeInspected(UPCGEditorGraphNodeBase* InspectedNode, bool b
 
 		for (TSharedPtr<SPCGEditorGraphAttributeListView> AttributeListView : AttributesWidgets)
 		{
-			if (AttributeListView->GetNodeBeingInspected() == InspectedNode)
+			if (AttributeListView.IsValid() && AttributeListView->GetNodeBeingInspected() == InspectedNode)
 			{
 				AttributeListView->SetNodeBeingInspected(nullptr);
 			}
@@ -2019,12 +1922,15 @@ void FPCGEditor::SetNodeInspected(UPCGEditorGraphNodeBase* InspectedNode, bool b
 
 		for (TSharedPtr<SPCGEditorGraphAttributeListView> AttributeListView : AttributesWidgets)
 		{
-			if (!AttributeListView->IsLocked())
+			if (AttributeListView.IsValid())
 			{
-				AttributeListView->SetNodeBeingInspected(InspectedNode);
-			}
+				if (!AttributeListView->IsLocked())
+				{
+					AttributeListView->SetNodeBeingInspected(InspectedNode);
+				}
 
-			InspectedNodesAfter.Add(AttributeListView->GetNodeBeingInspected());
+				InspectedNodesAfter.Add(AttributeListView->GetNodeBeingInspected());
+			}
 		}
 
 		if (InspectedNode && InspectedNodesAfter.Contains(InspectedNode))
@@ -2047,16 +1953,16 @@ void FPCGEditor::SetNodeInspected(UPCGEditorGraphNodeBase* InspectedNode, bool b
 		// Summon the first attribute list view that is inspecting this node
 		auto InvokeFirstTab = [this, InspectedNode](bool bVisibleOnly) -> bool
 		{
-			for (int AttributeListViewIndex = 0; AttributeListViewIndex < AttributesWidgets.Num(); ++AttributeListViewIndex)
+			for (int AttributeListViewIndex = 0; AttributeListViewIndex < PCGEditorTabs::NumAttributeTabs; ++AttributeListViewIndex)
 			{
 				TSharedPtr<SPCGEditorGraphAttributeListView> AttributeListView = AttributesWidgets[AttributeListViewIndex];
-				if (AttributeListView->GetNodeBeingInspected() == InspectedNode)
+				if (AttributeListView.IsValid() && AttributeListView->GetNodeBeingInspected() == InspectedNode)
 				{
-					if (!bVisibleOnly || TabManager->FindExistingLiveTab(FPCGEditor_private::AttributesID[AttributeListViewIndex]))
+					if (!bVisibleOnly || TabManager->FindExistingLiveTab(FTabId(PCGEditorTabs::AttributesID[AttributeListViewIndex])))
 					{
-						if (IsPanelAvailable(FPCGEditor_private::AttributesID[AttributeListViewIndex]))
+						if (IsPanelAvailable(PCGEditorTabs::AttributesID[AttributeListViewIndex]))
 						{
-							GetTabManager()->TryInvokeTab(FPCGEditor_private::AttributesID[AttributeListViewIndex]);
+							GetTabManager()->TryInvokeTab(FTabId(PCGEditorTabs::AttributesID[AttributeListViewIndex]));
 							return true;
 						}
 					}
@@ -2069,9 +1975,9 @@ void FPCGEditor::SetNodeInspected(UPCGEditorGraphNodeBase* InspectedNode, bool b
 		const bool bTabSummoned = (InvokeFirstTab(true) || InvokeFirstTab(false));
 
 		// Default to first if they are all locked
-		if (!bTabSummoned && IsPanelAvailable(FPCGEditor_private::AttributesID[0]))
+		if (!bTabSummoned && IsPanelAvailable(PCGEditorTabs::AttributesID[0]))
 		{
-			GetTabManager()->TryInvokeTab(FPCGEditor_private::AttributesID[0]);
+			GetTabManager()->TryInvokeTab(FTabId(PCGEditorTabs::AttributesID[0]));
 		}
 
 		DebugObjectTreeWidget->SetNodeBeingInspected(PCGNode);
@@ -2805,8 +2711,13 @@ void FPCGEditor::OnCreateComment()
 	}
 }
 
-TSharedRef<SGraphEditor> FPCGEditor::CreateGraphEditorWidget()
+void FPCGEditor::CreateGraphEditorCommands()
 {
+	if (GraphEditorCommands.IsValid())
+	{
+		return;
+	}
+
 	GraphEditorCommands = MakeShareable(new FUICommandList);
 
 	// Editing commands
@@ -2875,26 +2786,55 @@ TSharedRef<SGraphEditor> FPCGEditor::CreateGraphEditorWidget()
 	GraphEditorCommands->MapAction(FGraphEditorCommands::Get().DistributeNodesVertically,
 		FExecuteAction::CreateSP(this, &FPCGEditor::OnDistributeNodesV)
 	);
+}
+
+TSharedRef<SGraphEditor> FPCGEditor::CreateGraphEditorWidget()
+{
+	CreateGraphEditorCommands();
+
+	return CreateGraphEditorWidget(nullptr, PCGEditorGraph);
+}
+
+TSharedRef<SGraphEditor> FPCGEditor::GetOrCreateGraphEditorWidget()
+{
+	if (!GraphEditorWidget.IsValid())
+	{
+		GraphEditorWidget = CreateGraphEditorWidget();
+	}
 
+	return GraphEditorWidget.ToSharedRef();
+}
+
+TSharedRef<SGraphEditor> FPCGEditor::CreateGraphEditorWidget(TSharedPtr<FTabInfo> InTabInfo, UEdGraph* InGraph)
+{
 	SGraphEditor::FGraphEditorEvents InEvents;
 	InEvents.OnSelectionChanged = SGraphEditor::FOnSelectionChanged::CreateSP(this, &FPCGEditor::OnSelectedNodesChanged);
 	InEvents.OnVerifyTextCommit = FOnNodeVerifyTextCommit::CreateSP(this, &FPCGEditor::OnValidateNodeTitle);
 	InEvents.OnTextCommitted = FOnNodeTextCommitted::CreateSP(this, &FPCGEditor::OnNodeTitleCommitted);
 	InEvents.OnNodeDoubleClicked = FSingleNodeEvent::CreateSP(this, &FPCGEditor::OnNodeDoubleClicked);
-	InEvents.OnSpawnNodeByShortcutAtLocation = SGraphEditor::FOnSpawnNodeByShortcutAtLocation::CreateSP(this, &FPCGEditor::OnSpawnNodeByShortcut, PCGEditorGraph);
+	InEvents.OnSpawnNodeByShortcutAtLocation = SGraphEditor::FOnSpawnNodeByShortcutAtLocation::CreateSP(this, &FPCGEditor::OnSpawnNodeByShortcut, Cast<UPCGEditorGraph>(InGraph));
 	InEvents.OnCreateActionMenuAtLocation = SGraphEditor::FOnCreateActionMenuAtLocation::CreateSP(this, &FPCGEditor::OnCreateActionMenuContent);
 
 	return SNew(SGraphEditor)
 		.AdditionalCommands(GraphEditorCommands)
 		.IsEditable(true)
 		.Appearance(GetAppearanceInfo())
-		.GraphToEdit(PCGEditorGraph)
+		.GraphToEdit(InGraph)
 		.GraphEvents(InEvents)
 		.ShowGraphStateOverlay(false);
 }
 
+TSharedPtr<FWorkflowTabFactory> FPCGEditor::CreateGraphEditorTabFactory()
+{
+	FPCGGenericTabFactoryParams Params(PCGEditorTabs::GraphEditorID, SharedThis<FPCGEditor>(this), [this]() { return GetOrCreateGraphEditorWidget(); });
+	Params.Label = LOCTEXT("GraphLabel", "Graph");
+	return MakeShared<FPCGGenericTabFactory>(Params);
+}
+
 void FPCGEditor::OnClose()
 {
+	FWorkflowCentricApplication::OnClose();
+
 	if (PCGEditorGraph)
 	{
 		PCGEditorGraph->OnClose();
@@ -2956,50 +2896,243 @@ void FPCGEditor::InitToolMenuContext(FToolMenuContext& MenuContext)
 	MenuContext.AddObject(Context);
 }
 
-TSharedRef<SPCGEditorGraphNodePalette> FPCGEditor::CreatePaletteWidget()
+void FPCGEditor::SaveEditedObjectState()
 {
-	return SNew(SPCGEditorGraphNodePalette, SharedThis(this));
+	DocumentManager->SaveAllState();
 }
 
-TSharedRef<SPCGEditorGraphDebugObjectTree> FPCGEditor::CreateDebugObjectTreeWidget()
+void FPCGEditor::RestoreEditedObjectState()
 {
-	return SNew(SPCGEditorGraphDebugObjectTree, SharedThis(this));
+	// @todo_pcg: implement
 }
 
-TSharedRef<SPCGEditorGraphFind> FPCGEditor::CreateFindWidget()
+TSharedPtr<FWorkflowTabFactory> FPCGEditor::CreatePaletteTabFactory()
 {
-	return SNew(SPCGEditorGraphFind, SharedThis(this));
+	FPCGGenericTabFactoryParams Params(PCGEditorTabs::PaletteID, SharedThis<FPCGEditor>(this), [this]() { return GetOrCreatePaletteWidget(); });
+	Params.Label = LOCTEXT("PaletteLabel", "Palette");
+	return MakeShared<FPCGGenericTabFactory>(Params);
 }
 
-TSharedRef<SPCGEditorGraphAttributeListView> FPCGEditor::CreateAttributesWidget(int32 Index)
+TSharedRef<SPCGEditorGraphNodePalette> FPCGEditor::GetOrCreatePaletteWidget()
 {
-	return SNew(SPCGEditorGraphAttributeListView, SharedThis(this))
-		.WidgetEntryNumber(Index);
+	if (!PaletteWidget.IsValid())
+	{
+		PaletteWidget = SNew(SPCGEditorGraphNodePalette, SharedThis(this));
+	}
+
+	return PaletteWidget.ToSharedRef();
 }
 
-TSharedRef<SPCGEditorGraphDeterminismListView> FPCGEditor::CreateDeterminismWidget()
+TSharedPtr<FWorkflowTabFactory> FPCGEditor::CreateDebugObjectTreeTabFactory()
 {
-	return SNew(SPCGEditorGraphDeterminismListView, SharedThis(this));
+	FPCGGenericTabFactoryParams Params(PCGEditorTabs::DebugObjectID, SharedThis<FPCGEditor>(this), [this]() { return GetOrCreateDebugObjectTreeWidget(); });
+	Params.Label = LOCTEXT("DebugLabel", "Debug Object");
+	return MakeShared<FPCGGenericTabFactory>(Params);
 }
 
-TSharedRef<SPCGEditorGraphProfilingView> FPCGEditor::CreateProfilingWidget()
+TSharedRef<SPCGEditorGraphDebugObjectTree> FPCGEditor::GetOrCreateDebugObjectTreeWidget()
 {
-	return SNew(SPCGEditorGraphProfilingView, SharedThis(this));
+	if (!DebugObjectTreeWidget.IsValid())
+	{
+		DebugObjectTreeWidget = SNew(SPCGEditorGraphDebugObjectTree, SharedThis(this));
+	}
+
+	return DebugObjectTreeWidget.ToSharedRef();
 }
 
-TSharedRef<SPCGEditorGraphLogView> FPCGEditor::CreateLogWidget()
+TSharedPtr<FWorkflowTabFactory> FPCGEditor::CreateFindTabFactory()
 {
-	return SNew(SPCGEditorGraphLogView, SharedThis(this));
+	FPCGGenericTabFactoryParams Params(PCGEditorTabs::FindID, SharedThis<FPCGEditor>(this), [this]() { return GetOrCreateFindWidget(); });
+	Params.Label = LOCTEXT("FindLabel", "Find");
+	return MakeShared<FPCGGenericTabFactory>(Params);
 }
 
-TSharedRef<SPCGEditorNodeSource> FPCGEditor::CreateNodeSourceWidget()
+TSharedRef<SPCGEditorGraphFind> FPCGEditor::GetOrCreateFindWidget()
 {
-	return SNew(SPCGEditorNodeSource);
+	if (!FindWidget.IsValid())
+	{
+		FindWidget = SNew(SPCGEditorGraphFind, SharedThis(this));
+	}
+
+	return FindWidget.ToSharedRef();
 }
 
-TSharedRef<SPCGEditorGraphUserParametersView> FPCGEditor::CreateGraphParamsWidget()
+TSharedPtr<FWorkflowTabFactory> FPCGEditor::CreateAttributesTabFactory(int32 Index, TSharedRef<FWorkspaceItem> Group)
 {
-	return SNew(SPCGEditorGraphUserParametersView, SharedThis(this));
+	check(Index >= 0 && Index < PCGEditorTabs::NumAttributeTabs);
+
+	FPCGGenericTabFactoryParams Params(PCGEditorTabs::AttributesID[Index], SharedThis<FPCGEditor>(this), [this, Index]() { return GetOrCreateAttributesWidget(Index); });
+	
+	static FText AttributeText[4] =
+	{
+		LOCTEXT("AttributeLabel1", "Attributes"),
+		LOCTEXT("AttributeLabel2", "Attributes 2"),
+		LOCTEXT("AttributeLabel3", "Attributes 3"),
+		LOCTEXT("AttributeLabel4", "Attributes 4")
+	};
+	
+	Params.Label = AttributeText[Index];
+	Params.Group = Group;
+	return MakeShared<FPCGGenericTabFactory>(Params);
+}
+
+TSharedRef<SPCGEditorGraphAttributeListView> FPCGEditor::GetOrCreateAttributesWidget(int32 Index)
+{
+	check(Index >= 0 && Index < PCGEditorTabs::NumAttributeTabs);
+
+	static_assert(FPCGEditorInspectionDataManager::NumberOfEntries == PCGEditorTabs::NumAttributeTabs, "FPCGEditorInspectionDataManager::NumberOfEntries == PCGEditorTabs::NumAttributeTabs");
+
+	if (!AttributesWidgets[Index].IsValid())
+	{
+		AttributesWidgets[Index] = SNew(SPCGEditorGraphAttributeListView, SharedThis(this))
+			.WidgetEntryNumber(Index);
+		AttributesWidgets[Index]->SetViewportWidget(GetOrCreateViewportWidget(Index), PCGEditorTabs::ViewportID[Index]);
+	}
+
+	return AttributesWidgets[Index].ToSharedRef();
+}
+
+TSharedPtr<FWorkflowTabFactory> FPCGEditor::CreateDeterminismTabFactory()
+{
+	FPCGGenericTabFactoryParams Params(PCGEditorTabs::DeterminismID, SharedThis<FPCGEditor>(this), [this]() { return GetOrCreateDeterminismWidget(); });
+	Params.Label = LOCTEXT("DeterminismLabel", "Determinism");
+	return MakeShared<FPCGGenericTabFactory>(Params);
+}
+
+TSharedRef<SPCGEditorGraphDeterminismListView> FPCGEditor::GetOrCreateDeterminismWidget()
+{
+	if (!DeterminismWidget.IsValid())
+	{
+		DeterminismWidget = SNew(SPCGEditorGraphDeterminismListView, SharedThis(this));
+	}
+
+	return DeterminismWidget.ToSharedRef();
+}
+
+TSharedPtr<FWorkflowTabFactory> FPCGEditor::CreateProfilingTabFactory()
+{
+	FPCGGenericTabFactoryParams Params(PCGEditorTabs::ProfilingID, SharedThis<FPCGEditor>(this), [this]() { return GetOrCreateProfilingWidget(); });
+	Params.Label = LOCTEXT("ProfilingLabel", "Profiling");
+	return MakeShared<FPCGGenericTabFactory>(Params);
+}
+
+TSharedRef<SPCGEditorGraphProfilingView> FPCGEditor::GetOrCreateProfilingWidget()
+{
+	if (!ProfilingWidget.IsValid())
+	{
+		ProfilingWidget = SNew(SPCGEditorGraphProfilingView, SharedThis(this));
+	}
+
+	return ProfilingWidget.ToSharedRef();
+}
+
+TSharedPtr<FWorkflowTabFactory> FPCGEditor::CreateLogTabFactory()
+{
+	FPCGGenericTabFactoryParams Params(PCGEditorTabs::LogID, SharedThis<FPCGEditor>(this), [this]() { return GetOrCreateLogWidget(); });
+	Params.Label = LOCTEXT("FindLog", "Log");
+	return MakeShared<FPCGGenericTabFactory>(Params);
+}
+
+TSharedRef<SPCGEditorGraphLogView> FPCGEditor::GetOrCreateLogWidget()
+{
+	if (!LogWidget.IsValid())
+	{
+		LogWidget = SNew(SPCGEditorGraphLogView, SharedThis(this));
+	}
+
+	return LogWidget.ToSharedRef();
+}
+
+TSharedPtr<FWorkflowTabFactory> FPCGEditor::CreateNodeSourceTabFactory()
+{
+	FPCGGenericTabFactoryParams Params(PCGEditorTabs::HLSLSourceID, SharedThis<FPCGEditor>(this), [this]() { return GetOrCreateNodeSourceWidget(); });
+	Params.Label = LOCTEXT("NodeSourceLabel", "HLSL Source");
+	return MakeShared<FPCGGenericTabFactory>(Params);
+}
+
+TSharedRef<SPCGEditorNodeSource> FPCGEditor::GetOrCreateNodeSourceWidget()
+{
+	if (!NodeSourceWidget.IsValid())
+	{
+		NodeSourceWidget = SNew(SPCGEditorNodeSource);
+	}
+	
+	return NodeSourceWidget.ToSharedRef();
+}
+
+TSharedPtr<FWorkflowTabFactory> FPCGEditor::CreateGraphParamsTabFactory()
+{
+	FPCGGenericTabFactoryParams Params(PCGEditorTabs::UserParamsID, SharedThis<FPCGEditor>(this), [this]() { return GetOrCreateGraphParamsWidget(); });
+	Params.Label = LOCTEXT("UserParamsLabel", "Graph Parameters");
+	return MakeShared<FPCGGenericTabFactory>(Params);
+}
+
+TSharedRef<SPCGEditorGraphUserParametersView> FPCGEditor::GetOrCreateGraphParamsWidget()
+{
+	if (!UserParamsWidget.IsValid())
+	{
+		UserParamsWidget = SNew(SPCGEditorGraphUserParametersView, SharedThis(this));
+	}
+
+	return UserParamsWidget.ToSharedRef();
+}
+
+TSharedPtr<FWorkflowTabFactory> FPCGEditor::CreatePropertyDetailsTabFactory(int32 Index, TSharedRef<FWorkspaceItem> Group)
+{
+	check(Index >= 0 && Index < PCGEditorTabs::NumPropertyDetailTabs);
+
+	FPCGGenericTabFactoryParams Params(PCGEditorTabs::PropertyDetailsID[Index], SharedThis<FPCGEditor>(this), [this, Index]() { return GetOrCreatePropertyDetailsWidget(Index); });
+
+	static FText PropertyDetailText[4] =
+	{
+		LOCTEXT("PropertyDetailsLabel1", "Details"),
+		LOCTEXT("PropertyDetailsLabel2", "Details 2"),
+		LOCTEXT("PropertyDetailsLabel3", "Details 3"),
+		LOCTEXT("PropertyDetailsLabel4", "Details 4")
+	};
+
+	Params.Label = PropertyDetailText[Index];
+	Params.Group = Group;
+	return MakeShared<FPCGGenericTabFactory>(Params);
+}
+
+TSharedRef<SPCGEditorGraphDetailsView> FPCGEditor::GetOrCreatePropertyDetailsWidget(int32 Index)
+{
+	check(Index >= 0 && Index < PCGEditorTabs::NumPropertyDetailTabs);
+
+	if (!PropertyDetailsWidgets[Index].IsValid())
+	{
+		PropertyDetailsWidgets[Index] = SNew(SPCGEditorGraphDetailsView);
+		PropertyDetailsWidgets[Index]->SetEditor(SharedThis(this));
+		PropertyDetailsWidgets[Index]->SetObject(PCGGraphBeingEdited);
+	}
+
+	return PropertyDetailsWidgets[Index].ToSharedRef();
+}
+
+TSharedPtr<FWorkflowTabFactory> FPCGEditor::CreateViewportTabFactory(int32 Index, TSharedRef<FWorkspaceItem> Group)
+{
+	check(Index >= 0 && Index < PCGEditorTabs::NumViewportTabs);
+
+	FPCGGenericTabFactoryParams Params(PCGEditorTabs::ViewportID[Index], SharedThis<FPCGEditor>(this), [this, Index]() { return GetOrCreateViewportWidget(Index); });
+
+	static FText ViewportText[4] =
+	{
+		LOCTEXT("ViewportLabel1", "Viewport"),
+		LOCTEXT("ViewportLabel2", "Viewport 2"),
+		LOCTEXT("ViewportLabel3", "Viewport 3"),
+		LOCTEXT("ViewportLabel4", "Viewport 4")
+	};
+
+	Params.Label = ViewportText[Index];
+	Params.Group = Group;
+	return MakeShared<FPCGGenericTabFactory>(Params);
+}
+
+TSharedPtr<SPCGEditorViewport> FPCGEditor::GetViewportWidget(int32 Index) const
+{
+	check(Index >= 0 && Index < PCGEditorTabs::NumViewportTabs);
+	return ViewportWidgets[Index];
 }
 
 TSharedRef<SPCGEditorViewport> FPCGEditor::CreateViewportWidget()
@@ -3007,6 +3140,18 @@ TSharedRef<SPCGEditorViewport> FPCGEditor::CreateViewportWidget()
 	return SNew(SPCGEditorViewport);
 }
 
+TSharedRef<SPCGEditorViewport> FPCGEditor::GetOrCreateViewportWidget(int32 Index)
+{
+	check(Index >= 0 && Index < PCGEditorTabs::NumViewportTabs);
+
+	if (!ViewportWidgets[Index].IsValid())
+	{
+		ViewportWidgets[Index] = CreateViewportWidget();
+	}
+
+	return ViewportWidgets[Index].ToSharedRef();
+}
+
 void FPCGEditor::OnSelectedNodesChanged(const TSet<UObject*>& NewSelection)
 {
 	TArray<TWeakObjectPtr<UObject>> SelectedObjects;
@@ -3028,7 +3173,10 @@ void FPCGEditor::OnSelectedNodesChanged(const TSet<UObject*>& NewSelection)
 
 	for (TSharedPtr<SPCGEditorGraphDetailsView> PropertyDetailsWidget : PropertyDetailsWidgets)
 	{
-		PropertyDetailsWidget->SetObjects(SelectedObjects, /*bForceRefresh=*/true);
+		if (PropertyDetailsWidget.IsValid())
+		{
+			PropertyDetailsWidget->SetObjects(SelectedObjects, /*bForceRefresh=*/true);
+		}
 	}
 
 	// Give a single selected node with valid settings to the source editor, or give it null so it can clear the UI.
@@ -3261,7 +3409,7 @@ void FPCGEditor::OnGraphChanged(UPCGGraphInterface* InGraph, EPCGChangeType Chan
 		if (PaletteWidget)
 		{
 			PaletteWidget->RequestRefresh();
-		}		
+		}
 	}
 
 	if (!!(ChangeType & EPCGChangeType::Edge))
@@ -3361,239 +3509,7 @@ void FPCGEditor::UpdateDefaultExecutionSource()
 
 TSharedRef<FTabManager::FLayout> FPCGEditor::GetDefaultLayout() const
 {
-	return FTabManager::NewLayout("Standalone_PCGGraphEditor_DefaultLayout_v1.0")
-		->AddArea // Main PCG Graph Editor Area
-		(
-			FTabManager::NewPrimaryArea()->SetOrientation(Orient_Vertical)
-			->Split // Top Section - Graph, Data Viewport, HLSL Source Editor, and Details View
-			(
-				FTabManager::NewSplitter()->SetOrientation(Orient_Horizontal)
-				->SetSizeCoefficient(0.65f)
-				->Split // Graph Palette
-				(
-					FTabManager::NewStack()
-					->AddTab(FPCGEditor_private::PaletteID, ETabState::SidebarTab, ESidebarLocation::Left, /*SidebarSizeCoefficient=*/0.13f)
-				)
-				->Split // Data Viewport/HLSL Source Editor
-				(
-					FTabManager::NewStack()
-					->SetSizeCoefficient(0.2f)
-					->AddTab(FPCGEditor_private::ViewportID[0], ETabState::OpenedTab)
-					->AddTab(FPCGEditor_private::HLSLSourceID, ETabState::OpenedTab)
-					->SetForegroundTab(FPCGEditor_private::ViewportID[0])
-				)
-				->Split // Node Graph
-				(
-					FTabManager::NewStack()
-					->SetSizeCoefficient(0.6f)
-					->AddTab(FPCGEditor_private::GraphEditorID, ETabState::OpenedTab)
-					->SetHideTabWell(true)
-				)
-				->Split // Details View
-				(
-					FTabManager::NewStack()
-					->SetSizeCoefficient(0.2f)
-					->AddTab(FPCGEditor_private::PropertyDetailsID[0], ETabState::OpenedTab)
-				)
-			)
-			->Split // Bottom Section - Debug/Params and ALV
-			(
-				FTabManager::NewSplitter()->SetOrientation(Orient_Horizontal)
-				->SetSizeCoefficient(0.35f)
-				->Split
-				(
-					FTabManager::NewStack()
-					->SetSizeCoefficient(0.2f)
-					->AddTab(FPCGEditor_private::DebugObjectID, ETabState::OpenedTab)
-					->AddTab(FPCGEditor_private::UserParamsID, ETabState::OpenedTab)
-					->SetForegroundTab(FPCGEditor_private::DebugObjectID)
-				)
-				->Split // ALV, Profiling, Find, Determinism
-				(
-					FTabManager::NewStack()
-					->SetSizeCoefficient(0.8f)
-					->SetHideTabWell(false)
-					->AddTab(FPCGEditor_private::AttributesID[0], ETabState::OpenedTab)
-					->AddTab(FPCGEditor_private::ProfilingID, ETabState::OpenedTab)
-					->AddTab(FPCGEditor_private::FindID, ETabState::OpenedTab)
-					->AddTab(FPCGEditor_private::DeterminismID, ETabState::ClosedTab)
-					->SetForegroundTab(FPCGEditor_private::AttributesID[0])
-				)
-			)
-		);
-}
-
-TSharedRef<SDockTab> FPCGEditor::SpawnTab_GraphEditor(const FSpawnTabArgs& Args)
-{
-	return SNew(SDockTab)
-		.Label(LOCTEXT("PCGGraphTitle", "Graph"))
-		.TabColorScale(GetTabColorScale())
-		[
-			GraphEditorWidget.ToSharedRef()
-		];
-}
-
-TSharedRef<SDockTab> FPCGEditor::SpawnTab_PropertyDetails(const FSpawnTabArgs& Args, int PropertyDetailsIndex)
-{
-	TAttribute<FText> Label = TAttribute<FText>::Create(TAttribute<FText>::FGetter::CreateRaw(this, &FPCGEditor::GetDetailsTabLabel, PropertyDetailsIndex));
-	TSharedPtr<SPCGEditorGraphDetailsView> DetailsView = PropertyDetailsWidgets[PropertyDetailsIndex];
-
-	return SNew(SDockTab)
-		.Label(Label)
-		.OnTabClosed_Raw(this, &FPCGEditor::OnDetailsViewTabClosed, PropertyDetailsIndex)
-		.TabColorScale(GetTabColorScale())
-		[
-			DetailsView.ToSharedRef()
-		];
-}
-
-TSharedRef<SDockTab> FPCGEditor::SpawnTab_Palette(const FSpawnTabArgs& Args)
-{
-	return SNew(SDockTab)
-		.Label(LOCTEXT("PCGPaletteTitle", "Palette"))
-		.TabColorScale(GetTabColorScale())
-		[
-			PaletteWidget.ToSharedRef()
-		];
-}
-
-TSharedRef<SDockTab> FPCGEditor::SpawnTab_DebugObjectTree(const FSpawnTabArgs& Args)
-{
-	return SNew(SDockTab)
-		.Label(LOCTEXT("PCGDebugObjectTitle", "Debug Object"))
-		.TabColorScale(GetTabColorScale())
-		[
-			DebugObjectTreeWidget.ToSharedRef()
-		];
-}
-
-TSharedRef<SDockTab> FPCGEditor::SpawnTab_Attributes(const FSpawnTabArgs& Args, int AttributesIndex)
-{
-	TAttribute<FText> Label = TAttribute<FText>::Create(TAttribute<FText>::FGetter::CreateRaw(this, &FPCGEditor::GetAttributesTabLabel, AttributesIndex));
-
-	return SNew(SDockTab)
-		.Label(Label)
-		.OnTabClosed_Raw(this, &FPCGEditor::OnAttributeListViewTabClosed, AttributesIndex)
-		.TabColorScale(GetTabColorScale())
-		[
-			AttributesWidgets[AttributesIndex].ToSharedRef()
-		];
-}
-
-TSharedRef<SDockTab> FPCGEditor::SpawnTab_Find(const FSpawnTabArgs& Args)
-{
-	return SNew(SDockTab)
-		.Label(LOCTEXT("PCGFindTitle", "Find"))
-		.TabColorScale(GetTabColorScale())
-		[
-			FindWidget.ToSharedRef()
-		];
-}
-
-TSharedRef<SDockTab> FPCGEditor::SpawnTab_Determinism(const FSpawnTabArgs& Args)
-{
-	return SNew(SDockTab)
-		.Label(LOCTEXT("PCGDeterminismTitle", "Determinism"))
-		.TabColorScale(GetTabColorScale())
-		[
-			DeterminismWidget.ToSharedRef()
-		];
-}
-
-TSharedRef<SDockTab> FPCGEditor::SpawnTab_Profiling(const FSpawnTabArgs& Args)
-{
-	return SNew(SDockTab)
-		.Label(LOCTEXT("PCGProfilingTitle", "Profiling"))
-		.TabColorScale(GetTabColorScale())
-		[
-			ProfilingWidget.ToSharedRef()
-		];
-}
-
-TSharedRef<SDockTab> FPCGEditor::SpawnTab_Log(const FSpawnTabArgs& Args)
-{
-	return SNew(SDockTab)
-		.Label(LOCTEXT("PCGLogTitle", "Log Capture"))
-		.TabColorScale(GetTabColorScale())
-		[
-			LogWidget.ToSharedRef()
-		];
-}
-
-TSharedRef<SDockTab> FPCGEditor::SpawnTab_NodeSource(const FSpawnTabArgs& Args)
-{
-	return SNew(SDockTab)
-		.Label(LOCTEXT("PCGHLSLSourceTitle", "HLSL Source"))
-		.TabColorScale(GetTabColorScale())
-		[
-			NodeSourceWidget.ToSharedRef()
-		];
-}
-
-TSharedRef<SDockTab> FPCGEditor::SpawnTab_UserParams(const FSpawnTabArgs& Args)
-{
-	return SNew(SDockTab)
-		.Label(FPCGEditor_private::UserParamsTabName)
-		.TabColorScale(GetTabColorScale())
-		[
-			UserParamsWidget.ToSharedRef()
-		];
-}
-
-TSharedRef<SDockTab> FPCGEditor::SpawnTab_Viewport(const FSpawnTabArgs& Args, int ViewportIndex)
-{
-	TAttribute<FText> Label = TAttribute<FText>::Create(TAttribute<FText>::FGetter::CreateRaw(this, &FPCGEditor::GetViewportTabLabel, ViewportIndex));
-	
-	AttributesWidgets[ViewportIndex]->RequestViewportRefresh();
-
-	return SNew(SDockTab)
-		.Label(Label)
-		.OnTabClosed_Raw(this, &FPCGEditor::OnViewportViewTabClosed, ViewportIndex)
-		.TabColorScale(GetTabColorScale())
-		[
-			AttributesWidgets[ViewportIndex]->GetViewportWidget().ToSharedRef()
-		];
-}
-
-FText FPCGEditor::GetDetailsTabLabel(int DetailsIndex)
-{
-	if (DetailsIndex == 0)
-	{
-		return LOCTEXT("PCGDetailsTitle", "Details");
-	}
-	else
-	{
-		return FText::Format(LOCTEXT("PCGDetailsTitle_Multi", "Details {0}"), DetailsIndex + 1);
-	}
-}
-
-FText FPCGEditor::GetDetailsViewObjectName(int DetailsIndex)
-{
-	return LOCTEXT("PCGDetailsName", "This is a node name placeholder");
-}
-
-FText FPCGEditor::GetAttributesTabLabel(int AttributesIndex)
-{
-	if (AttributesIndex == 0)
-	{
-		return LOCTEXT("PCGAttributesTitle", "Attributes");
-	}
-	else
-	{
-		return FText::Format(LOCTEXT("PCGAttributesTitle_Multi", "Attributes {0}"), AttributesIndex + 1);
-	}
-}
-
-FText FPCGEditor::GetViewportTabLabel(int ViewportIndex)
-{
-	if (ViewportIndex == 0)
-	{
-		return LOCTEXT("PCGViewportTitle", "Data Viewport");
-	}
-	else
-	{
-		return FText::Format(LOCTEXT("PCGViewportTitle_Multi", "Data Viewport {0}"), ViewportIndex + 1);
-	}
+	return PCGEditorDefaultMode::GetTabLayout().ToSharedRef();
 }
 
 #undef LOCTEXT_NAMESPACE
\ No newline at end of file
diff --git a/Engine/Plugins/PCG/Source/PCGEditor/Private/PCGEditorDefaultMode.cpp b/Engine/Plugins/PCG/Source/PCGEditor/Private/PCGEditorDefaultMode.cpp
new file mode 100644
--- /dev/null
+++ b/Engine/Plugins/PCG/Source/PCGEditor/Private/PCGEditorDefaultMode.cpp
@@ -0,0 +1,174 @@
+// Copyright Epic Games, Inc. All Rights Reserved.
+
+#include "PCGEditorDefaultMode.h"
+
+#include "PCGEditor.h"
+#include "PCGEditorTabFactories.h"
+
+#include "Delegates/DelegateCombinations.h"
+
+#define LOCTEXT_NAMESPACE "PCGGraphEditor"
+
+namespace PCGEditorDefaultMode
+{
+	DECLARE_DELEGATE_RetVal(TSharedPtr<FWorkflowTabFactory>, FCreateTabFactoryDelegate);
+
+	FText GetLocalizedMode(FName ModeName)
+	{
+		return LOCTEXT("DefaultModeName", "PCG");
+	}
+
+	TSharedPtr<FTabManager::FLayout> GetTabLayout()
+	{
+		return FTabManager::NewLayout("Standalone_PCGGraphEditor_DefaultLayout_v1.0")
+			->AddArea // Main PCG Graph Editor Area
+			(
+				FTabManager::NewPrimaryArea()->SetOrientation(Orient_Vertical)
+				->Split // Top Section - Graph, Data Viewport, HLSL Source Editor, and Details View
+				(
+					FTabManager::NewSplitter()->SetOrientation(Orient_Horizontal)
+					->SetSizeCoefficient(0.65f)
+					->Split // Graph Palette
+					(
+						FTabManager::NewStack()
+						->AddTab(PCGEditorTabs::PaletteID, ETabState::SidebarTab, ESidebarLocation::Left, /*SidebarSizeCoefficient=*/0.13f)
+					)
+					->Split // Data Viewport/HLSL Source Editor
+					(
+						FTabManager::NewStack()
+						->SetSizeCoefficient(0.2f)
+						->AddTab(PCGEditorTabs::ViewportID[0], ETabState::OpenedTab)
+						->AddTab(PCGEditorTabs::HLSLSourceID, ETabState::OpenedTab)
+						->SetForegroundTab(FName(PCGEditorTabs::ViewportID[0]))
+					)
+					->Split // Node Graph
+					(
+						FTabManager::NewStack()
+						->SetSizeCoefficient(0.6f)
+						->AddTab(PCGEditorTabs::GraphEditorID, ETabState::OpenedTab)
+						->AddTab(PCGEditorTabs::PCGEditorDocumentID, ETabState::ClosedTab)
+						->SetHideTabWell(true)
+					)
+					->Split // Details View
+					(
+						FTabManager::NewStack()
+						->SetSizeCoefficient(0.2f)
+						->AddTab(PCGEditorTabs::PropertyDetailsID[0], ETabState::OpenedTab)
+					)
+				)
+				->Split // Bottom Section - Debug/Params and ALV
+				(
+					FTabManager::NewSplitter()->SetOrientation(Orient_Horizontal)
+					->SetSizeCoefficient(0.35f)
+					->Split
+					(
+						FTabManager::NewStack()
+						->SetSizeCoefficient(0.2f)
+						->AddTab(PCGEditorTabs::DebugObjectID, ETabState::OpenedTab)
+						->AddTab(PCGEditorTabs::UserParamsID, ETabState::OpenedTab)
+						->SetForegroundTab(FTabId(PCGEditorTabs::DebugObjectID))
+					)
+					->Split // ALV, Profiling, Find, Determinism
+					(
+						FTabManager::NewStack()
+						->SetSizeCoefficient(0.8f)
+						->SetHideTabWell(false)
+						->AddTab(PCGEditorTabs::AttributesID[0], ETabState::OpenedTab)
+						->AddTab(PCGEditorTabs::ProfilingID, ETabState::OpenedTab)
+						->AddTab(PCGEditorTabs::FindID, ETabState::OpenedTab)
+						->AddTab(PCGEditorTabs::DeterminismID, ETabState::ClosedTab)
+						->SetForegroundTab(FTabId(PCGEditorTabs::AttributesID[0]))
+					)
+				)
+			);
+	}
+}
+
+FPCGEditorDefaultMode::FPCGEditorDefaultMode(TSharedPtr<FPCGEditor> InPCGEditor)
+	: FApplicationMode(StaticName(), PCGEditorDefaultMode::GetLocalizedMode)
+{
+	TabLayout = InPCGEditor->GetDefaultLayout();
+
+	using CreateFactoryFunc = decltype(PCGEditorDefaultMode::FCreateTabFactoryDelegate::CreateSP(InPCGEditor.Get(), &FPCGEditor::CreatePaletteTabFactory));
+		
+	struct TabInfo
+	{
+		FName ID;
+		CreateFactoryFunc CreateFactory;
+		TSharedPtr<FWorkspaceItem> Group;
+	};
+
+	TSharedRef<FWorkspaceItem> PropertyDetailsGroup = GetWorkspaceMenuCategory()->AddGroup(LOCTEXT("WorkspaceMenu_PCGEditor_Details", "Details"));
+	TSharedRef<FWorkspaceItem> AttributesGroup = GetWorkspaceMenuCategory()->AddGroup(LOCTEXT("WorkspaceMenu_PCGEditor_Attributes", "Attributes"));
+	TSharedRef<FWorkspaceItem> ViewportGroup = GetWorkspaceMenuCategory()->AddGroup(LOCTEXT("WorkspaceMenu_PCGEditor_Viewport", "Data Viewport"));
+
+	TabInfo DefaultTabInfo[] =
+	{
+		{PCGEditorTabs::GraphEditorID, PCGEditorDefaultMode::FCreateTabFactoryDelegate::CreateSP(InPCGEditor.Get(), &FPCGEditor::CreateGraphEditorTabFactory)},
+		{PCGEditorTabs::PropertyDetailsID[0], PCGEditorDefaultMode::FCreateTabFactoryDelegate::CreateSP(InPCGEditor.Get(), &FPCGEditor::CreatePropertyDetailsTabFactory, 0, PropertyDetailsGroup)},
+		{PCGEditorTabs::PropertyDetailsID[1], PCGEditorDefaultMode::FCreateTabFactoryDelegate::CreateSP(InPCGEditor.Get(), &FPCGEditor::CreatePropertyDetailsTabFactory, 1, PropertyDetailsGroup)},
+		{PCGEditorTabs::PropertyDetailsID[2], PCGEditorDefaultMode::FCreateTabFactoryDelegate::CreateSP(InPCGEditor.Get(), &FPCGEditor::CreatePropertyDetailsTabFactory, 2, PropertyDetailsGroup)},
+		{PCGEditorTabs::PropertyDetailsID[3], PCGEditorDefaultMode::FCreateTabFactoryDelegate::CreateSP(InPCGEditor.Get(), &FPCGEditor::CreatePropertyDetailsTabFactory, 3, PropertyDetailsGroup)},
+		{PCGEditorTabs::PaletteID, PCGEditorDefaultMode::FCreateTabFactoryDelegate::CreateSP(InPCGEditor.Get(), &FPCGEditor::CreatePaletteTabFactory)},
+		{PCGEditorTabs::DebugObjectID, PCGEditorDefaultMode::FCreateTabFactoryDelegate::CreateSP(InPCGEditor.Get(), &FPCGEditor::CreateDebugObjectTreeTabFactory)},
+		{PCGEditorTabs::AttributesID[0], PCGEditorDefaultMode::FCreateTabFactoryDelegate::CreateSP(InPCGEditor.Get(), &FPCGEditor::CreateAttributesTabFactory, 0, AttributesGroup)},
+		{PCGEditorTabs::AttributesID[1], PCGEditorDefaultMode::FCreateTabFactoryDelegate::CreateSP(InPCGEditor.Get(), &FPCGEditor::CreateAttributesTabFactory, 1, AttributesGroup)},
+		{PCGEditorTabs::AttributesID[2], PCGEditorDefaultMode::FCreateTabFactoryDelegate::CreateSP(InPCGEditor.Get(), &FPCGEditor::CreateAttributesTabFactory, 2, AttributesGroup)},
+		{PCGEditorTabs::AttributesID[3], PCGEditorDefaultMode::FCreateTabFactoryDelegate::CreateSP(InPCGEditor.Get(), &FPCGEditor::CreateAttributesTabFactory, 3, AttributesGroup)},
+		{PCGEditorTabs::FindID, PCGEditorDefaultMode::FCreateTabFactoryDelegate::CreateSP(InPCGEditor.Get(), &FPCGEditor::CreateFindTabFactory)},
+		{PCGEditorTabs::DeterminismID, PCGEditorDefaultMode::FCreateTabFactoryDelegate::CreateSP(InPCGEditor.Get(), &FPCGEditor::CreateDeterminismTabFactory)},
+		{PCGEditorTabs::ProfilingID, PCGEditorDefaultMode::FCreateTabFactoryDelegate::CreateSP(InPCGEditor.Get(), &FPCGEditor::CreateProfilingTabFactory)},
+		{PCGEditorTabs::LogID, PCGEditorDefaultMode::FCreateTabFactoryDelegate::CreateSP(InPCGEditor.Get(), &FPCGEditor::CreateLogTabFactory)},
+		{PCGEditorTabs::HLSLSourceID, PCGEditorDefaultMode::FCreateTabFactoryDelegate::CreateSP(InPCGEditor.Get(), &FPCGEditor::CreateNodeSourceTabFactory)},
+		{PCGEditorTabs::UserParamsID, PCGEditorDefaultMode::FCreateTabFactoryDelegate::CreateSP(InPCGEditor.Get(), &FPCGEditor::CreateGraphParamsTabFactory)},
+		{PCGEditorTabs::ViewportID[0], PCGEditorDefaultMode::FCreateTabFactoryDelegate::CreateSP(InPCGEditor.Get(), &FPCGEditor::CreateViewportTabFactory, 0, ViewportGroup)},
+		{PCGEditorTabs::ViewportID[1], PCGEditorDefaultMode::FCreateTabFactoryDelegate::CreateSP(InPCGEditor.Get(), &FPCGEditor::CreateViewportTabFactory, 1, ViewportGroup)},
+		{PCGEditorTabs::ViewportID[2], PCGEditorDefaultMode::FCreateTabFactoryDelegate::CreateSP(InPCGEditor.Get(), &FPCGEditor::CreateViewportTabFactory, 2, ViewportGroup)},
+		{PCGEditorTabs::ViewportID[3], PCGEditorDefaultMode::FCreateTabFactoryDelegate::CreateSP(InPCGEditor.Get(), &FPCGEditor::CreateViewportTabFactory, 3, ViewportGroup)},
+	};
+
+	for (TabInfo& Info : DefaultTabInfo)
+	{
+		if (!InPCGEditor->IsPanelAvailable(Info.ID))
+		{
+			continue;
+		}
+
+		TabFactories.RegisterFactory(Info.CreateFactory.Execute());
+	}
+}
+
+FName FPCGEditorDefaultMode::StaticName()
+{
+	return FName("PCGEditorDefaultMode");
+}
+
+void FPCGEditorDefaultMode::RegisterTabFactories(TSharedPtr<FTabManager> InTabManager)
+{
+	if (TSharedPtr<FPCGEditor> PCGEditor = StaticCastSharedPtr<FPCGEditor>(GetHost()))
+	{
+		PCGEditor->PushTabFactories(TabFactories);
+	}
+}
+
+void FPCGEditorDefaultMode::PreDeactivateMode()
+{
+	if (TSharedPtr<FPCGEditor> PCGEditor = StaticCastSharedPtr<FPCGEditor>(GetHost()))
+	{
+		PCGEditor->SaveEditedObjectState();
+	}
+
+	FApplicationMode::PreDeactivateMode();
+}
+
+void FPCGEditorDefaultMode::PostActivateMode()
+{
+	FApplicationMode::PostActivateMode();
+
+	if (TSharedPtr<FPCGEditor> PCGEditor = StaticCastSharedPtr<FPCGEditor>(GetHost()))
+	{
+		PCGEditor->RestoreEditedObjectState();
+	}
+}
+
+#undef LOCTEXT_NAMESPACE
\ No newline at end of file
diff --git a/Engine/Plugins/PCG/Source/PCGEditor/Private/PCGEditorDefaultMode.h b/Engine/Plugins/PCG/Source/PCGEditor/Private/PCGEditorDefaultMode.h
new file mode 100644
--- /dev/null
+++ b/Engine/Plugins/PCG/Source/PCGEditor/Private/PCGEditorDefaultMode.h
@@ -0,0 +1,32 @@
+// Copyright Epic Games, Inc. All Rights Reserved.
+
+#pragma once
+
+#include "CoreMinimal.h"
+#include "Framework/Docking/TabManager.h"
+#include "WorkflowOrientedApp/ApplicationMode.h"
+#include "WorkflowOrientedApp/WorkflowCentricApplication.h"
+#include "WorkflowOrientedApp/WorkflowTabManager.h"
+
+class FPCGEditor;
+
+namespace PCGEditorDefaultMode
+{
+	TSharedPtr<FTabManager::FLayout> GetTabLayout();
+}
+
+class FPCGEditorDefaultMode : public FApplicationMode
+{
+public:
+	static FName StaticName();
+
+	FPCGEditorDefaultMode(TSharedPtr<FPCGEditor> InPCGEditor);
+
+	/** FApplicationMode interface */
+	virtual void RegisterTabFactories(TSharedPtr<FTabManager> InTabManager) override;
+	virtual void PreDeactivateMode() override;
+	virtual void PostActivateMode() override;
+
+private:
+	FWorkflowAllowedTabSet TabFactories;
+};
\ No newline at end of file
diff --git a/Engine/Plugins/PCG/Source/PCGEditor/Private/PCGEditorTabFactories.cpp b/Engine/Plugins/PCG/Source/PCGEditor/Private/PCGEditorTabFactories.cpp
new file mode 100644
--- /dev/null
+++ b/Engine/Plugins/PCG/Source/PCGEditor/Private/PCGEditorTabFactories.cpp
@@ -0,0 +1,42 @@
+// Copyright Epic Games, Inc. All Rights Reserved.
+
+#include "PCGEditorTabFactories.h"
+
+#include "PCGEditor.h"
+
+#define LOCTEXT_NAMESPACE "PCGGraphEditor"
+
+FPCGGenericTabFactory::FPCGGenericTabFactory(const FPCGGenericTabFactoryParams& InParams)
+	: FWorkflowTabFactory(InParams.Id, InParams.Editor)
+{	
+	EnableTabPadding();
+
+	CreateTabBodyFunc = InParams.CreateTabBodyFunc;
+	check(CreateTabBodyFunc.IsSet());
+
+	bIsSingleton = InParams.bSingleton;
+
+	if (InParams.Label.IsSet())
+	{
+		TabLabel = InParams.Label.GetValue();
+	}
+
+	if (InParams.Icon.IsSet())
+	{
+		TabIcon = InParams.Icon.GetValue();
+	}
+
+	if (InParams.Role.IsSet())
+	{
+		TabRole = InParams.Role.GetValue();
+	}
+
+	Group = InParams.Group;
+}
+
+TSharedRef<SWidget> FPCGGenericTabFactory::CreateTabBody(const FWorkflowTabSpawnInfo& Info) const
+{
+	return CreateTabBodyFunc();
+}
+
+#undef LOCTEXT_NAMESPACE
\ No newline at end of file
diff --git a/Engine/Plugins/PCG/Source/PCGEditor/Private/Schema/PCGEditorGraphSchemaActions.cpp b/Engine/Plugins/PCG/Source/PCGEditor/Private/Schema/PCGEditorGraphSchemaActions.cpp
--- a/Engine/Plugins/PCG/Source/PCGEditor/Private/Schema/PCGEditorGraphSchemaActions.cpp
+++ b/Engine/Plugins/PCG/Source/PCGEditor/Private/Schema/PCGEditorGraphSchemaActions.cpp
@@ -578,7 +578,7 @@ void FPCGEditorGraphSchemaAction_NewSubgraphElement::MakeGraphNodesOrContextualM
 	else if (!Graph.IsEmpty())
 	{
 		FMenuBuilder MenuBuilder(true, nullptr);
-		const FText GraphTextName = ((Graph.Num() == 1) ? FText::FromName(Graph[0]->GetFName()) : LOCTEXT("MultipleGraphs", "Multiple Graphs"));
+		const FText GraphTextName = ((Graph.Num() == 1) ? Graph[0]->GetDisplayName() : LOCTEXT("MultipleGraphs", "Multiple Graphs"));
 
 		MenuBuilder.BeginSection("GraphDroppedOn", GraphTextName);
 
diff --git a/Engine/Plugins/PCG/Source/PCGEditor/Private/Widgets/SPCGEditorGraphAttributeListView.cpp b/Engine/Plugins/PCG/Source/PCGEditor/Private/Widgets/SPCGEditorGraphAttributeListView.cpp
--- a/Engine/Plugins/PCG/Source/PCGEditor/Private/Widgets/SPCGEditorGraphAttributeListView.cpp
+++ b/Engine/Plugins/PCG/Source/PCGEditor/Private/Widgets/SPCGEditorGraphAttributeListView.cpp
@@ -825,7 +825,7 @@ void SPCGEditorGraphAttributeListView::SetNodeBeingInspected(UPCGEditorGraphNode
 	RequestRefresh();
 }
 
-void SPCGEditorGraphAttributeListView::SetViewportWidget(TSharedPtr<SPCGEditorViewport> InViewportWidget, EPCGEditorPanel InViewportEditorPanel)
+void SPCGEditorGraphAttributeListView::SetViewportWidget(TSharedPtr<SPCGEditorViewport> InViewportWidget, FName InViewportEditorPanel)
 {
 	ViewportWidget = InViewportWidget;
 	ViewportEditorPanel = InViewportEditorPanel;
@@ -2098,8 +2098,7 @@ bool SPCGEditorGraphAttributeListView::IsViewportOpen() const
 
 bool SPCGEditorGraphAttributeListView::IsOpen() const
 {
-	const int AttributeEditorPanel = static_cast<int>(EPCGEditorPanel::Attributes1) + WidgetEntryNumber;
-	return PCGEditorPtr.IsValid() && PCGEditorPtr.Pin()->IsPanelCurrentlyOpen(static_cast<EPCGEditorPanel>(AttributeEditorPanel));
+	return PCGEditorPtr.IsValid() && (WidgetEntryNumber >= 0) && (WidgetEntryNumber < PCGEditorTabs::NumAttributeTabs) && PCGEditorPtr.Pin()->IsPanelCurrentlyOpen(PCGEditorTabs::AttributesID[WidgetEntryNumber]);
 }
 
 
diff --git a/Engine/Plugins/PCG/Source/PCGEditor/Private/Widgets/SPCGEditorGraphAttributeListView.h b/Engine/Plugins/PCG/Source/PCGEditor/Private/Widgets/SPCGEditorGraphAttributeListView.h
--- a/Engine/Plugins/PCG/Source/PCGEditor/Private/Widgets/SPCGEditorGraphAttributeListView.h
+++ b/Engine/Plugins/PCG/Source/PCGEditor/Private/Widgets/SPCGEditorGraphAttributeListView.h
@@ -16,7 +16,6 @@
 #include "Widgets/Views/STableRow.h"
 
 enum class EPCGMetadataTypes : uint8;
-enum class EPCGEditorPanel;
 struct FPCGDataCollection;
 struct FPCGInspectionData;
 class FPCGEditor;
@@ -164,7 +163,7 @@ public:
 	void SetIsLocked(bool bInIsLocked) { bIsLocked = bInIsLocked; }
 
 	TSharedPtr<SPCGEditorViewport> GetViewportWidget() const { return ViewportWidget; }
-	void SetViewportWidget(TSharedPtr<SPCGEditorViewport> InViewportWidget, EPCGEditorPanel InViewportEditorPanel);
+	void SetViewportWidget(TSharedPtr<SPCGEditorViewport> InViewportWidget, FName InViewportEditorPanel);
 
 	/** Resets/clears the viewport scene. */
 	void ResetViewport();
@@ -283,8 +282,8 @@ private:
 	/** Data visualization viewport. */
 	TSharedPtr<SPCGEditorViewport> ViewportWidget;
 
-	/** Data visualization viewport editor panel. */
-	EPCGEditorPanel ViewportEditorPanel;
+	/** Data visualization viewport editor panel id. */
+	FName ViewportEditorPanel;
 
 	/** Cached PCGGraphNode being viewed */
 	TWeakObjectPtr<UPCGEditorGraphNodeBase> PCGEditorGraphNode;
diff --git a/Engine/Plugins/PCG/Source/PCGEditor/Private/Widgets/SPCGEditorGraphDebugObjectTree.cpp b/Engine/Plugins/PCG/Source/PCGEditor/Private/Widgets/SPCGEditorGraphDebugObjectTree.cpp
--- a/Engine/Plugins/PCG/Source/PCGEditor/Private/Widgets/SPCGEditorGraphDebugObjectTree.cpp
+++ b/Engine/Plugins/PCG/Source/PCGEditor/Private/Widgets/SPCGEditorGraphDebugObjectTree.cpp
@@ -810,7 +810,7 @@ bool SPCGEditorGraphDebugObjectTree::GetFirstStackFromSelection(const UPCGNode*
 
 bool SPCGEditorGraphDebugObjectTree::IsOpen() const
 {
-	return PCGEditor.IsValid() && PCGEditor.Pin()->IsPanelCurrentlyForeground(EPCGEditorPanel::DebugObjectTree);
+	return PCGEditor.IsValid() && PCGEditor.Pin()->IsPanelCurrentlyForeground(PCGEditorTabs::DebugObjectID);
 }
 
 FPCGEditorGraphDebugObjectItemPtr SPCGEditorGraphDebugObjectTree::GetFirstDebugObjectFromSelection() const
diff --git a/Engine/Plugins/PCG/Source/PCGEditor/Public/PCGEditor.h b/Engine/Plugins/PCG/Source/PCGEditor/Public/PCGEditor.h
--- a/Engine/Plugins/PCG/Source/PCGEditor/Public/PCGEditor.h
+++ b/Engine/Plugins/PCG/Source/PCGEditor/Public/PCGEditor.h
@@ -9,7 +9,9 @@
 #include "EditorUndoClient.h"
 #include "GraphEditor.h"
 #include "Schema/PCGEditorGraphSchema.h"
-#include "Toolkits/AssetEditorToolkit.h"
+#include "WorkflowOrientedApp/WorkflowCentricApplication.h"
+#include "WorkflowOrientedApp/WorkflowTabFactory.h"
+#include "WorkflowOrientedApp/WorkflowTabManager.h"
 
 class FPCGEditorInspectionDataManager;
 enum class ECheckBoxState : uint8;
@@ -39,10 +41,11 @@ class UPCGEditorGraphNodeBase;
 class UPCGGraph;
 class UPCGSubsystem;
 struct FPCGCompilerDiagnostics;
+class FWorkspaceItem;
 
 DECLARE_MULTICAST_DELEGATE_OneParam(FOnInspectedNodeChanged, UPCGEditorGraphNodeBase*);
 
-enum class EPCGEditorPanel
+enum class UE_DEPRECATED(5.8, "Use PCGEditorTabs instead") EPCGEditorPanel
 {
 	Attributes1,
 	Attributes2,
@@ -78,7 +81,40 @@ enum class EPCGToolbarButtons
 	GraphSettings
 };
 
-class FPCGEditor : public FAssetEditorToolkit, public FGCObject, public FSelfRegisteringEditorUndoClient
+namespace PCGEditorTabs
+{
+	constexpr FLazyName GraphEditorID = TEXT("GraphEditor");
+	constexpr FLazyName PCGEditorDocumentID = TEXT("PCGEditorDocumentID");
+	constexpr FLazyName PropertyDetailsID[] = {
+		TEXT("PropertyDetails"),
+		TEXT("PropertyDetails2"),
+		TEXT("PropertyDetails3"),
+		TEXT("PropertyDetails4") };
+	constexpr FLazyName PaletteID = TEXT("Palette");
+	constexpr FLazyName DebugObjectID = TEXT("DebugObject");
+	constexpr FLazyName AttributesID[] = {
+		TEXT("Attributes"),
+		TEXT("Attributes2"),
+		TEXT("Attributes3"),
+		TEXT("Attributes4") };
+	constexpr FLazyName FindID = TEXT("Find");
+	constexpr FLazyName DeterminismID = TEXT("Determinism");
+	constexpr FLazyName ProfilingID = TEXT("Profiling");
+	constexpr FLazyName LogID = TEXT("Log");
+	constexpr FLazyName HLSLSourceID = TEXT("HLSLSource");
+	constexpr FLazyName UserParamsID = TEXT("UserParams");
+	constexpr FLazyName ViewportID[] = {
+		TEXT("Viewport"),
+		TEXT("Viewport2"),
+		TEXT("Viewport3"),
+		TEXT("Viewport4") };
+
+	constexpr int32 NumPropertyDetailTabs = sizeof(PropertyDetailsID) / sizeof(FLazyName);
+	constexpr int32 NumAttributeTabs = sizeof(AttributesID) / sizeof(FLazyName);
+	constexpr int32 NumViewportTabs = sizeof(ViewportID) / sizeof(FLazyName);
+}
+
+class FPCGEditor : public FWorkflowCentricApplication, public FGCObject, public FSelfRegisteringEditorUndoClient
 {
 public:
 	/** Edits the specified PCGGraph */
@@ -122,16 +158,11 @@ public:
 	PCGEDITOR_API void JumpToNode(const UEdGraphNode* InNode);
 	PCGEDITOR_API void JumpToNode(const UPCGNode* InNode);
 
+PRAGMA_DISABLE_DEPRECATION_WARNINGS
 	/** Get the TabID of the editor panel. */
+	UE_DEPRECATED(5.8, "Use PCGEditorTabs directly")
 	PCGEDITOR_API FName GetPanelID(EPCGEditorPanel Panel) const;
-	/** Focuses the user on a specific panel and flashes the tab. */
-	void BringFocusToPanel(EPCGEditorPanel Panel) const { return BringFocusToPanel(GetPanelID(Panel)); }
-	/** Attempts to close the specific panel if it's open. */
-	void CloseGraphPanel(EPCGEditorPanel Panel) const { CloseGraphPanel(GetPanelID(Panel)); }
-	/** Returns true if the selected tab is currently open. */
-	bool IsPanelCurrentlyOpen(EPCGEditorPanel Panel) const { return IsPanelCurrentlyOpen(GetPanelID(Panel)); }
-	/** Returns true if the selected tab is currently open and focused. */
-	bool IsPanelCurrentlyForeground(EPCGEditorPanel Panel) const { return IsPanelCurrentlyForeground(GetPanelID(Panel)); }
+PRAGMA_ENABLE_DEPRECATION_WARNINGS
 
 	/** Focuses the user on a specific panel and flashes the tab. */
 	PCGEDITOR_API void BringFocusToPanel(const FName PanelID) const;
@@ -151,7 +182,6 @@ public:
 
 	// ~Begin IToolkit interface
 	PCGEDITOR_API virtual void RegisterTabSpawners(const TSharedRef<class FTabManager>& InTabManager) override;
-	PCGEDITOR_API virtual void UnregisterTabSpawners(const TSharedRef<class FTabManager>& InTabManager) override;
 	// ~End IToolkit interface
 
 	// ~Begin FGCObject interface
@@ -169,6 +199,7 @@ public:
 	// ~End FEditorUndoClient interface
 
 	// ~Begin FAssetEditorToolkit interface
+	PCGEDITOR_API virtual FText GetToolkitName() const override;
 	PCGEDITOR_API virtual FName GetToolkitFName() const override;
 	PCGEDITOR_API virtual FText GetBaseToolkitName() const override;
 	PCGEDITOR_API virtual FLinearColor GetWorldCentricTabColorScale() const override;
@@ -177,6 +208,9 @@ public:
 	PCGEDITOR_API virtual void InitToolMenuContext(FToolMenuContext& MenuContext) override;
 	// ~End FAssetEditorToolkit interface
 
+	PCGEDITOR_API void SaveEditedObjectState();
+	PCGEDITOR_API void RestoreEditedObjectState();
+
 	/**
 	 * Handles spawning a graph node in the current graph using the passed in chord.
 	 *
@@ -202,6 +236,19 @@ public:
 
 	UPCGDefaultExecutionSource* GetDefaultExecutionSource() const { return PCGDefaultExecutionSource; }
 
+	PCGEDITOR_API TSharedPtr<FWorkflowTabFactory> CreateGraphEditorTabFactory();
+	PCGEDITOR_API TSharedPtr<FWorkflowTabFactory> CreatePaletteTabFactory();
+	PCGEDITOR_API TSharedPtr<FWorkflowTabFactory> CreateDebugObjectTreeTabFactory();
+	PCGEDITOR_API TSharedPtr<FWorkflowTabFactory> CreateFindTabFactory();
+	PCGEDITOR_API TSharedPtr<FWorkflowTabFactory> CreateAttributesTabFactory(int32 Index, TSharedRef<FWorkspaceItem> Group);
+	PCGEDITOR_API TSharedPtr<FWorkflowTabFactory> CreateDeterminismTabFactory();
+	PCGEDITOR_API TSharedPtr<FWorkflowTabFactory> CreateProfilingTabFactory();
+	PCGEDITOR_API TSharedPtr<FWorkflowTabFactory> CreateLogTabFactory();
+	PCGEDITOR_API TSharedPtr<FWorkflowTabFactory> CreateNodeSourceTabFactory();
+	PCGEDITOR_API TSharedPtr<FWorkflowTabFactory> CreateGraphParamsTabFactory();
+	PCGEDITOR_API TSharedPtr<FWorkflowTabFactory> CreatePropertyDetailsTabFactory(int32 Index, TSharedRef<FWorkspaceItem> Group);
+	PCGEDITOR_API TSharedPtr<FWorkflowTabFactory> CreateViewportTabFactory(int32 Index, TSharedRef<FWorkspaceItem> Group);
+
 protected:
 	PCGEDITOR_API virtual TAttribute<FGraphAppearanceInfo> GetAppearanceInfo() const;
 
@@ -244,13 +291,17 @@ protected:
 	/** Create a new viewport widget. Can be subclassed to have customization around the viewport. */
 	PCGEDITOR_API virtual TSharedRef<SPCGEditorViewport> CreateViewportWidget();
 
+	/** Returns viewport at specified index if it exists, nullptr otherwise */
+	PCGEDITOR_API TSharedPtr<SPCGEditorViewport> GetViewportWidget(int32 Index) const;
+
 	// TO BE REMOVED IN 5.8
 	// @todo_pcg
 	// Temporary boolean to force the Attribute List view to Update even if it is closed. To be set in the Editor Ctor.
 	UE_DEPRECATED(all, "Do not use.")
 	bool bForceRefreshAttributeEvenIfClosed = false;
 
-	friend SPCGEditorGraphAttributeListView;
+	friend class SPCGEditorGraphAttributeListView;
+	friend class FPCGEditorDefaultMode;
 
 private:
 	/** Register PCG specific toolbar for the editor */
@@ -327,7 +378,7 @@ private:
 
 	/** Open panel to view and edit the graph parameters. */
 	void OnToggleGraphParamsPanel() const;
-	/** Is the graph params panel is open. */
+	/** Is the graph params panel open. */
 	bool IsToggleGraphParamsToggled() const;
 
 	/** Select every node in the graph */
@@ -419,35 +470,44 @@ private:
 	void OnDistributeNodesV();
 	void OnCreateComment();
 
-	/** Create new graph editor widget */
+	/** Get or Create new Graph editor widget */
+	TSharedRef<SGraphEditor> GetOrCreateGraphEditorWidget();
+	void CreateGraphEditorCommands();
 	TSharedRef<SGraphEditor> CreateGraphEditorWidget();
+	TSharedRef<SGraphEditor> CreateGraphEditorWidget(TSharedPtr<FTabInfo> InTabInfo, UEdGraph* InGraph);
+	
+	/** Get or Create palette widget */
+	TSharedRef<SPCGEditorGraphNodePalette> GetOrCreatePaletteWidget();
+	
+	/** Get or Create new debug object tree widget */
+	TSharedRef<SPCGEditorGraphDebugObjectTree> GetOrCreateDebugObjectTreeWidget();
 
-	/** Create new palette widget */
-	TSharedRef<SPCGEditorGraphNodePalette> CreatePaletteWidget();
+	/** Get or Create new find widget */
+	TSharedRef<SPCGEditorGraphFind> GetOrCreateFindWidget();
 
-	/** Create new debug object tree widget */
-	TSharedRef<SPCGEditorGraphDebugObjectTree> CreateDebugObjectTreeWidget();
+	/** Get or Create new attributes widget */
+	TSharedRef<SPCGEditorGraphAttributeListView> GetOrCreateAttributesWidget(int32 Index);
 
-	/** Create new find widget */
-	TSharedRef<SPCGEditorGraphFind> CreateFindWidget();
+	/** Get or Create a new determinism tab widget */
+	TSharedRef<SPCGEditorGraphDeterminismListView> GetOrCreateDeterminismWidget();
 
-	/** Create new attributes widget */
-	TSharedRef<SPCGEditorGraphAttributeListView> CreateAttributesWidget(int32 Index);
+	/** Get or Create a new profiling tab widget */
+	TSharedRef<SPCGEditorGraphProfilingView> GetOrCreateProfilingWidget();
 
-	/** Create a new determinism tab widget */
-	TSharedRef<SPCGEditorGraphDeterminismListView> CreateDeterminismWidget();
+	/** Get or Create a new log capture tab widget */
+	TSharedRef<SPCGEditorGraphLogView> GetOrCreateLogWidget();
 
-	/** Create a new profiling tab widget */
-	TSharedRef<SPCGEditorGraphProfilingView> CreateProfilingWidget();
+	/** Get or Create a new node source editor tab widget */
+	TSharedRef<SPCGEditorNodeSource> GetOrCreateNodeSourceWidget();
 
-	/** Create a new log capture tab widget */
-	TSharedRef<SPCGEditorGraphLogView> CreateLogWidget();
+	/** Get or Create a new user graph parameters tab widget */
+	TSharedRef<SPCGEditorGraphUserParametersView> GetOrCreateGraphParamsWidget();
 
-	/** Create a new node source editor tab widget */
-	TSharedRef<SPCGEditorNodeSource> CreateNodeSourceWidget();
+	/** Get or Create a new details tab widget */
+	TSharedRef<SPCGEditorGraphDetailsView> GetOrCreatePropertyDetailsWidget(int32 Index);
 
-	/** Create a new user graph parameters tab widget */
-	TSharedRef<SPCGEditorGraphUserParametersView> CreateGraphParamsWidget();
+	/** Get or Create a new viewport tab widget */
+	TSharedRef<SPCGEditorViewport> GetOrCreateViewportWidget(int32 Index);
 
 	/** Called when the component inspected is generated/cleaned */
 	void OnSourceGenerated(IPCGGraphExecutionSource* InSource);
@@ -488,37 +548,20 @@ private:
 
 	void UpdateDefaultExecutionSource();
 
-	TSharedRef<SDockTab> SpawnTab_GraphEditor(const FSpawnTabArgs& Args);
-	TSharedRef<SDockTab> SpawnTab_PropertyDetails(const FSpawnTabArgs& Args, int PropertyDetailsIndex);
-	TSharedRef<SDockTab> SpawnTab_Palette(const FSpawnTabArgs& Args);
-	TSharedRef<SDockTab> SpawnTab_DebugObjectTree(const FSpawnTabArgs& Args);
-	TSharedRef<SDockTab> SpawnTab_Attributes(const FSpawnTabArgs& Args, int AttributesIndex);
-	TSharedRef<SDockTab> SpawnTab_Find(const FSpawnTabArgs& Args);
-	TSharedRef<SDockTab> SpawnTab_Determinism(const FSpawnTabArgs& Args);
-	TSharedRef<SDockTab> SpawnTab_Profiling(const FSpawnTabArgs& Args);
-	TSharedRef<SDockTab> SpawnTab_Log(const FSpawnTabArgs& Args);
-	TSharedRef<SDockTab> SpawnTab_NodeSource(const FSpawnTabArgs& Args);
-	TSharedRef<SDockTab> SpawnTab_UserParams(const FSpawnTabArgs& Args);
-	TSharedRef<SDockTab> SpawnTab_Viewport(const FSpawnTabArgs& Args, int ViewportIndex);
-
-	FText GetDetailsTabLabel(int DetailsIndex);
-	FText GetDetailsViewObjectName(int DetailsIndex);
-	FText GetAttributesTabLabel(int AttributesIndex);
-	FText GetViewportTabLabel(int ViewportIndex);
-
 	PCGEDITOR_API virtual TSharedRef<FTabManager::FLayout> GetDefaultLayout() const;
 
 	TSharedPtr<SGraphEditor> GraphEditorWidget;
-	TArray<TSharedPtr<SPCGEditorGraphDetailsView>> PropertyDetailsWidgets;
+	TSharedPtr<SPCGEditorGraphDetailsView> PropertyDetailsWidgets[PCGEditorTabs::NumPropertyDetailTabs];
 	TSharedPtr<SPCGEditorGraphNodePalette> PaletteWidget;
 	TSharedPtr<SPCGEditorGraphDebugObjectTree> DebugObjectTreeWidget;
 	TSharedPtr<SPCGEditorGraphFind> FindWidget;
-	TArray<TSharedPtr<SPCGEditorGraphAttributeListView>> AttributesWidgets;
+	TSharedPtr<SPCGEditorGraphAttributeListView> AttributesWidgets[PCGEditorTabs::NumAttributeTabs];
 	TSharedPtr<SPCGEditorGraphDeterminismListView> DeterminismWidget;
 	TSharedPtr<SPCGEditorGraphProfilingView> ProfilingWidget;
 	TSharedPtr<SPCGEditorGraphLogView> LogWidget;
 	TSharedPtr<SPCGEditorNodeSource> NodeSourceWidget;
 	TSharedPtr<SPCGEditorGraphUserParametersView> UserParamsWidget;
+	TSharedPtr<SPCGEditorViewport> ViewportWidgets[PCGEditorTabs::NumViewportTabs];
 
 	TSharedPtr<FUICommandList> GraphEditorCommands;
 
@@ -531,4 +574,7 @@ private:
 	// Keep track of the last execution status to be able to break infinite loop when a source is triggered to be generated by inspection
 	// aborted and re-triggered.
 	TOptional<TPair<const IPCGGraphExecutionSource*, EPCGGenerationStatus>> LastExecutionStatus;
+
+	/** Document manager for workflow tabs */
+	TSharedPtr<FDocumentTracker> DocumentManager;
 };
diff --git a/Engine/Plugins/PCG/Source/PCGEditor/Public/PCGEditorTabFactories.h b/Engine/Plugins/PCG/Source/PCGEditor/Public/PCGEditorTabFactories.h
new file mode 100644
--- /dev/null
+++ b/Engine/Plugins/PCG/Source/PCGEditor/Public/PCGEditorTabFactories.h
@@ -0,0 +1,44 @@
+// Copyright Epic Games, Inc. All Rights Reserved.
+
+#pragma once
+
+#include "WorkflowOrientedApp/WorkflowTabFactory.h"
+
+enum ETabRole : uint8;
+class FPCGEditor;
+struct FSlateIcon;
+class FText;
+class FWorkspaceItem;
+class SWidget;
+
+struct FPCGGenericTabFactoryParams
+{
+	FPCGGenericTabFactoryParams(FName InId, TSharedPtr<FPCGEditor> InEditor, TFunction<TSharedRef<SWidget>()> InCreateTabBodyFunc)
+		: Id(InId)
+		, Editor(InEditor)
+		, CreateTabBodyFunc(InCreateTabBodyFunc)
+	{ }
+
+	FName Id;
+	TSharedPtr<FPCGEditor> Editor;
+	TFunction<TSharedRef<SWidget>()> CreateTabBodyFunc;
+	
+	bool bSingleton = true;
+
+	TOptional<ETabRole> Role;
+	TOptional<FSlateIcon> Icon;
+	TOptional<FText> Label;
+
+	TSharedPtr<FWorkspaceItem> Group;
+};
+
+struct FPCGGenericTabFactory : public FWorkflowTabFactory
+{
+public:
+	FPCGGenericTabFactory(const FPCGGenericTabFactoryParams& InParams);
+
+	virtual TSharedRef<SWidget> CreateTabBody(const FWorkflowTabSpawnInfo& Info) const override;
+
+private:
+	TFunction<TSharedRef<SWidget>()> CreateTabBodyFunc;
+};
\ No newline at end of file
diff --git a/Engine/Source/Editor/UnrealEd/Private/WorkflowOrientedApp/WorkflowTabFactory.cpp b/Engine/Source/Editor/UnrealEd/Private/WorkflowOrientedApp/WorkflowTabFactory.cpp
--- a/Engine/Source/Editor/UnrealEd/Private/WorkflowOrientedApp/WorkflowTabFactory.cpp
+++ b/Engine/Source/Editor/UnrealEd/Private/WorkflowOrientedApp/WorkflowTabFactory.cpp
@@ -140,7 +140,11 @@ FTabSpawnerEntry& FWorkflowTabFactory::RegisterTabSpawner(TSharedRef<FTabManager
 		.SetDisplayName(ConstructTabName(SpawnInfo).Get())
 		.SetTooltipText(GetTabToolTipText(SpawnInfo));
 
-	if (CurrentApplicationMode)
+	if (Group)
+	{
+		SpawnerEntry.SetGroup(Group.ToSharedRef());
+	}
+	else if (CurrentApplicationMode)
 	{
 		SpawnerEntry.SetGroup(CurrentApplicationMode->GetWorkspaceMenuCategory());
 	}
diff --git a/Engine/Source/Editor/UnrealEd/Public/WorkflowOrientedApp/WorkflowTabFactory.h b/Engine/Source/Editor/UnrealEd/Public/WorkflowOrientedApp/WorkflowTabFactory.h
--- a/Engine/Source/Editor/UnrealEd/Public/WorkflowOrientedApp/WorkflowTabFactory.h
+++ b/Engine/Source/Editor/UnrealEd/Public/WorkflowOrientedApp/WorkflowTabFactory.h
@@ -85,6 +85,8 @@ protected:
 	FText ViewMenuDescription;
 	FText ViewMenuTooltip;
 
+	TSharedPtr<FWorkspaceItem> Group;
+
 	TWeakPtr<class FAssetEditorToolkit> HostingApp;
 
 public:
