[PCG] Add missing DebugObjectTreeWidget valid checks (regression from 49926118)

#rnx

#Synced-CL 49926118

diff --git a/Engine/Plugins/PCG/Source/PCGEditor/Private/PCGEditor.cpp b/Engine/Plugins/PCG/Source/PCGEditor/Private/PCGEditor.cpp
--- a/Engine/Plugins/PCG/Source/PCGEditor/Private/PCGEditor.cpp
+++ b/Engine/Plugins/PCG/Source/PCGEditor/Private/PCGEditor.cpp
@@ -1980,9 +1980,12 @@ void FPCGEditor::SetNodeInspected(UPCGEditorGraphNodeBase* InspectedNode, bool b
 			GetTabManager()->TryInvokeTab(FTabId(PCGEditorTabs::AttributesID[0]));
 		}
 
-		DebugObjectTreeWidget->SetNodeBeingInspected(PCGNode);
+		if (DebugObjectTreeWidget.IsValid())
+		{
+			DebugObjectTreeWidget->SetNodeBeingInspected(PCGNode);
+		}
 	}
-	else
+	else if(DebugObjectTreeWidget.IsValid())
 	{
 		DebugObjectTreeWidget->SetNodeBeingInspected(nullptr);
 	}
@@ -3301,7 +3304,10 @@ void FPCGEditor::OnComponentUnregistered(UPCGComponent* Component)
 	// Refresh the debug object tree to avoid stale entries from components that have been unregistered.
 	if (!Component || Component->GetGraph() == PCGGraphBeingEdited)
 	{
-		DebugObjectTreeWidget->RequestRefresh();
+		if (DebugObjectTreeWidget.IsValid())
+		{
+			DebugObjectTreeWidget->RequestRefresh();
+		}
 	}
 
 	if (IPCGEditorModule* PCGEditorModule = IPCGEditorModule::Get())
