[PCG] Always create matching Attribute view for Viewport now that widgets are lazy created (regression from 49926118)

#rnx

#Synced-CL 49926118

diff --git a/Engine/Plugins/PCG/Source/PCGEditor/Private/PCGEditor.cpp b/Engine/Plugins/PCG/Source/PCGEditor/Private/PCGEditor.cpp
--- a/Engine/Plugins/PCG/Source/PCGEditor/Private/PCGEditor.cpp
+++ b/Engine/Plugins/PCG/Source/PCGEditor/Private/PCGEditor.cpp
@@ -3117,7 +3117,12 @@ TSharedPtr<FWorkflowTabFactory> FPCGEditor::CreateViewportTabFactory(int32 Index
 {
 	check(Index >= 0 && Index < PCGEditorTabs::NumViewportTabs);
 
-	FPCGGenericTabFactoryParams Params(PCGEditorTabs::ViewportID[Index], SharedThis<FPCGEditor>(this), [this, Index]() { return GetOrCreateViewportWidget(Index); });
+	FPCGGenericTabFactoryParams Params(PCGEditorTabs::ViewportID[Index], SharedThis<FPCGEditor>(this), [this, Index]() 
+	{
+		// Attribute Widget is tied to Viewport Widget so create it to (to make sure Viewport gets refreshed properly through Attribute Widget)
+		GetOrCreateAttributesWidget(Index);
+		return GetOrCreateViewportWidget(Index); 
+	});
 
 	static FText ViewportText[4] =
 	{
