[PCG] Support execution sources in a few places

#rnx

#Synced-CL 49947524

diff --git a/Engine/Plugins/PCG/Source/PCG/Private/Compute/DataInterfaces/PCGVirtualTextureDataInterface.cpp b/Engine/Plugins/PCG/Source/PCG/Private/Compute/DataInterfaces/PCGVirtualTextureDataInterface.cpp
--- a/Engine/Plugins/PCG/Source/PCG/Private/Compute/DataInterfaces/PCGVirtualTextureDataInterface.cpp
+++ b/Engine/Plugins/PCG/Source/PCG/Private/Compute/DataInterfaces/PCGVirtualTextureDataInterface.cpp
@@ -2,8 +2,8 @@
 
 #include "PCGVirtualTextureDataInterface.h"
 
-#include "PCGComponent.h"
 #include "PCGData.h"
+#include "PCGGraphExecutionStateInterface.h"
 #include "Compute/PCGDataBinding.h"
 #include "Data/PCGVirtualTextureData.h"
 
@@ -121,14 +121,14 @@ void UPCGVirtualTextureDataProvider::Initialize(const UComputeDataInterface* InD
 
 	const UPCGVirtualTextureDataInterface* DataInterface = CastChecked<UPCGVirtualTextureDataInterface>(InDataInterface);
 	UPCGDataBinding* Binding = CastChecked<UPCGDataBinding>(InBinding);
-	const UPCGComponent* SourceComponent = Cast<UPCGComponent>(Binding->GetExecutionSource());
+	const IPCGGraphExecutionSource* Source = Cast<IPCGGraphExecutionSource>(Binding->GetExecutionSource());
 
 	// Take any input pin label alias to obtain the data from the input data collection.
 	check(!DataInterface->GetDownstreamInputPinLabelAliases().IsEmpty());
 	const TArray<FPCGTaggedData> InputTaggedData = Binding->GetInputDataCollection().GetInputsByPin(DataInterface->GetDownstreamInputPinLabelAliases()[0]);
 
 	// Sampling of virtual textures is not supported outside of runtime generation at this time, because the RuntimeGenScheduler is responsible for priming the virtual textures.
-	if (!SourceComponent || !SourceComponent->IsManagedByRuntimeGenSystem())
+	if (!Source || !Source->GetExecutionState().IsManagedByRuntimeGenSystem())
 	{
 		UE_LOG(LogPCG, Error, TEXT("Virtual texture data interface on pin '%s' failed to initialize. Virtual texture sampling is only supported for runtime generation."),
 			*GetDownstreamInputPinLabelAliases()[0].ToString());
diff --git a/Engine/Plugins/PCG/Source/PCG/Private/Graph/PCGGraphExecutor.cpp b/Engine/Plugins/PCG/Source/PCG/Private/Graph/PCGGraphExecutor.cpp
--- a/Engine/Plugins/PCG/Source/PCG/Private/Graph/PCGGraphExecutor.cpp
+++ b/Engine/Plugins/PCG/Source/PCG/Private/Graph/PCGGraphExecutor.cpp
@@ -421,12 +421,11 @@ FPCGTaskId FPCGGraphExecutor::ScheduleGraph(const FPCGScheduleGraphParams& InPar
 	uint32 GenerationGridSize = PCGHiGenGrid::UninitializedGridSize();
 	if (InParams.bAllowHierarchicalGeneration && InParams.Graph->IsHierarchicalGenerationEnabled())
 	{
-		if (const UPCGComponent* PCGComponent = Cast<UPCGComponent>(InParams.ExecutionSource))
+		const IPCGGraphExecutionState& ExecutionState = InParams.ExecutionSource->GetExecutionState();
+
+		if (ExecutionState.IsLocalSource() || ExecutionState.IsPartitioned())
 		{
-			if (PCGComponent->IsLocalComponent() || PCGComponent->IsPartitioned())
-			{
-				GenerationGridSize = PCGComponent->GetGenerationGridSize();
-			}
+			GenerationGridSize = ExecutionState.GetGenerationGridSize();
 		}
 	}
 
diff --git a/Engine/Plugins/PCG/Source/PCG/Private/Graph/PCGStackContext.cpp b/Engine/Plugins/PCG/Source/PCG/Private/Graph/PCGStackContext.cpp
--- a/Engine/Plugins/PCG/Source/PCG/Private/Graph/PCGStackContext.cpp
+++ b/Engine/Plugins/PCG/Source/PCG/Private/Graph/PCGStackContext.cpp
@@ -83,9 +83,9 @@ bool FPCGStack::CreateStackFramePath(FString& OutString, const UPCGNode* InNode,
 					return false;
 				}
 
-				if (Object->IsA<UPCGComponent>())
+				if (Object->Implements<UPCGGraphExecutionSource>())
 				{
-					StringBuilder << TEXT("COMPONENT:") << Object->GetFullName();
+					StringBuilder << TEXT("EXECUTION SOURCE:") << Object->GetFullName();
 				}
 				else if (Object->IsA<UPCGGraph>())
 				{
