[PCG] Fix cases where Get Actor Data would grab data for actors that returned invalid bounds.

#rnx

#Synced-CL 49945467

diff --git a/Engine/Plugins/PCG/Source/PCG/Private/Elements/PCGDataFromActor.cpp b/Engine/Plugins/PCG/Source/PCG/Private/Elements/PCGDataFromActor.cpp
--- a/Engine/Plugins/PCG/Source/PCG/Private/Elements/PCGDataFromActor.cpp
+++ b/Engine/Plugins/PCG/Source/PCG/Private/Elements/PCGDataFromActor.cpp
@@ -381,7 +381,7 @@ bool FPCGDataFromActorElement::ExecuteInternal(FPCGContext* InContext) const
 			BoundsCheck = [Settings, Bounds, ExecutionSource](const AActor* OtherActor) -> bool
 			{
 				const FBox OtherActorBounds = OtherActor ? PCGHelpers::GetGridBounds(OtherActor, Cast<UPCGComponent>(ExecutionSource)) : FBox(EForceInit::ForceInit);
-				return Bounds.Intersect(OtherActorBounds);
+				return OtherActorBounds.IsValid && Bounds.IsValid && Bounds.Intersect(OtherActorBounds);
 			};
 		}
 
