[PCG] Optimize game thread cost of the cleanup scan in the runtime generation scheduler by only running it if any generation source changed.

It relies on the generation source change detector which was never stood up for non-partitioned generation sources, so that will come later (reverts to scanning every frame if these are present).

[RN] [PCG] Optimized per-frame runtime generation scheduler game thread cost by only scanning components for cleanup when generation sources change.

#Synced-CL 49945467

diff --git a/Engine/Plugins/PCG/Source/PCG/Private/RuntimeGen/PCGRuntimeGenScheduler.cpp b/Engine/Plugins/PCG/Source/PCG/Private/RuntimeGen/PCGRuntimeGenScheduler.cpp
--- a/Engine/Plugins/PCG/Source/PCG/Private/RuntimeGen/PCGRuntimeGenScheduler.cpp
+++ b/Engine/Plugins/PCG/Source/PCG/Private/RuntimeGen/PCGRuntimeGenScheduler.cpp
@@ -58,7 +58,7 @@ namespace PCGRuntimeGenSchedulerHelpers
 	static TAutoConsoleVariable<bool> CVarRuntimeGenerationEnableChangeDetection(
 		TEXT("pcg.RuntimeGeneration.EnableChangeDetection"),
 		true,
-		TEXT("Skips execution of scheduling loop if generation sources and other world state have not changed."));
+		TEXT("Skips execution of scheduling and cleanup loops if generation sources and other world state have not changed to reduce tick cost on game thread."));
 
 	static TAutoConsoleVariable<int32> CVarNumGeneratingComponentsAtSameTime(
 		TEXT("pcg.RuntimeGeneration.NumGeneratingComponents"),
@@ -347,7 +347,7 @@ void FPCGRuntimeGenScheduler::Tick(APCGWorldActor* InPCGWorldActor, double InEnd
 
 	// 2. Schedule cleanup on components that become out of range.
 
-	if (!GeneratedComponents.IsEmpty())
+	if (!GeneratedComponents.IsEmpty() && (GenSources.IsEmpty() || bCleanupScanRequired || !IsChangeDetectionEnabled()))
 	{
 		TickCleanup(GenSources, InPCGWorldActor, InEndTime);
 	}
@@ -515,7 +515,7 @@ void FPCGRuntimeGenScheduler::TickQueueComponentsForGeneration(
 		return bIsLoaded;
 	};
 
-	const bool bDoChangeDetection = PCGRuntimeGenSchedulerHelpers::CVarRuntimeGenerationEnableChangeDetection.GetValueOnGameThread();
+	const bool bDoChangeDetection = IsChangeDetectionEnabled();
 #if UE_ENABLE_DEBUG_DRAWING
 	const bool bDebugDrawGenerationSources = PCGSystemSwitches::CVarPCGDebugDrawGeneratedCells.GetValueOnGameThread();
 #endif
@@ -601,6 +601,9 @@ void FPCGRuntimeGenScheduler::TickQueueComponentsForGeneration(
 					continue;
 				}
 
+				// If any runtime gen source changed, trigger a cleanup scan. Simple heuristic, could break this down per-component in the future.
+				bCleanupScanRequired = true;
+
 #if PCG_RGS_ONSCREENDEBUGMESSAGES
 				++PCGRuntimeGenSchedulerHelpers::Stats.GridCellScanCounter;
 #endif
@@ -810,6 +813,9 @@ void FPCGRuntimeGenScheduler::TickQueueComponentsForGeneration(
 		// The generation key for a non-partitioned component should always have unbounded grid size and 0,0,0 cell coord.
 		if (Inputs.GeneratedComponents->Contains({ PCGHiGenGrid::UnboundedGridSize(), FIntVector(0), OriginalComponent }))
 		{
+			// todo_pcg: Needs change detection to throttle clean up scans, for now just always cleanup if any runtime gen unbounded component is present.
+			bCleanupScanRequired = true;
+
 			continue;
 		}
 
@@ -818,6 +824,9 @@ void FPCGRuntimeGenScheduler::TickQueueComponentsForGeneration(
 		// TODO: For each execution domain (for now only GenAtRuntime/dynamic), assuming we run Preview through this scheduler, which it seems like we will.
 		if (OriginalComponent->IsManagedByRuntimeGenSystem() && ensure(Policy))
 		{
+			// todo_pcg: Needs change detection to throttle clean up scans, for now just always cleanup if any runtime gen unbounded component is present.
+			bCleanupScanRequired = true;
+
 			const FBox OriginalComponentBounds = OriginalComponent->GetGridBounds();
 
 			// Poll for streaming to complete.
@@ -928,7 +937,7 @@ void FPCGRuntimeGenScheduler::TickCleanup(const TSet<IPCGGenSourceBase*>& InGenS
 	std::atomic<int> InvalidKeysCounter = 0;
 
 	ParallelFor(
-		TEXT("FPCGRuntimeGenScheduler::TickCleanup"),
+		TEXT("CleanupLoop"),
 		NumGeneratedComponents,
 		/*MinBatchSize*/6, // Optimal value found in testing.
 		[&GeneratedComponentsArray, &ComponentsToClean, &InvalidKeys, &ComponentsToCleanCounter, &InvalidKeysCounter, TrackingManager=TrackingManager, &CheckIfAllGenSourcesWantToCleanup](int32 Index)
@@ -1015,17 +1024,24 @@ void FPCGRuntimeGenScheduler::TickCleanup(const TSet<IPCGGenSourceBase*>& InGenS
 		GeneratedComponents.Remove(InvalidKeys[I]);
 	}
 
-	for (int I = 0; I < NumComponentsToClean; ++I)
+	int CleanupIndex = 0;
+	for (; CleanupIndex < NumComponentsToClean; ++CleanupIndex)
 	{
-		const PCGComponentToClean& ComponentToClean = ComponentsToClean[I];
+		const PCGComponentToClean& ComponentToClean = ComponentsToClean[CleanupIndex];
 		CleanupComponent(ComponentToClean.Get<0>(), ComponentToClean.Get<1>());
 
 		if (FPlatformTime::Seconds() >= InEndTime)
 		{
-			UE_LOG(LogPCG, Verbose, TEXT("FPCGRuntimeGenScheduler: Time budget exceeded, aborted after cleaning up %d / %d components"), (I + 1), NumComponentsToClean);
+			UE_LOG(LogPCG, Verbose, TEXT("FPCGRuntimeGenScheduler: Time budget exceeded, aborted after cleaning up %d / %d components"), (CleanupIndex + 1), NumComponentsToClean);
 			break;
 		}
 	}
+
+	if (CleanupIndex == NumComponentsToClean)
+	{
+		// Record that we completed a cleanup scan.
+		bCleanupScanRequired = false;
+	}
 }
 
 void FPCGRuntimeGenScheduler::TickScheduleGeneration(TMap<FGridGenerationKey, double>& InOutComponentsToGenerate)
@@ -1441,6 +1457,9 @@ void FPCGRuntimeGenScheduler::OnOriginalComponentUnregistered(UPCGComponent* InO
 	// When an original/non-partitioned component is unregistered, we need to dirty the state.
 	bContextDirty = true;
 
+	// Make sure we perform a scan for cells that need cleanup on next tick.
+	bCleanupScanRequired = true;
+
 #if WITH_EDITOR
 	UnregisterGraphParameterChangedEvent(InOriginalComponent);
 #endif
@@ -1991,6 +2010,7 @@ void FPCGRuntimeGenScheduler::RefreshComponent(IPCGGraphExecutionSource* InCompo
 	{
 		// When an original/non-partitioned component is refreshed, we need to dirty the state.
 		bContextDirty = true;
+		bCleanupScanRequired = true;
 	}
 }
 
@@ -2193,6 +2213,11 @@ void FPCGRuntimeGenScheduler::RefreshContext()
 	bContextDirty = false;
 }
 
+bool FPCGRuntimeGenScheduler::IsChangeDetectionEnabled() const
+{
+	return PCGRuntimeGenSchedulerHelpers::CVarRuntimeGenerationEnableChangeDetection.GetValueOnGameThread();
+}
+
 #if WITH_EDITOR
 void FPCGRuntimeGenScheduler::RegisterGraphParameterChangedEvent(UPCGComponent* InOriginalComponent)
 {
diff --git a/Engine/Plugins/PCG/Source/PCG/Public/RuntimeGen/PCGRuntimeGenScheduler.h b/Engine/Plugins/PCG/Source/PCG/Public/RuntimeGen/PCGRuntimeGenScheduler.h
--- a/Engine/Plugins/PCG/Source/PCG/Public/RuntimeGen/PCGRuntimeGenScheduler.h
+++ b/Engine/Plugins/PCG/Source/PCG/Public/RuntimeGen/PCGRuntimeGenScheduler.h
@@ -219,6 +219,9 @@ private:
 	/** Update RGS context information. */
 	void RefreshContext();
 
+	/** Whether the change detection is enabled which avoids scanning grids when not necessary. */
+	bool IsChangeDetectionEnabled() const;
+
 #if WITH_EDITOR
 	void RegisterGraphParameterChangedEvent(UPCGComponent* InOriginalComponent);
 	void UnregisterGraphParameterChangedEvent(UPCGComponent* InOriginalComponent);
@@ -281,6 +284,9 @@ private:
 	/** Used to detect when world state (generation sources, cvars etc) have changed enough to warrant rescanning generation cells. */
 	PCGRuntimeGenChangeDetection::FDetector ChangeDetector;
 
+	/** Triggers scanning components for cleanup. If change detection is enabled, cleanup scans are only performed when needed rather than every frame. */
+	bool bCleanupScanRequired = false;
+
 	/** Meta-information about current context. */
 	FPCGRuntimeGenContext Context;
 
