[PCG] Time slice generate landscape textures.

Fixes large GPU cost when multiple landscape components exported in a single frame.

[RN] [PCG] Added support for timeslicing generation of landscape height map and grass maps, to spread GPU cost across frames.

#Synced-CL 49945467

diff --git a/Engine/Plugins/PCG/Shaders/Private/PCGGrassMapUnpackerCS.usf b/Engine/Plugins/PCG/Shaders/Private/PCGGrassMapUnpackerCS.usf
--- a/Engine/Plugins/PCG/Shaders/Private/PCGGrassMapUnpackerCS.usf
+++ b/Engine/Plugins/PCG/Shaders/Private/PCGGrassMapUnpackerCS.usf
@@ -20,33 +20,37 @@
 
 Texture2D<float4> InPackedGrassMaps;
 
-// Full output domain consists of multiple components, but we only run on a "window" that covers the generation cell bounds.
+// This dispatch will in general target a sub area of the output texture corresponding to the exported landscape component(s), defined by InDispatchOffset and InDispatchSize.
+// The bottom left corner of the output texture is the origin of the coordinates, and the tiles of landscape components are indexed relative to this origin.
 // The shader writes to an output texture with height data and a texture array where each slice is a single grass combined grass map.
-// Each component is an inner grid cell. Res is vert count and includes outer edges, so a row/col is removed from each internal boundary.
+// InLandscapeComponentResolution (Res below) is vert count and includes outer edges, so a row/col is removed from each internal boundary.
+// Below shows grass map configuration, height map is same mapping but only a single slice.
 //
-//                         -------------------------------
-//                       -------------------------------  |
-//                      |  Res-1  |  Res-1  |    Res    | |
-//                      |         |         |           | |
-//  (Res-1) * NumTilesY |Res-1    |     +---------+     | |
-//         + 1          |         |     |         |     | |
-//                      |---------+-----|  Exec   |-----| |
-//                      |         |     |  Bounds |     | |
-//                      |Res-1    |     |         |     | |
-//                      |         |     +---------+     | |
-//                      |         |      InWindow       | |
-//                      |---------+---------+-----------| |
-//                      |         |         |           | |
-//                      |         |         |           | |
-//                      |Res      |         |           | |
-//                      |         |         |           | |
-//                      |         |         |           |-  NumGrassMaps slices
-//                       -------------------------------
+//                         --------------------------------
+//                       --------------------------------  |
+//                      |  Res     |   Res-1  |   Res-1  | |
+//                      |          |          |          | |
+//  (Res-1) * NumTilesY |Res-1   +-------------(OW,OH)   | |
+//         + 1          |        | Output Texture |      | |
+//                      |--------|                |------| |       OW, OH: OutputWidth, InOutputHeight - output texture dimensions
+//                      |        |     +--(DO+DS) |      | |       DO, DS: InDispatchOffset, InDispatchSize - domain of dispatched threads
+//                      |Res-1   |     |      |   |OH    | |
+//                      |        |     |      |   |      | |
+//                      |        |    (DO)----+   |      | |
+//                      |--------|                |------| |
+//                      |      (0,0)--------------+      | | .
+//                      |          |         |           | | |
+//                      |Res       |         |           | |-
+//                      |          |         |           | | NumGrassMaps slices
+//                      |   Res    |  Res-1  |  Res-1    |-
+//                       --------------------------------
 //                            (Res-1) * NumTilesX + 1
 //
+
+// Grass map output. Dimension is (OutputWidth, InOutputHeight, NumGrassMaps).
 RWTexture2DArray<float> OutUnpackedGrassMaps;
 
-// Height output. Dimension is (NumComponentsX * InLandscapeComponentResolution, NumComponentsY * InLandscapeComponentResolution, NumGrassMaps).
+// Height map output. Dimension is (OutputWidth, InOutputHeight).
 RWTexture2D<float> OutUnpackedHeight;
 
 StructuredBuffer<int> InComponentHeightOffsets;
@@ -56,28 +60,31 @@ int InNumExportedGrassTypes;
 
 int4 InLinearTileIndexToComponentIndex[PCG_MAX_NUM_LANDSCAPE_COMPONENTS]; // Only x component used, 16 byte alignment on arrays
 uint2 InNumTiles;
+int2 InTileGridOrigin; // Coordinate of bottom left corner of export tile grid, in output coordinates.
 uint InLandscapeComponentResolution; // One texel per vert, so number of quads + 1
-float3 InLandscapeGridScale; // x == LS Actor DrawScale.X, y == LS Actor DrawScale.y, z == LS Actor DrawScale.z / 128.0f (ZSCALE)
-float InLandscapeLocationZ;
+float InLandscapeScaleZ; // LS Actor DrawScale.z / 128.0f (ZSCALE)
+float InLandscapeRootZ;
 uint InOutputHeight;
-uint4 InWindow; // xy == Window start coord, zw = Window end coord
+uint2 InDispatchOffset; // Bottom left coordinate of this dispatch, in output coordinates.
+uint2 InDispatchSize; // Number of XY threads in this dispatch.
 
-// Dispatch group count is ((InWindow.zw - InWindow.xy + 1) / 8, Max(1, NumGrassMaps)).
+// Dispatch group count is (ceil(InDispatchSize / 8), Max(1, NumGrassMaps)).
 [numthreads(THREADGROUPSIZE_X, THREADGROUPSIZE_Y, THREADGROUPSIZE_Z)]
-void PCGGrassMapUnpacker_CS(const uint3 WindowCoord_OutputSlice : SV_DispatchThreadID)
+void PCGGrassMapUnpacker_CS(const uint3 DispatchThreadId : SV_DispatchThreadID)
 {
-	const uint2 WindowCoord = WindowCoord_OutputSlice.xy;
-	const uint2 FullCoord = InWindow.xy + WindowCoord;
-	if (any(FullCoord > InWindow.zw))
+	if (any(DispatchThreadId.xy >= InDispatchSize))
 	{
 		return;
 	}
 
-	const uint OutputSlice = WindowCoord_OutputSlice.z;
+	const uint2 OutputCoord = DispatchThreadId.xy + InDispatchOffset;
+	const uint OutputSlice = DispatchThreadId.z;
+
+	const uint2 TileGridCoord = uint2((int2)OutputCoord - InTileGridOrigin);
 
 	// The rendered landscape components in the input are not in any particular order, so lookup which component index for this output tile.
 	// Clamp because we are dividing by resolution-1 and last col/row of texture will overflow.
-	const uint2 TileCoords = min(FullCoord / (InLandscapeComponentResolution - 1u), InNumTiles - 1u);
+	const uint2 TileCoords = min(TileGridCoord / (InLandscapeComponentResolution - 1u), InNumTiles - 1u);
 
 	const uint TileLinearIndex = TileCoords.y * InNumTiles.x + TileCoords.x;
 	const int ComponentIndex = InLinearTileIndexToComponentIndex[TileLinearIndex][0];
@@ -88,14 +95,14 @@ void PCGGrassMapUnpacker_CS(const uint3 WindowCoord_OutputSlice : SV_DispatchThr
 		// First layer of threads also deals with height.
 		if (InOutputHeight == 1 && OutputSlice == 0)
 		{
-			OutUnpackedHeight[WindowCoord] = 0.0f;
+			OutUnpackedHeight[OutputCoord] = 0.0f;
 		}
 
-		OutUnpackedGrassMaps[uint3(WindowCoord, OutputSlice)] = 0.0f;
+		OutUnpackedGrassMaps[uint3(OutputCoord, OutputSlice)] = 0.0f;
 		return;
 	}
 
-	const uint2 LocalTileCoords = uint2(FullCoord.x - TileCoords.x * (InLandscapeComponentResolution - 1), FullCoord.y - TileCoords.y * (InLandscapeComponentResolution - 1));
+	const uint2 LocalTileCoords = uint2(TileGridCoord.x - TileCoords.x * (InLandscapeComponentResolution - 1), TileGridCoord.y - TileCoords.y * (InLandscapeComponentResolution - 1));
 
 	// Heights. The first layer threads also deal with height.
 	if (InOutputHeight == 1 && OutputSlice == 0)
@@ -104,8 +111,8 @@ void PCGGrassMapUnpacker_CS(const uint3 WindowCoord_OutputSlice : SV_DispatchThr
 
 		const uint2 InputCoords = uint2(LocalTileCoords.x + ComponentXOffset, LocalTileCoords.y);
 
-		const float RelativeHeight = (UnpackHeight(InPackedGrassMaps[InputCoords].xy) - LANDSCAPE_MID_VALUE) * InLandscapeGridScale.z * TERRAIN_ZSCALE;
-		OutUnpackedHeight[WindowCoord] = RelativeHeight + InLandscapeLocationZ;
+		const float RelativeHeight = (UnpackHeight(InPackedGrassMaps[InputCoords].xy) - LANDSCAPE_MID_VALUE) * InLandscapeScaleZ * TERRAIN_ZSCALE;
+		OutUnpackedHeight[OutputCoord] = RelativeHeight + InLandscapeRootZ;
 	}
 
 	// Coordinates into input height & grass maps (see packing above).
@@ -120,19 +127,19 @@ void PCGGrassMapUnpacker_CS(const uint3 WindowCoord_OutputSlice : SV_DispatchThr
 		const uint2 InputCoords = uint2(LocalTileCoords.x + ComponentXOffset, LocalTileCoords.y);
 
 		// Grass maps
-		OutUnpackedGrassMaps[uint3(WindowCoord, OutputSlice)] = InPackedGrassMaps[InputCoords][ExportChannelIndex];
+		OutUnpackedGrassMaps[uint3(OutputCoord, OutputSlice)] = InPackedGrassMaps[InputCoords][ExportChannelIndex];
 
 		// TODO: position of height data can move around based on slices, so also needs to be in mapping somehow
 		// Heights. The first layer threads also deal with height.
 		if (InOutputHeight == 1 && OutputSlice == 0)
 		{
-			const float RelativeHeight = (UnpackHeight(InPackedGrassMaps[InputCoords].xy) - LANDSCAPE_MID_VALUE) * InLandscapeGridScale.z * TERRAIN_ZSCALE;
-			OutUnpackedHeight[WindowCoord] = RelativeHeight + InLandscapeLocationZ;
+			const float RelativeHeight = (UnpackHeight(InPackedGrassMaps[InputCoords].xy) - LANDSCAPE_MID_VALUE) * InLandscapeScaleZ * TERRAIN_ZSCALE;
+			OutUnpackedHeight[OutputCoord] = RelativeHeight + InLandscapeRootZ;
 		}
 	}
 	else
 	{
 		// Grass map was not exported, write 0. Happens if for example user specifies a grass type name that isn't present on any of the landscape components.
-		OutUnpackedGrassMaps[uint3(WindowCoord, OutputSlice)] = 0.0f;
+		OutUnpackedGrassMaps[uint3(OutputCoord, OutputSlice)] = 0.0f;
 	}
 }
diff --git a/Engine/Plugins/PCG/Source/PCG/Private/Compute/PCGComputeCommon.cpp b/Engine/Plugins/PCG/Source/PCG/Private/Compute/PCGComputeCommon.cpp
--- a/Engine/Plugins/PCG/Source/PCG/Private/Compute/PCGComputeCommon.cpp
+++ b/Engine/Plugins/PCG/Source/PCG/Private/Compute/PCGComputeCommon.cpp
@@ -267,7 +267,14 @@ namespace PCGComputeHelpers
 
 	FString GetPrefixedDataLabel(const FString& InLabel)
 	{
-		return FString::Format(TEXT("{0}:{1}"), { PCGComputeConstants::DataLabelTagPrefix, InLabel });
+		FString PrefixedLabel;
+		GetPrefixedDataLabel(InLabel, PrefixedLabel);
+		return PrefixedLabel;
+	}
+
+	void GetPrefixedDataLabel(const FString& InLabel, FString& OutPrefixedLabel)
+	{
+		OutPrefixedLabel = FString::Format(TEXT("{0}:{1}"), { PCGComputeConstants::DataLabelTagPrefix, InLabel });
 	}
 
 	FString GetDataLabelResolverName(FName InPinLabel)
diff --git a/Engine/Plugins/PCG/Source/PCG/Private/Elements/PCGGenerateLandscapeTextures.cpp b/Engine/Plugins/PCG/Source/PCG/Private/Elements/PCGGenerateLandscapeTextures.cpp
--- a/Engine/Plugins/PCG/Source/PCG/Private/Elements/PCGGenerateLandscapeTextures.cpp
+++ b/Engine/Plugins/PCG/Source/PCG/Private/Elements/PCGGenerateLandscapeTextures.cpp
@@ -40,11 +40,13 @@ namespace PCGGenerateLandscapeTextures
 
 	const FText LandscapeComponentLostError = LOCTEXT("LandscapeComponentLost", "Reference to one or more landscape components lost, grass maps will not be generated.");
 
+#if !UE_BUILD_SHIPPING
 	static int32 GTriggerGPUCaptureDispatches = 0;
 	static FAutoConsoleVariableRef CVarTriggerGPUCapture(
 		TEXT("pcg.GPU.TriggerRenderCaptures.GrassMapGeneration"),
 		GTriggerGPUCaptureDispatches,
 		TEXT("Trigger GPU captures for this many of the subsequent grass generations."));
+#endif // !UE_BUILD_SHIPPING
 
 #if WITH_EDITOR
 	TAutoConsoleVariable<bool> CVarDebugDrawGeneratedComponents(
@@ -61,6 +63,340 @@ namespace PCGGenerateLandscapeTextures
 #endif // WITH_EDITOR
 			!InTexture->HasPendingInitOrStreaming() && InTexture->IsFullyStreamedIn();
 	}
+
+	struct FGenerateComponentsData
+	{
+		FPCGGenerateLandscapeTexturesContext* Context = nullptr;
+		TArray<ULandscapeComponent*> LandscapeComponents;
+		const UPCGGenerateLandscapeTexturesSettings* Settings = nullptr;
+		ALandscape* Landscape = nullptr;
+		TArray<FIntVector2> LandscapeTileCoords;
+		TArray<FIntPoint> LandscapeComponentKeys;
+		FBox LandscapeComponentBounds = FBox(EForceInit::ForceInit);
+		bool bLastGenerationPass = false;
+	};
+
+	bool GenerateComponents(FGenerateComponentsData&& InData)
+	{
+		TRACE_CPUPROFILER_EVENT_SCOPE(PCGGenerateLandscapeTextures::GenerateComponents);
+		if (!InData.Context || !InData.Settings || !InData.Landscape || !InData.Context->OutputTextureData)
+		{
+			return false;
+		}
+
+		const bool bGenerateGrassMaps = !InData.Context->GrassTypesToExport.IsEmpty();
+		const bool bGenerateHeightMap = InData.Settings->bGenerateHeightMap;
+
+		UE::Landscape::Grass::EGrassWeightExporterFlags Flags = UE::Landscape::Grass::EGrassWeightExporterFlags::None;
+		if (bGenerateGrassMaps)
+		{
+			Flags |= UE::Landscape::Grass::EGrassWeightExporterFlags::NeedsGrassmap;
+		}
+		if (bGenerateHeightMap)
+		{
+			Flags |= UE::Landscape::Grass::EGrassWeightExporterFlags::NeedsHeightmap;
+		}
+
+		TSharedPtr<FLandscapeGrassWeightExporter> Exporter = MakeShared<FLandscapeGrassWeightExporter>(
+			InData.Landscape,
+			InData.LandscapeComponents,
+			Flags,
+			InData.Context->GrassTypesToExport,
+			/*InRequestedHeightMips = */TConstArrayView<int32>{});
+
+		const FIntPoint& OutputTextureSize = Exporter->GetTargetSize();
+		const int32 Max2DTextureDimension = GetMax2DTextureDimension();
+
+		if (OutputTextureSize.X < 1 || OutputTextureSize.X > Max2DTextureDimension || OutputTextureSize.Y < 1 || OutputTextureSize.Y > Max2DTextureDimension)
+		{
+			PCGLog::LogErrorOnGraph(FText::Format(LOCTEXT("InvalidTextureSize", "Invalid output texture size ({0}, {1}), cannot exceed ({2}, {2}). Consider partitioning the component onto a smaller grid size."), OutputTextureSize.X, OutputTextureSize.Y, Max2DTextureDimension), InData.Context);
+			return false;
+		}
+
+		const int32 ComponentSizeQuads = InData.Landscape->ComponentSizeQuads;
+		const double QuadSizeWorld = InData.Context->LandscapeComponentExtent / ComponentSizeQuads;
+		// Each corner of a quad in the landscape corresponds to one texel in the grass map. 
+		const uint32 LandscapeComponentResolution = ComponentSizeQuads + 1;
+
+		const FBox& TotalBounds = InData.LandscapeComponentBounds;
+		const FBox& GenerationBounds = InData.Context->GenerationBounds;
+
+		UE_LOG(LogPCG, Verbose, TEXT("PCGUnpackGrassMap: %s, Cell bounds (%.2f, %.2f -> %.2f, %.2f), Total bounds (%.2f, %.2f -> %.2f, %.2f), efficiency %.2f"),
+			*InData.Context->GetExecutionSourceName(),
+			GenerationBounds.Min.X, GenerationBounds.Min.Y, GenerationBounds.Max.X, GenerationBounds.Max.Y,
+			TotalBounds.Min.X, TotalBounds.Min.Y, TotalBounds.Max.X, TotalBounds.Max.Y,
+			(GenerationBounds.Max.X - GenerationBounds.Min.X) * (GenerationBounds.Max.Y - GenerationBounds.Min.Y) / ((TotalBounds.Max.X - TotalBounds.Min.X) * (TotalBounds.Max.Y - TotalBounds.Min.Y)));
+
+		FIntPoint TileGridStart, TileGridEnd;
+		TileGridStart.X = FMath::Floor((TotalBounds.Min.X - GenerationBounds.Min.X) / QuadSizeWorld);
+		TileGridStart.Y = FMath::Floor((TotalBounds.Min.Y - GenerationBounds.Min.Y) / QuadSizeWorld);
+		TileGridEnd.X = FMath::CeilToInt((TotalBounds.Max.X - GenerationBounds.Min.X) / QuadSizeWorld);
+		TileGridEnd.Y = FMath::CeilToInt((TotalBounds.Max.Y - GenerationBounds.Min.Y) / QuadSizeWorld);
+
+		const FVector TotalLandscapeComponentExtent = InData.LandscapeComponentBounds.GetExtent() * 2;
+		const uint32 NumTilesX = FMath::RoundToInt(TotalLandscapeComponentExtent.X / InData.Context->LandscapeComponentExtent);
+		const uint32 NumTilesY = FMath::RoundToInt(TotalLandscapeComponentExtent.Y / InData.Context->LandscapeComponentExtent);
+
+		const FIntPoint FullExportResolution((LandscapeComponentResolution - 1u) * NumTilesX + 1u, (LandscapeComponentResolution - 1u) * NumTilesY + 1u);
+		UE_LOG(LogPCG, Verbose, TEXT("PCGUnpackGrassMap: %s, Export resolution %s, Window %s -> %s, Window proportion %.2f"),
+			*InData.Context->GetExecutionSourceName(),
+			*FullExportResolution.ToString(),
+			*TileGridStart.ToString(),
+			*TileGridEnd.ToString(),
+			(InData.Context->OutputResolution.X * InData.Context->OutputResolution.Y) / double(FullExportResolution.X * FullExportResolution.Y));
+
+		const FIntPoint DispatchMin(FMath::Max(0, TileGridStart.X), FMath::Max(0, TileGridStart.Y));
+		const FIntPoint DispatchMax(FMath::Min(InData.Context->OutputResolution.X - 1, TileGridEnd.X), FMath::Min(InData.Context->OutputResolution.Y - 1, TileGridEnd.Y));
+		const FIntPoint DispatchSize(DispatchMax.X - DispatchMin.X + 1, DispatchMax.Y - DispatchMin.Y + 1);
+		if (DispatchSize.X <= 0 || DispatchSize.Y <= 0)
+		{
+			UE_LOG(LogPCG, Error, TEXT("PCGGenerateLandscapeTextures: Landscape components for export do not overlap generation volume, should not happen."));
+			return true;
+		}
+
+		UE::RenderCommandPipe::FSyncScope SyncScope;
+
+#if !UE_BUILD_SHIPPING
+		RenderCaptureInterface::FScopedCapture RenderCapture(PCGGenerateLandscapeTextures::GTriggerGPUCaptureDispatches > 0, TEXT("PCGLandscapeGrassmapCapture"));
+		PCGGenerateLandscapeTextures::GTriggerGPUCaptureDispatches = FMath::Max(PCGGenerateLandscapeTextures::GTriggerGPUCaptureDispatches - 1, 0);
+#endif  // !UE_BUILD_SHIPPING
+
+		ENQUEUE_RENDER_COMMAND(PCGGenerateLandscapeTextures)(
+			[ContextHandle = InData.Context->GetOrCreateHandle()
+			, LandscapeComponentExtent = InData.Context->LandscapeComponentExtent
+			, LandscapeTileCoords = MoveTemp(InData.LandscapeTileCoords)
+			, LandscapeComponentKeys = MoveTemp(InData.LandscapeComponentKeys)
+			, OutputIndexToExportIndex = InData.Context->OutputIndexToExportIndex
+			, LandscapeComponentResolution
+			, LandscapeRootZ = InData.Context->LandscapeRootZ
+			, LandscapeScaleZ = InData.Context->LandscapeScaleZ
+			, bGenerateGrassMaps
+			, bGenerateHeightMap
+			, NumTilesX
+			, NumTilesY
+			, OutputResolution = InData.Context->OutputResolution
+			, TileGridStart
+			, DispatchMin
+			, DispatchSize
+			, OutputTextureData = InData.Context->OutputTextureData
+			, bLastGenerationPass = InData.bLastGenerationPass
+			, Exporter](FRHICommandListImmediate& RHICmdList)
+		{
+			TRACE_CPUPROFILER_EVENT_SCOPE(PCGGenerateLandscapeTextures);
+			LLM_SCOPE_BYTAG(PCG);
+
+			FPCGContext::FSharedContext<FPCGGenerateLandscapeTexturesContext> SharedContext(ContextHandle);
+			FPCGGenerateLandscapeTexturesContext* Context = SharedContext.Get();
+			if (!Context)
+			{
+				return;
+			}
+
+			FRDGBuilder GraphBuilder(RHICmdList);
+
+			FRDGTextureDesc GrassMapTextureDesc = FRDGTextureDesc::Create2D(
+				Exporter->GetTargetSize(),
+				PF_B8G8R8A8,
+				FClearValueBinding(),
+				ETextureCreateFlags::RenderTargetable | ETextureCreateFlags::ShaderResource);
+
+			FRDGTextureRef GrassMapTexture = GraphBuilder.CreateTexture(GrassMapTextureDesc, TEXT("PCGLandscapeGrassMapRenderTarget"));
+			FRDGTextureSRVRef GrassMapTextureSRV = GraphBuilder.CreateSRV(GrassMapTexture);
+
+			// Generate grass maps. All will be generated to a single texture.
+			Exporter->RenderLandscapeComponentToTexture_RenderThread(GraphBuilder, GrassMapTexture);
+
+			// Must have at least one slice to be a valid resource.
+			const uint32 NumOutputSlices = FMath::Max(Context->GrassTypesToOutput.Num(), 1);
+
+			FRDGTextureRef HeightMap = nullptr;
+			bool bHeightMapRequiresExport = false;
+			if (OutputTextureData->HeightMapHandle)
+			{
+				HeightMap = GraphBuilder.RegisterExternalTexture(OutputTextureData->HeightMapHandle, TEXT("PCGHeightMapExternal"));
+			}
+			else
+			{
+				FRDGTextureDesc HeightMapDesc = FRDGTextureDesc::Create2D(
+					bGenerateHeightMap ? OutputResolution : FIntPoint(1, 1), // No default system texture usable for UAV, so always create but with 1x1 resolution
+					PF_R32_FLOAT,
+					FClearValueBinding(),
+					ETextureCreateFlags::ShaderResource | ETextureCreateFlags::UAV);
+
+				HeightMap = GraphBuilder.CreateTexture(HeightMapDesc, TEXT("PCGLandscapeHeightMapUnpacked"));
+
+				bHeightMapRequiresExport = true;
+			}
+
+
+			// Output texture is array of textures, one per grass map.
+			FRDGTextureRef GrassMap = nullptr;
+			bool bGrassMapRequiresExport = false;
+			if (OutputTextureData->GrassMapHandle)
+			{
+				GrassMap = GraphBuilder.RegisterExternalTexture(OutputTextureData->GrassMapHandle, TEXT("PCGGrassMapExternal"));
+			}
+			else
+			{
+				FRDGTextureDesc GrassMapDesc = FRDGTextureDesc::Create2DArray(
+					bGenerateGrassMaps ? OutputResolution : FIntPoint(1, 1), // No default system texture usable for UAV, so always create but with 1x1 resolution
+					PF_G8,
+					FClearValueBinding(),
+					ETextureCreateFlags::ShaderResource | ETextureCreateFlags::UAV,
+					NumOutputSlices);
+
+				GrassMap = GraphBuilder.CreateTexture(GrassMapDesc, TEXT("PCGLandscapeGrassMapUnpacked"));
+
+				bGrassMapRequiresExport = true;
+			}
+
+			// Unpack the generated results to simple world aligned textures.
+			FPCGGrassMapUnpackerCS::FParameters* Parameters = GraphBuilder.AllocParameters<FPCGGrassMapUnpackerCS::FParameters>();
+			Parameters->InPackedGrassMaps = GrassMapTextureSRV;
+			Parameters->OutUnpackedGrassMaps = GraphBuilder.CreateUAV(GrassMap);
+			Parameters->OutUnpackedHeight = GraphBuilder.CreateUAV(HeightMap);
+			Parameters->InNumTiles = FUintVector2(NumTilesX, NumTilesY);
+			Parameters->InLandscapeRootZ = static_cast<float>(LandscapeRootZ);
+			Parameters->InLandscapeScaleZ = static_cast<float>(LandscapeScaleZ);
+			Parameters->InLandscapeComponentResolution = LandscapeComponentResolution;
+			Parameters->InOutputHeight = bGenerateHeightMap ? 1u : 0u; // No equivalent flag for grass maps, instead ComponentGrassmapOffsetAndChannelsBuffer will have invalid indices
+			Parameters->InTileGridOrigin = TileGridStart;
+			Parameters->InDispatchOffset = FUintVector2(DispatchMin.X, DispatchMin.Y);
+			Parameters->InDispatchSize = FUintVector2(DispatchSize.X, DispatchSize.Y);
+			Parameters->InNumExportedGrassTypes = Context->GrassTypesToExport.Num();
+
+			// Initialize to invalid component indices.
+			FMemory::Memset(&Parameters->InLinearTileIndexToComponentIndex[0], -1, FPCGGrassMapUnpackerCS::MaxNumLandscapeComponents * sizeof(int32) * 4);
+
+			// Now write component mapping.
+			check(LandscapeTileCoords.Num() <= FPCGGrassMapUnpackerCS::MaxNumLandscapeComponents);
+			for (int Index = 0; Index < LandscapeTileCoords.Num(); ++Index)
+			{
+				const int32 LinearTileIndex = LandscapeTileCoords[Index].Y * NumTilesX + LandscapeTileCoords[Index].X;
+				check((LinearTileIndex >= 0) && (LinearTileIndex < FPCGGrassMapUnpackerCS::MaxNumLandscapeComponents));
+				Parameters->InLinearTileIndexToComponentIndex[LinearTileIndex].X = Index;
+			}
+			check(LandscapeTileCoords.Num() == LandscapeComponentKeys.Num());
+
+			FRDGBufferRef ComponentHeightOffsetsBuffer;
+			if (bGenerateHeightMap)
+			{
+				// X offset for where to find each component's height data in the grass texture (always in the .xy channels)
+				TArray<int32> ComponentHeightOffsets;
+				ComponentHeightOffsets.Reserve(LandscapeComponentKeys.Num());
+				for (int ComponentIndex = 0; ComponentIndex < LandscapeComponentKeys.Num(); ++ComponentIndex)
+				{
+					TValueOrError<FIntRect /*OutputRect*/, FLandscapeGrassWeightExporter_RenderThread::EHeightmapTextureInfoError> TextureInfo =
+						Exporter->GetTextureInfoForHeight(LandscapeComponentKeys[ComponentIndex]);
+					// Not supposed to fail : 
+					checkf(!TextureInfo.HasError(), TEXT("Failed to retrieve height map info on component %s : reason : %i"), *LandscapeComponentKeys[ComponentIndex].ToString(), TextureInfo.GetError());
+					ComponentHeightOffsets.Add(TextureInfo.GetValue().Min.X);
+				}
+				check(ComponentHeightOffsets.Num() == LandscapeComponentKeys.Num());
+				ComponentHeightOffsetsBuffer = CreateStructuredBuffer(GraphBuilder, TEXT("PCGGenerateLandscapeTexturesComponentHeightOffsets"), ComponentHeightOffsets);
+			}
+			else
+			{
+				ComponentHeightOffsetsBuffer = GSystemTextures.GetDefaultStructuredBuffer(GraphBuilder, sizeof(int32));
+			}
+			Parameters->InComponentHeightOffsets = GraphBuilder.CreateSRV(ComponentHeightOffsetsBuffer);
+
+			FRDGBufferRef ComponentGrassmapOffsetAndChannelsBuffer;
+			if (!Context->GrassTypesToExport.IsEmpty())
+			{
+				// For each component, for each grass type, X offset for where to find the grass data in the grass texture (.x) and on which channel (.y)
+				//  It's possible a grass map is absent from a given component, in which case, the offset and channel will be set to -1
+				TArray<FIntVector2> ComponentGrassmapOffsetAndChannels;
+				ComponentGrassmapOffsetAndChannels.SetNumUninitialized(LandscapeComponentKeys.Num() * Context->GrassTypesToExport.Num());
+				FMemory::Memset(ComponentGrassmapOffsetAndChannels.GetData(), -1, ComponentGrassmapOffsetAndChannels.Num() * sizeof(FIntVector2));
+
+				for (int32 ComponentIndex = 0; ComponentIndex < LandscapeComponentKeys.Num(); ++ComponentIndex)
+				{
+					for (int32 GrassmapIndex = 0; GrassmapIndex < Context->GrassTypesToExport.Num(); ++GrassmapIndex)
+					{
+						FIntVector2& ComponentGrassmapOffsetAndChannel = ComponentGrassmapOffsetAndChannels[ComponentIndex * Context->GrassTypesToExport.Num() + GrassmapIndex];
+
+						TValueOrError<TPair<FIntRect /*OutputRect*/, uint8 /*ChannelIndex*/>, FLandscapeGrassWeightExporter_RenderThread::EGrassmapTextureInfoError> TextureInfo =
+							Exporter->GetTextureInfoForGrass(LandscapeComponentKeys[ComponentIndex], Context->GrassTypesToExport[GrassmapIndex]);
+						if (TextureInfo.HasError())
+						{
+							// It can fail, but only if that's because the grass map is absent from this component
+							checkf(TextureInfo.GetError() == FLandscapeGrassWeightExporter_RenderThread::EGrassmapTextureInfoError::InvalidGrassType,
+								TEXT("Failed to retrieve grass map info on component %s : reason : %i"), *LandscapeComponentKeys[ComponentIndex].ToString(), TextureInfo.GetError());
+						}
+						else
+						{
+							TPair<FIntRect /*OutputRect*/, uint8 /*ChannelIndex*/> TextureInfoValue = TextureInfo.GetValue();
+							check((TextureInfoValue.Value >= 0) && (TextureInfoValue.Value < 4));
+							ComponentGrassmapOffsetAndChannel = FIntVector2(TextureInfoValue.Key.Min.X, static_cast<int32>(TextureInfoValue.Value));
+						}
+					}
+				}
+				check(ComponentGrassmapOffsetAndChannels.Num() == LandscapeComponentKeys.Num() * Context->GrassTypesToExport.Num());
+				ComponentGrassmapOffsetAndChannelsBuffer = CreateStructuredBuffer(GraphBuilder, TEXT("PCGGenerateLandscapeTexturesComponentGrassmapOffsetAndChannels"), ComponentGrassmapOffsetAndChannels);
+			}
+			else
+			{
+				ComponentGrassmapOffsetAndChannelsBuffer = GSystemTextures.GetDefaultStructuredBuffer(GraphBuilder, sizeof(FIntVector2));;
+			}
+			Parameters->InComponentGrassmapOffsetAndChannels = GraphBuilder.CreateSRV(ComponentGrassmapOffsetAndChannelsBuffer);
+
+			FRDGBufferRef OutputIndexToExportIndexBuffer;
+			if (!OutputIndexToExportIndex.IsEmpty())
+			{
+				OutputIndexToExportIndexBuffer = CreateStructuredBuffer(GraphBuilder, TEXT("PCGGenerateLandscapeTexturesOutputIndexToExportIndex"), OutputIndexToExportIndex);
+			}
+			else
+			{
+				OutputIndexToExportIndexBuffer = GSystemTextures.GetDefaultStructuredBuffer(GraphBuilder, sizeof(int32));
+			}
+			Parameters->InOutputIndexToExportIndex = GraphBuilder.CreateSRV(OutputIndexToExportIndexBuffer);
+
+			TShaderMapRef<FPCGGrassMapUnpackerCS> Shader(GetGlobalShaderMap(GMaxRHIFeatureLevel));
+			const int GroupCountX = FMath::DivideAndRoundUp<int>(DispatchSize.X, FPCGGrassMapUnpackerCS::ThreadGroupDim);
+			const int GroupCountY = FMath::DivideAndRoundUp<int>(DispatchSize.Y, FPCGGrassMapUnpackerCS::ThreadGroupDim);
+			const int GroupCountZ = NumOutputSlices;
+
+			FComputeShaderUtils::AddPass(GraphBuilder, RDG_EVENT_NAME("PCGUnpackGrassMap"), ERDGPassFlags::Compute, Shader, Parameters, FIntVector(GroupCountX, GroupCountY, GroupCountZ));
+
+			UE_LOG(LogPCG, Verbose, TEXT("PCGUnpackGrassMap: %s, Components %d, OutputResolution %s, OutputSlices %u"),
+				*Context->GetExecutionSourceName(), LandscapeTileCoords.Num(), *OutputResolution.ToString(), NumOutputSlices);
+
+			// Export the output textures so they can be used downstream.
+			if (bHeightMapRequiresExport)
+			{
+				OutputTextureData->HeightMapHandle = GraphBuilder.ConvertToExternalTexture(HeightMap);
+			}
+
+			if (bGrassMapRequiresExport)
+			{
+				OutputTextureData->GrassMapHandle = GraphBuilder.ConvertToExternalTexture(GrassMap);
+			}
+
+			GraphBuilder.SetTextureAccessFinal(HeightMap, ERHIAccess::SRVCompute);
+			GraphBuilder.SetTextureAccessFinal(GrassMap, ERHIAccess::SRVCompute);
+
+			GraphBuilder.Execute();
+
+			if (bLastGenerationPass)
+			{
+				ExecuteOnGameThread(UE_SOURCE_LOCATION, [ContextHandle, OutputTextureData]()
+				{
+					LLM_SCOPE_BYTAG(PCG);
+
+					FPCGContext::FSharedContext<FPCGGenerateLandscapeTexturesContext> SharedContext(ContextHandle);
+					if (FPCGGenerateLandscapeTexturesContext* Context = SharedContext.Get())
+					{
+						check(OutputTextureData);
+						OutputTextureData->bRenderingComplete = true;
+						Context->bIsPaused = false;
+					}
+				});
+			}
+		});
+
+		return true;
+	}
 }
 
 UPCGGenerateLandscapeTexturesSettings::UPCGGenerateLandscapeTexturesSettings()
@@ -158,13 +494,6 @@ EPCGChangeType UPCGGenerateLandscapeTexturesSettings::GetChangeTypeForProperty(c
 }
 #endif // WITH_EDITOR
 
-PRAGMA_DISABLE_DEPRECATION_WARNINGS
-FPCGGenerateLandscapeTexturesContext::~FPCGGenerateLandscapeTexturesContext()
-{
-	delete LandscapeGrassWeightExporter;
-}
-PRAGMA_ENABLE_DEPRECATION_WARNINGS
-
 void FPCGGenerateLandscapeTexturesContext::AddExtraStructReferencedObjects(FReferenceCollector& Collector)
 {
 	if (HeightTextureData)
@@ -293,7 +622,7 @@ bool FPCGGenerateLandscapeTexturesElement::ExecuteInternal(FPCGContext* InContex
 			FMemory::Memset(Context->OutputIndexToExportIndex.GetData(), -1, Context->OutputIndexToExportIndex.Num() * Context->OutputIndexToExportIndex.GetTypeSize());
 		}
 
-		const FBox GenerationBounds = ExecutionSource->GetExecutionState().GetBounds();
+		Context->GenerationBounds = ExecutionSource->GetExecutionState().GetBounds();
 		TArray<ULandscapeComponent*> OverlappingComponents;
 		TSet<TObjectPtr<ULandscapeGrassType>> ProcessedGrassTypes;
 
@@ -334,7 +663,7 @@ bool FPCGGenerateLandscapeTexturesElement::ExecuteInternal(FPCGContext* InContex
 				{
 					// Only generate grass map if there is meaningful overlap with our domain of interest.
 					const FBox LandscapeComponentBounds = LandscapeComponent->Bounds.GetBox();
-					if (LandscapeComponentBounds.Overlap(GenerationBounds).GetVolume() > UE_KINDA_SMALL_NUMBER)
+					if (LandscapeComponentBounds.Overlap(Context->GenerationBounds).GetVolume() > UE_KINDA_SMALL_NUMBER)
 					{
 						if (OverlappingComponents.Num() >= FPCGGrassMapUnpackerCS::MaxNumLandscapeComponents)
 						{
@@ -387,10 +716,10 @@ bool FPCGGenerateLandscapeTexturesElement::ExecuteInternal(FPCGContext* InContex
 						Context->LandscapeComponents.Add(LandscapeComponent);
 						OverlappingComponents.Add(LandscapeComponent);
 
+						Context->TotalLandscapeComponentBounds += LandscapeComponentBounds;
+
 						if (Context->LandscapeComponents.Num() == 1)
 						{
-							Context->TotalLandscapeComponentBounds = LandscapeComponentBounds;
-
 							Context->LandscapeComponentExtent = LandscapeComponent->Bounds.BoxExtent.X * 2.0;
 							Context->TexelSizeWorld = Context->LandscapeComponentExtent / static_cast<double>(LandscapeProxy->ComponentSizeQuads);
 
@@ -406,8 +735,6 @@ bool FPCGGenerateLandscapeTexturesElement::ExecuteInternal(FPCGContext* InContex
 						}
 						else
 						{
-							Context->TotalLandscapeComponentBounds += LandscapeComponentBounds;
-
 							// Currently assuming all landscape have the same extents.
 							ensure(FMath::IsNearlyEqual(Context->LandscapeComponentExtent, LandscapeComponent->Bounds.BoxExtent.X * 2.0));
 						}
@@ -428,11 +755,37 @@ bool FPCGGenerateLandscapeTexturesElement::ExecuteInternal(FPCGContext* InContex
 			}
 		}
 
-		if (!Context->Landscape.Get() || Context->LandscapeComponents.IsEmpty() || (Context->GrassTypesToOutput.IsEmpty() && !Settings->bGenerateHeightMap))
+		if (!Context->Landscape.Get() || Context->LandscapeComponents.IsEmpty() || (Context->GrassTypesToOutput.IsEmpty() && !Settings->bGenerateHeightMap) || !Context->TotalLandscapeComponentBounds.IsValid)
 		{
 			return true;
 		}
 
+		// Precompute some quantities that are used in each export pass.
+		{
+			ALandscape* Landscape = Context->Landscape.Get();
+			check(Landscape);
+
+			Context->LandscapeRootZ = Landscape->GetRootComponent()->GetRelativeLocation().Z;
+			Context->LandscapeScaleZ = Landscape->GetRootComponent()->GetRelativeScale3D().Z;
+
+			const double QuadSizeWorld = Context->LandscapeComponentExtent / Landscape->ComponentSizeQuads;
+
+			// Dilate the horizontal extents by half a texel in both X and Y. This will place the texel centers on the vert grid.
+			FBox OutputBounds;
+			OutputBounds.Min.X = (FMath::Floor(Context->GenerationBounds.Min.X / QuadSizeWorld) - 0.5) * QuadSizeWorld;
+			OutputBounds.Min.Y = (FMath::Floor(Context->GenerationBounds.Min.Y / QuadSizeWorld) - 0.5) * QuadSizeWorld;
+			OutputBounds.Max.X = (FMath::CeilToInt(Context->GenerationBounds.Max.X / QuadSizeWorld) + 0.5) * QuadSizeWorld;
+			OutputBounds.Max.Y = (FMath::CeilToInt(Context->GenerationBounds.Max.Y / QuadSizeWorld) + 0.5) * QuadSizeWorld;
+			OutputBounds.Min.Z = Context->TotalLandscapeComponentBounds.Min.Z;
+			OutputBounds.Max.Z = Context->TotalLandscapeComponentBounds.Max.Z;
+
+			Context->OutputTextureTransform = FTransform(FQuat::Identity, OutputBounds.GetCenter(), OutputBounds.GetExtent());
+
+			Context->OutputResolution = FIntPoint(
+				FMath::CeilToInt((Context->GenerationBounds.Max.X - Context->GenerationBounds.Min.X) / QuadSizeWorld) + 1,
+				FMath::CeilToInt((Context->GenerationBounds.Max.Y - Context->GenerationBounds.Min.Y) / QuadSizeWorld) + 1);
+		}
+
 		Context->bLandscapeComponentsFiltered = true;
 	}
 
@@ -540,331 +893,114 @@ bool FPCGGenerateLandscapeTexturesElement::ExecuteInternal(FPCGContext* InContex
 		ALandscape* Landscape = Context->Landscape.IsValid() ? Context->Landscape.Get() : nullptr;
 		if (!Landscape)
 		{
-			PCGLog::LogErrorOnGraph(LOCTEXT("LandscapeProxyLost", "Reference to landscape proxy actor lost, grass maps will not be generated."), InContext);
+			PCGLog::LogErrorOnGraph(LOCTEXT("LandscapeLost", "Reference to landscape actor lost, grass maps will not be generated."), InContext);
 			return true;
 		}
 
-		const FVector LandscapeGridScale = Landscape->GetRootComponent()->GetRelativeScale3D();
-		const double LandscapeZ = Landscape->GetRootComponent()->GetRelativeLocation().Z;
-
 		TArray<ULandscapeComponent*> LandscapeComponents;
-		TArray<FIntVector2> LandscapeTileCoords;
-		LandscapeComponents.Reserve(Context->LandscapeComponents.Num());
-		// Sister array that stores the global identifier for a landscape component so that we can query info about this component in the grass map exporter
+		LandscapeComponents.Reserve(FMath::Min(Context->LandscapeComponents.Num(), Settings->NumLandscapeComponentsGeneratedPerFrame));
+		
+		// Stores the global identifier for a landscape component so that we can query info about this component in the grass map exporter
 		TArray<FIntPoint> LandscapeComponentKeys;
+		LandscapeComponentKeys.Reserve(FMath::Min(Context->LandscapeComponents.Num(), Settings->NumLandscapeComponentsGeneratedPerFrame));
+
+		FBox LandscapeComponentBounds = FBox(EForceInit::ForceInit);
 		LandscapeComponentKeys.Reserve(Context->LandscapeComponents.Num());
 
-		for (TWeakObjectPtr<ULandscapeComponent> LandscapeComponentPtr : Context->LandscapeComponents)
-		{
-			ULandscapeComponent* LandscapeComponent = LandscapeComponentPtr.IsValid() ? LandscapeComponentPtr.Get() : nullptr;
+		int32 NumComponentsToGenerate = 0;
 
-			if (!LandscapeComponent)
+		for (int32 Index = Context->NumGenerationPassesScheduled; Index < Context->LandscapeComponents.Num(); ++Index)
+		{
+			if (!Context->LandscapeComponents[Index].IsValid())
 			{
 				PCGLog::LogErrorOnGraph(PCGGenerateLandscapeTextures::LandscapeComponentLostError, InContext);
 				return true;
 			}
 
+			ULandscapeComponent* LandscapeComponent = Context->LandscapeComponents[Index].Get();
 			LandscapeComponents.Add(LandscapeComponent);
-
-			// Landscape components are not ordered, so store a 2d index of each component within the grass map.
-			// TODO could likely sort the LandscapeComponents array so the order is known and no indices need looking up.
-			LandscapeTileCoords.Add(FIntVector2(
-				(LandscapeComponent->Bounds.Origin.X - Context->TotalLandscapeComponentBounds.Min.X) / Context->LandscapeComponentExtent,
-				(LandscapeComponent->Bounds.Origin.Y - Context->TotalLandscapeComponentBounds.Min.Y) / Context->LandscapeComponentExtent));
 			LandscapeComponentKeys.Add(LandscapeComponent->GetComponentKey());
-		}
+			LandscapeComponentBounds += LandscapeComponent->Bounds.GetBox();
 
-		UE::Landscape::Grass::EGrassWeightExporterFlags Flags = UE::Landscape::Grass::EGrassWeightExporterFlags::None;
-		if (!Context->GrassTypesToExport.IsEmpty())
-		{
-			Flags |= UE::Landscape::Grass::EGrassWeightExporterFlags::NeedsGrassmap;
-		}
-		if (Settings->bGenerateHeightMap)
-		{
-			Flags |= UE::Landscape::Grass::EGrassWeightExporterFlags::NeedsHeightmap;
+			++NumComponentsToGenerate;
+			if (NumComponentsToGenerate >= Settings->NumLandscapeComponentsGeneratedPerFrame)
+			{
+				break;
+			}
 		}
 
-		Context->LandscapeGrassWeightExporter = new FLandscapeGrassWeightExporter(
-			Landscape, 
-			LandscapeComponents,
-			Flags,
-			Context->GrassTypesToExport,
-			/*InRequestedHeightMips = */{});
-
-		const FIntPoint& OutputTextureSize = Context->LandscapeGrassWeightExporter->GetTargetSize();
-		const int32 Max2DTextureDimension = GetMax2DTextureDimension();
-
-		if (OutputTextureSize.X < 1 || OutputTextureSize.X > Max2DTextureDimension || OutputTextureSize.Y < 1 || OutputTextureSize.Y > Max2DTextureDimension)
+		if (!LandscapeComponentBounds.IsValid)
 		{
-			PCGLog::LogErrorOnGraph(FText::Format(LOCTEXT("InvalidTextureSize", "Invalid output texture size ({0}, {1}), cannot exceed ({2}, {2}). Consider partitioning the component onto a smaller grid size."), OutputTextureSize.X, OutputTextureSize.Y, Max2DTextureDimension), InContext);
+			PCGLog::LogErrorOnGraph(LOCTEXT("LandscapeComonentBoundsInvalid", "Landscape component encountered with invalid bounds, grass maps will not be generated."), InContext);
 			return true;
 		}
 
-		const int32 ComponentSizeQuads = Landscape->ComponentSizeQuads;
-		const double QuadSizeWorld = Context->LandscapeComponentExtent / ComponentSizeQuads;
-		// Each corner of a quad in the landscape corresponds to one texel in the grass map. 
-		const uint32 LandscapeComponentResolution = ComponentSizeQuads + 1;
-		// Internal boundaries between components in the output share a row/column, so these are dropped to get a correct world space output.
-		const FVector TotalLandscapeComponentExtent = Context->TotalLandscapeComponentBounds.GetExtent() * 2;
-		const uint32 NumTilesX = FMath::RoundToInt(TotalLandscapeComponentExtent.X / Context->LandscapeComponentExtent);
-		const uint32 NumTilesY = FMath::RoundToInt(TotalLandscapeComponentExtent.Y / Context->LandscapeComponentExtent);
-		const FIntPoint FullExportResolution((LandscapeComponentResolution - 1u)* NumTilesX + 1u, (LandscapeComponentResolution - 1u)* NumTilesY + 1u);
-
-		FIntPoint WindowStart, WindowEnd;
-
-		const FBox CellBounds = Context->ExecutionSource->GetExecutionState().GetBounds();
-		const FBox& TotalBounds = Context->TotalLandscapeComponentBounds;
-
-		UE_LOG(LogPCG, Verbose, TEXT("PCGUnpackGrassMap: %s, Cell bounds (%.2f, %.2f -> %.2f, %.2f), Total bounds (%.2f, %.2f -> %.2f, %.2f), efficiency %.2f"),
-			*Context->GetExecutionSourceName(),
-			CellBounds.Min.X, CellBounds.Min.Y, CellBounds.Max.X, CellBounds.Max.Y,
-			TotalBounds.Min.X, TotalBounds.Min.Y, TotalBounds.Max.X, TotalBounds.Max.Y,
-			(CellBounds.Max.X - CellBounds.Min.X)* (CellBounds.Max.Y - CellBounds.Min.Y) / ((TotalBounds.Max.X - TotalBounds.Min.X) * (TotalBounds.Max.Y - TotalBounds.Min.Y)));
-
-		WindowStart.X = FMath::Clamp(FMath::Floor((CellBounds.Min.X - TotalBounds.Min.X) / QuadSizeWorld), 0, FullExportResolution.X - 1);
-		WindowStart.Y = FMath::Clamp(FMath::Floor((CellBounds.Min.Y - TotalBounds.Min.Y) / QuadSizeWorld), 0, FullExportResolution.Y - 1);
-		WindowEnd.X = FMath::Clamp(FMath::CeilToInt((CellBounds.Max.X - TotalBounds.Min.X) / QuadSizeWorld), WindowStart.X, FullExportResolution.X - 1);
-		WindowEnd.Y = FMath::Clamp(FMath::CeilToInt((CellBounds.Max.Y - TotalBounds.Min.Y) / QuadSizeWorld), WindowStart.Y, FullExportResolution.Y - 1);
+		TArray<FIntVector2> LandscapeTileCoords;
+		LandscapeComponents.Reserve(LandscapeComponents.Num());
 
-		const FIntPoint OutputResolution(WindowEnd.X - WindowStart.X + 1, WindowEnd.Y - WindowStart.Y + 1);
+		for (ULandscapeComponent* LandscapeComponent : LandscapeComponents)
+		{
+			check(LandscapeComponent);
 
-		UE_LOG(LogPCG, Verbose, TEXT("PCGUnpackGrassMap: %s, Export resolution %s, Window %s -> %s, Window proportion %.2f"),
-			*Context->GetExecutionSourceName(),
-			*FullExportResolution.ToString(),
-			*WindowStart.ToString(),
-			*WindowEnd.ToString(),
-			(OutputResolution.X* OutputResolution.Y) / double(FullExportResolution.X * FullExportResolution.Y));
-
-		FBox OutputBounds;
-		OutputBounds.Min.X = Context->TotalLandscapeComponentBounds.Min.X + WindowStart.X * QuadSizeWorld;
-		OutputBounds.Min.Y = Context->TotalLandscapeComponentBounds.Min.Y + WindowStart.Y * QuadSizeWorld;
-		OutputBounds.Max.X = OutputBounds.Min.X + (OutputResolution.X - 1) * QuadSizeWorld; // -1 because resolution is verts, not quads (so discard last col/row)
-		OutputBounds.Max.Y = OutputBounds.Min.Y + (OutputResolution.Y - 1) * QuadSizeWorld;
-		OutputBounds.Min.Z = Context->TotalLandscapeComponentBounds.Min.Z;
-		OutputBounds.Max.Z = Context->TotalLandscapeComponentBounds.Max.Z;
-
-		// Dilate the horizontal extents by half a texel in both X and Y. This will place the texel centers on the vert grid.
-		OutputBounds.Min.X -= QuadSizeWorld / 2.0;
-		OutputBounds.Min.Y -= QuadSizeWorld / 2.0;
-		OutputBounds.Max.X += QuadSizeWorld / 2.0;
-		OutputBounds.Max.Y += QuadSizeWorld / 2.0;
-
-		Context->OutputTextureTransform = FTransform(FQuat::Identity, OutputBounds.GetCenter(), OutputBounds.GetExtent());
+			// Landscape components are not ordered, so store a 2d index of each component within the grass map.
+			// TODO could likely sort the LandscapeComponents array so the order is known and no indices need looking up.
+			LandscapeTileCoords.Add(FIntVector2(
+				(LandscapeComponent->Bounds.Origin.X - LandscapeComponentBounds.Min.X) / Context->LandscapeComponentExtent,
+				(LandscapeComponent->Bounds.Origin.Y - LandscapeComponentBounds.Min.Y) / Context->LandscapeComponentExtent));
+		}
 
-		UE::RenderCommandPipe::FSyncScope SyncScope;
+		if (!Context->OutputTextureData)
+		{
+			Context->OutputTextureData = MakeShared<FPCGGenerateLandscapeTexturesContext::FOutputTextureData>();
+		}
 
-		RenderCaptureInterface::FScopedCapture RenderCapture(PCGGenerateLandscapeTextures::GTriggerGPUCaptureDispatches > 0, TEXT("PCGLandscapeGrassmapCapture"));
-		PCGGenerateLandscapeTextures::GTriggerGPUCaptureDispatches = FMath::Max(PCGGenerateLandscapeTextures::GTriggerGPUCaptureDispatches - 1, 0);
+		const bool bLastGenerationPass = (Context->NumGenerationPassesScheduled + NumComponentsToGenerate) >= Context->LandscapeComponents.Num();
 		
-		TWeakPtr<FPCGContextHandle> ContextHandle = Context->GetOrCreateHandle();
-
-		ENQUEUE_RENDER_COMMAND(PCGGenerateGrassMaps)(
-			[ ContextHandle
-			, LandscapeComponentExtent=Context->LandscapeComponentExtent
-			, LandscapeTileCoords = MoveTemp(LandscapeTileCoords)
-			, LandscapeComponentKeys = MoveTemp(LandscapeComponentKeys)
-			, OutputIndexToExportIndex = MoveTemp(Context->OutputIndexToExportIndex)
-			, LandscapeComponentResolution
-			, LandscapeGridScale
-			, LandscapeZ
-			, bGenerateHeightMap = Settings->bGenerateHeightMap
-			, NumTilesX
-			, NumTilesY
-			, OutputResolution
-			, WindowStart
-			, WindowEnd](FRHICommandListImmediate& RHICmdList)
+		PCGGenerateLandscapeTextures::FGenerateComponentsData Data
 		{
-			LLM_SCOPE_BYTAG(PCG);
-
-			FPCGContext::FSharedContext<FPCGGenerateLandscapeTexturesContext> SharedContext(ContextHandle);
-			FPCGGenerateLandscapeTexturesContext* Context = SharedContext.Get();
-			if (!Context)
-			{
-				return;
-			}
-
-			FRDGBuilder GraphBuilder(RHICmdList);
-
-			FRDGTextureDesc GrassMapTextureDesc = FRDGTextureDesc::Create2D(
-				Context->LandscapeGrassWeightExporter->GetTargetSize(),
-				PF_B8G8R8A8,
-				FClearValueBinding(),
-				ETextureCreateFlags::RenderTargetable | ETextureCreateFlags::ShaderResource);
-
-			FRDGTextureRef GrassMapTexture = GraphBuilder.CreateTexture(GrassMapTextureDesc, TEXT("PCGLandscapeGrassMapRenderTarget"));
-			FRDGTextureSRVRef GrassMapTextureSRV = GraphBuilder.CreateSRV(GrassMapTexture);
-
-			// Generate grass maps. All will be generated to a single texture.
-			Context->LandscapeGrassWeightExporter->RenderLandscapeComponentToTexture_RenderThread(GraphBuilder, GrassMapTexture);
-
-			// Must have at least one slice to be a valid resource.
-			const uint32 NumOutputSlices = FMath::Max(Context->GrassTypesToOutput.Num(), 1);
-
-			FRDGTextureDesc GrassHeightMapDesc = FRDGTextureDesc::Create2D(
-				bGenerateHeightMap ? OutputResolution : FIntPoint(1, 1), // No default system texture usable for UAV, so always create but with 1x1 resolution
-				PF_R32_FLOAT,
-				FClearValueBinding(),
-				ETextureCreateFlags::ShaderResource | ETextureCreateFlags::UAV);
-
-			FRDGTextureRef GrassHeightMap = GraphBuilder.CreateTexture(GrassHeightMapDesc, TEXT("PCGLandscapeHeightMapUnpacked"));
-			FRDGTextureUAVRef GrassHeightMapUAV = GraphBuilder.CreateUAV(GrassHeightMap);
-
-			// Output texture is array of textures, one per grass map.
-			FRDGTextureDesc GrassMapDesc = FRDGTextureDesc::Create2DArray(
-				OutputResolution,
-				PF_G8,
-				FClearValueBinding(),
-				ETextureCreateFlags::ShaderResource | ETextureCreateFlags::UAV,
-				NumOutputSlices);
-
-			FRDGTextureRef GrassMap = GraphBuilder.CreateTexture(GrassMapDesc, TEXT("PCGLandscapeGrassMapUnpacked"));
-			FRDGTextureUAVRef GrassMapUAV = GraphBuilder.CreateUAV(GrassMap);
-
-			// Unpack the generated results to simple world aligned textures.
-			FPCGGrassMapUnpackerCS::FParameters* Parameters = GraphBuilder.AllocParameters<FPCGGrassMapUnpackerCS::FParameters>();
-			Parameters->InPackedGrassMaps = GrassMapTextureSRV;
-			Parameters->OutUnpackedGrassMaps = GrassMapUAV;
-			Parameters->OutUnpackedHeight = GrassHeightMapUAV;
-			Parameters->InNumTiles = FUintVector2(NumTilesX, NumTilesY);
-			Parameters->InLandscapeGridScale = static_cast<FVector3f>(LandscapeGridScale);
-			Parameters->InLandscapeLocationZ = static_cast<float>(LandscapeZ);
-			Parameters->InLandscapeComponentResolution = LandscapeComponentResolution;
-			Parameters->InOutputHeight = bGenerateHeightMap ? 1u : 0u;
-			Parameters->InWindow = FUintVector4(WindowStart.X, WindowStart.Y, WindowEnd.X, WindowEnd.Y);
-			Parameters->InNumExportedGrassTypes = Context->GrassTypesToExport.Num();
-
-			// Initialize to invalid component indices.
-			for (int Index = 0; Index < FPCGGrassMapUnpackerCS::MaxNumLandscapeComponents; ++Index)
-			{
-				Parameters->InLinearTileIndexToComponentIndex[Index].X = -1;
-			}
-
-			// Now write component mapping.
-			check(LandscapeTileCoords.Num() <= FPCGGrassMapUnpackerCS::MaxNumLandscapeComponents);
-			for (int Index = 0; Index < LandscapeTileCoords.Num(); ++Index)
-			{
-				const int32 LinearTileIndex = LandscapeTileCoords[Index].Y * NumTilesX + LandscapeTileCoords[Index].X;
-				check((LinearTileIndex >= 0) && (LinearTileIndex < FPCGGrassMapUnpackerCS::MaxNumLandscapeComponents));
-				Parameters->InLinearTileIndexToComponentIndex[LinearTileIndex].X = Index;
-			}
-			check(LandscapeTileCoords.Num() == LandscapeComponentKeys.Num());
-
-			FRDGBufferRef ComponentHeightOffsetsBuffer;
-			if (bGenerateHeightMap)
+			.Context = Context,
+			.LandscapeComponents = MoveTemp(LandscapeComponents),
+			.Settings = Settings,
+			.Landscape = Landscape,
+			.LandscapeTileCoords = MoveTemp(LandscapeTileCoords),
+			.LandscapeComponentKeys = MoveTemp(LandscapeComponentKeys),
+			.LandscapeComponentBounds = LandscapeComponentBounds,
+			.bLastGenerationPass = bLastGenerationPass,
+		};
+
+		Context->NumGenerationPassesScheduled += NumComponentsToGenerate;
+
+		if (PCGGenerateLandscapeTextures::GenerateComponents(MoveTemp(Data)))
+		{
+			if (bLastGenerationPass)
 			{
-				// X offset for where to find each component's height data in the grass texture (always in the .xy channels)
-				TArray<int32> ComponentHeightOffsets;
-				ComponentHeightOffsets.Reserve(LandscapeComponentKeys.Num());
-				for (int ComponentIndex = 0; ComponentIndex < LandscapeComponentKeys.Num(); ++ComponentIndex)
-				{
-					TValueOrError<FIntRect /*OutputRect*/, FLandscapeGrassWeightExporter_RenderThread::EHeightmapTextureInfoError> TextureInfo = 
-						Context->LandscapeGrassWeightExporter->GetTextureInfoForHeight(LandscapeComponentKeys[ComponentIndex]);
-					// Not supposed to fail : 
-					checkf(!TextureInfo.HasError(), TEXT("Failed to retrieve height map info on component %s : reason : %i"), *LandscapeComponentKeys[ComponentIndex].ToString(), TextureInfo.GetError()); 
-					ComponentHeightOffsets.Add(TextureInfo.GetValue().Min.X);
-				}
-				check(ComponentHeightOffsets.Num() == LandscapeComponentKeys.Num());
-				ComponentHeightOffsetsBuffer = CreateStructuredBuffer(GraphBuilder, TEXT("PCGGenerateLandscapeTexturesComponentHeightOffsets"), ComponentHeightOffsets);
+				Context->bGenerationScheduled = true;
+				// Render command will wake this task up after completing.
+				Context->bIsPaused = true;
 			}
 			else
 			{
-				ComponentHeightOffsetsBuffer = GSystemTextures.GetDefaultStructuredBuffer(GraphBuilder, sizeof(int32));
-			}
-			Parameters->InComponentHeightOffsets = GraphBuilder.CreateSRV(ComponentHeightOffsetsBuffer);
-
-			FRDGBufferRef ComponentGrassmapOffsetAndChannelsBuffer;
-			if (!Context->GrassTypesToExport.IsEmpty())
-			{
-				// For each component, for each grass type, X offset for where to find the grass data in the grass texture (.x) and on which channel (.y)
-				//  It's possible a grass map is absent from a given component, in which case, the offset and channel will be set to -1
-				TArray<FIntVector2> ComponentGrassmapOffsetAndChannels;
-				ComponentGrassmapOffsetAndChannels.Reserve(LandscapeComponentKeys.Num() * Context->GrassTypesToExport.Num());
-				for (int32 ComponentIndex = 0; ComponentIndex < LandscapeComponentKeys.Num(); ++ComponentIndex)
+				// Sleep until next frame to timeslice generation passes across frames.
+				Context->bIsPaused = true;
+				FPCGModule::GetPCGModuleChecked().ExecuteNextTick([ContextHandle = Context->GetOrCreateHandle()]()
 				{
-					for (int32 GrassmapIndex = 0; GrassmapIndex < Context->GrassTypesToExport.Num(); ++GrassmapIndex)
+					if (TSharedPtr<FPCGContextHandle> SharedHandle = ContextHandle.Pin())
 					{
-						FIntVector2& ComponentGrassmapOffsetAndChannel = ComponentGrassmapOffsetAndChannels.Add_GetRef(FIntVector2(-1, -1));
-
-						TValueOrError<TPair<FIntRect /*OutputRect*/, uint8 /*ChannelIndex*/>, FLandscapeGrassWeightExporter_RenderThread::EGrassmapTextureInfoError> TextureInfo =
-							Context->LandscapeGrassWeightExporter->GetTextureInfoForGrass(LandscapeComponentKeys[ComponentIndex], Context->GrassTypesToExport[GrassmapIndex]);
-						if (TextureInfo.HasError())
+						if (FPCGContext* ContextPtr = SharedHandle->GetContext())
 						{
-							// It can fail, but only if that's because the grass map is absent from this component
-							checkf(TextureInfo.GetError() == FLandscapeGrassWeightExporter_RenderThread::EGrassmapTextureInfoError::InvalidGrassType,
-								TEXT("Failed to retrieve grass map info on component %s : reason : %i"), *LandscapeComponentKeys[ComponentIndex].ToString(), TextureInfo.GetError());
-						}
-						else
-						{
-							TPair<FIntRect /*OutputRect*/, uint8 /*ChannelIndex*/> TextureInfoValue = TextureInfo.GetValue();
-							check((TextureInfoValue.Value >= 0) && (TextureInfoValue.Value < 4));
-							ComponentGrassmapOffsetAndChannel =	FIntVector2(TextureInfoValue.Key.Min.X, static_cast<int32>(TextureInfoValue.Value));
+							ContextPtr->bIsPaused = false;
 						}
 					}
-				}
-				check(ComponentGrassmapOffsetAndChannels.Num() == LandscapeComponentKeys.Num() * Context->GrassTypesToExport.Num());
-				ComponentGrassmapOffsetAndChannelsBuffer = CreateStructuredBuffer(GraphBuilder, TEXT("PCGGenerateLandscapeTexturesComponentGrassmapOffsetAndChannels"), ComponentGrassmapOffsetAndChannels);
+				});
 			}
-			else
-			{
-				ComponentGrassmapOffsetAndChannelsBuffer = GSystemTextures.GetDefaultStructuredBuffer(GraphBuilder, sizeof(FIntVector2));;
-			}
-			Parameters->InComponentGrassmapOffsetAndChannels = GraphBuilder.CreateSRV(ComponentGrassmapOffsetAndChannelsBuffer);
-
-			FRDGBufferRef OutputIndexToExportIndexBuffer;
-			if (!OutputIndexToExportIndex.IsEmpty())
-			{
-				OutputIndexToExportIndexBuffer = CreateStructuredBuffer(GraphBuilder, TEXT("PCGGenerateLandscapeTexturesOutputIndexToExportIndex"), OutputIndexToExportIndex);
-			}
-			else
-			{
-				OutputIndexToExportIndexBuffer = GSystemTextures.GetDefaultStructuredBuffer(GraphBuilder, sizeof(int32));
-			}
-			Parameters->InOutputIndexToExportIndex = GraphBuilder.CreateSRV(OutputIndexToExportIndexBuffer);
 
-			TShaderMapRef<FPCGGrassMapUnpackerCS> Shader(GetGlobalShaderMap(GMaxRHIFeatureLevel));
-			const int GroupCountX = FMath::DivideAndRoundUp<int>(OutputResolution.X, FPCGGrassMapUnpackerCS::ThreadGroupDim);
-			const int GroupCountY = FMath::DivideAndRoundUp<int>(OutputResolution.Y, FPCGGrassMapUnpackerCS::ThreadGroupDim);
-			const int GroupCountZ = NumOutputSlices;
-
-			FComputeShaderUtils::AddPass(GraphBuilder, RDG_EVENT_NAME("PCGUnpackGrassMap"), ERDGPassFlags::Compute, Shader, Parameters, FIntVector(GroupCountX, GroupCountY, GroupCountZ));
-
-			UE_LOG(LogPCG, Verbose, TEXT("PCGUnpackGrassMap: %s, Components %d, OutputResolution %s, OutputSlices %u"),
-				*Context->GetExecutionSourceName(), LandscapeTileCoords.Num(), *OutputResolution.ToString(), NumOutputSlices);
-
-			// Export the output textures so they can be used downstream.
-			TRefCountPtr<IPooledRenderTarget> HeightsExported;
-			if (bGenerateHeightMap)
-			{
-				HeightsExported = GraphBuilder.ConvertToExternalTexture(GrassHeightMap);
-				GraphBuilder.SetTextureAccessFinal(GrassHeightMap, ERHIAccess::SRVCompute);
-			}
-
-			TRefCountPtr<IPooledRenderTarget> GrassMapExported = GraphBuilder.ConvertToExternalTexture(GrassMap);
-			GraphBuilder.SetTextureAccessFinal(GrassMap, ERHIAccess::SRVCompute);
-
-			GraphBuilder.Execute();
-
-			// Pass exported buffer back to game thread and wake up this element.
-			ExecuteOnGameThread(UE_SOURCE_LOCATION, [ContextHandle, HeightsExported, GrassMapExported]()
-			{
-				LLM_SCOPE_BYTAG(PCG);
-
-				FPCGContext::FSharedContext<FPCGGenerateLandscapeTexturesContext> SharedContext(ContextHandle);
-				if (FPCGGenerateLandscapeTexturesContext* Context = SharedContext.Get())
-				{
-					Context->HeightHandle = HeightsExported;
-					Context->GrassMapHandle = GrassMapExported;
-					Context->bIsPaused = false;
-				}
-			});
-		});
-
-		Context->bGenerationScheduled = true;
-
-		// Render command will wake this task up after completing.
-		Context->bIsPaused = true;
-
-		return false;
+			return false;
+		}
+		else
+		{
+			// Failed, terminate element.
+			return true;
+		}
 	}
 
 	// 4. Initialize texture data objects.
@@ -892,27 +1028,27 @@ bool FPCGGenerateLandscapeTexturesElement::ExecuteInternal(FPCGContext* InContex
 		Context->bCreatedOutputGrassTextures = true;
 	}
 
-	if (!Context->GrassMapHandle.IsValid() && !Context->HeightHandle.IsValid())
+	if (!Context->OutputTextureData || (!Context->OutputTextureData->GrassMapHandle.IsValid() && !Context->OutputTextureData->HeightMapHandle.IsValid()) || !ensure(Context->OutputTextureData->bRenderingComplete))
 	{
 		return true;
 	}
 
 	bool bAllTexturesInitialized = true;
 
-	if (Context->HeightHandle)
+	if (Context->OutputTextureData->HeightMapHandle)
 	{
 		check(Context->HeightTextureData);
-		bAllTexturesInitialized &= Context->HeightTextureData->Initialize(Context->HeightHandle, /*TextureIndex=*/0, Context->OutputTextureTransform, Settings->bSkipReadbackToCPU);
+		bAllTexturesInitialized &= Context->HeightTextureData->Initialize(Context->OutputTextureData->HeightMapHandle, /*TextureIndex=*/0, Context->OutputTextureTransform, Settings->bSkipReadbackToCPU);
 	}
 
-	if (Context->GrassMapHandle)
+	if (Context->OutputTextureData->GrassMapHandle)
 	{
 		if (Settings->bOutputTextureArray)
 		{
 			FPCGTexture2DArrayDataInitParams InitParams;
 			InitParams.Transform = Context->OutputTextureTransform;
 
-			bAllTexturesInitialized &= Context->GrassMapTextureArrayData->Initialize(Context->GrassMapHandle, InitParams);
+			bAllTexturesInitialized &= Context->GrassMapTextureArrayData->Initialize(Context->OutputTextureData->GrassMapHandle, InitParams);
 		}
 		else
 		{
@@ -928,7 +1064,7 @@ bool FPCGGenerateLandscapeTexturesElement::ExecuteInternal(FPCGContext* InContex
 				{
 					// Poll initialize (fine to be called even when initialization was already complete).
 					bAllTexturesInitialized &= TextureData->Initialize(
-						Context->GrassMapHandle,
+						Context->OutputTextureData->GrassMapHandle,
 						/*TextureIndex=*/DataIndex,
 						Context->OutputTextureTransform,
 						Settings->bSkipReadbackToCPU);
@@ -1001,6 +1137,8 @@ bool FPCGGenerateLandscapeTexturesElement::ExecuteInternal(FPCGContext* InContex
 	}
 	else
 	{
+		FString TagString; // Used in loop but hoisted for efficiency.
+
 		for (int DataIndex = 0; DataIndex < Context->GrassTypesToOutput.Num(); ++DataIndex)
 		{
 			check(Context->GrassMapTextureDatas.IsValidIndex(DataIndex));
@@ -1026,12 +1164,14 @@ bool FPCGGenerateLandscapeTexturesElement::ExecuteInternal(FPCGContext* InContex
 
 			if (!GrassTypeName.IsEmpty())
 			{
-				OutTaggedData.Tags.Add(PCGComputeHelpers::GetPrefixedDataLabel(GrassTypeName));
+				PCGComputeHelpers::GetPrefixedDataLabel(GrassTypeName, TagString);
+				OutTaggedData.Tags.Add(TagString);
 			}
 			else if (!Settings->bExcludeSelectedGrassTypes && Settings->SelectedGrassTypes.IsValidIndex(DataIndex))
 			{
 				// We are using inclusion list, look up name
-				OutTaggedData.Tags.Add(PCGComputeHelpers::GetPrefixedDataLabel(Settings->SelectedGrassTypes[DataIndex]));
+				PCGComputeHelpers::GetPrefixedDataLabel(Settings->SelectedGrassTypes[DataIndex], TagString);
+				OutTaggedData.Tags.Add(TagString);
 			}
 			else
 			{
diff --git a/Engine/Plugins/PCG/Source/PCG/Public/Compute/PCGComputeCommon.h b/Engine/Plugins/PCG/Source/PCG/Public/Compute/PCGComputeCommon.h
--- a/Engine/Plugins/PCG/Source/PCG/Public/Compute/PCGComputeCommon.h
+++ b/Engine/Plugins/PCG/Source/PCG/Public/Compute/PCGComputeCommon.h
@@ -141,6 +141,7 @@ namespace PCGComputeHelpers
 
 	/** Produces the data label prefixed with PCGComputeConstants::DataLabelTagPrefix. */
 	FString GetPrefixedDataLabel(const FString& InLabel);
+	void GetPrefixedDataLabel(const FString& InLabel, FString& OutPrefixedLabel);
 
 	/** Produces the data interface name of a data label resolver. */
 	FString GetDataLabelResolverName(FName InPinLabel);
diff --git a/Engine/Plugins/PCG/Source/PCG/Public/Elements/PCGGenerateLandscapeTextures.h b/Engine/Plugins/PCG/Source/PCG/Public/Elements/PCGGenerateLandscapeTextures.h
--- a/Engine/Plugins/PCG/Source/PCG/Public/Elements/PCGGenerateLandscapeTextures.h
+++ b/Engine/Plugins/PCG/Source/PCG/Public/Elements/PCGGenerateLandscapeTextures.h
@@ -82,15 +82,14 @@ public:
 	/** Output a single texture 2D array data rather than a texture 2D data for each slice. */
 	UPROPERTY(BlueprintReadWrite, EditAnywhere, Category = Settings)
 	bool bOutputTextureArray = false;
+
+	UPROPERTY(BlueprintReadWrite, EditAnywhere, Category = Settings, meta=(ClampMin=1))
+	int32 NumLandscapeComponentsGeneratedPerFrame = 1;
 };
 
 struct FPCGGenerateLandscapeTexturesContext : public FPCGContext
 {
 public:
-	~FPCGGenerateLandscapeTexturesContext();
-
-	FLandscapeGrassWeightExporter* LandscapeGrassWeightExporter = nullptr;
-
 	/** Output height texture data. */
 	TObjectPtr<UPCGTextureData> HeightTextureData;
 
@@ -99,9 +98,16 @@ public:
 	TObjectPtr<UPCGTexture2DArrayData> GrassMapTextureArrayData;
 	bool bCreatedOutputGrassTextures = false;
 
+	struct FOutputTextureData
+	{
+		// Initialized and consumed on render thread during render thread, then read on game thread after all rendering is complete.
+		TRefCountPtr<IPooledRenderTarget> HeightMapHandle;
+		TRefCountPtr<IPooledRenderTarget> GrassMapHandle;
+		bool bRenderingComplete = false;
+	};
+
 	/** Exported result textures. */
-	TRefCountPtr<IPooledRenderTarget> HeightHandle = nullptr;
-	TRefCountPtr<IPooledRenderTarget> GrassMapHandle = nullptr;
+	TSharedPtr<FOutputTextureData, ESPMode::ThreadSafe> OutputTextureData;
 
 	/** The grass types passed to the grass exporter. Each of these grass types is present on at least one of the landscape components. */
 	TArray<TObjectPtr<ULandscapeGrassType>> GrassTypesToExport;
@@ -115,9 +121,14 @@ public:
 	TWeakObjectPtr<ALandscape> Landscape;
 	TArray<TWeakObjectPtr<ULandscapeComponent>> LandscapeComponents;
 
+	double LandscapeScaleZ = 1.0;
+	double LandscapeRootZ = 0.0;
+
 	/** World-space bounds containing all of the landscape components given to the grass weight exporter. */
 	FBox TotalLandscapeComponentBounds = FBox(EForceInit::ForceInit);
 
+	FBox GenerationBounds = FBox(EForceInit::ForceInit);
+
 	/** Extent (side length) of each landscape component. */
 	double LandscapeComponentExtent = 0.0;
 
@@ -136,10 +147,14 @@ public:
 	/** True when grass map generation has been scheduled on the render thread. */
 	bool bGenerationScheduled = false;
 
+	int32 NumGenerationPassesScheduled = 0;
+
 	double TexelSizeWorld = 0.0;
 
 	FTransform OutputTextureTransform;
 
+	FIntPoint OutputResolution = FIntPoint::ZeroValue;
+
 	TObjectPtr<UTexture> BlackTexture;
 
 protected:
@@ -155,6 +170,12 @@ public:
 	TArray<TTuple<TWeakObjectPtr<ULandscapeGrassType>, int32>> SelectedGrassTypes;
 	UE_DEPRECATED(5.8, "No longer used")
 	TArray<TObjectPtr<ULandscapeGrassType>> GlobalGrassTypes;
+	UE_DEPRECATED(5.8, "No longer used")
+	FLandscapeGrassWeightExporter* LandscapeGrassWeightExporter = nullptr;
+	UE_DEPRECATED(5.8, "Replaced with FOutputTextureData")
+	TRefCountPtr<IPooledRenderTarget> HeightHandle = nullptr;
+	UE_DEPRECATED(5.8, "Replaced with FOutputTextureData")
+	TRefCountPtr<IPooledRenderTarget> GrassMapHandle = nullptr;
 };
 
 class FPCGGenerateLandscapeTexturesElement : public IPCGElement
diff --git a/Engine/Plugins/PCG/Source/PCGCompute/Internal/PCGGrassMapUnpackerCS.h b/Engine/Plugins/PCG/Source/PCGCompute/Internal/PCGGrassMapUnpackerCS.h
--- a/Engine/Plugins/PCG/Source/PCGCompute/Internal/PCGGrassMapUnpackerCS.h
+++ b/Engine/Plugins/PCG/Source/PCGCompute/Internal/PCGGrassMapUnpackerCS.h
@@ -27,12 +27,14 @@ public:
 		SHADER_PARAMETER_RDG_TEXTURE_UAV(RWTexture2DArray<float>, OutUnpackedGrassMaps)
 		SHADER_PARAMETER_ARRAY(FIntVector4, InLinearTileIndexToComponentIndex, [MaxNumLandscapeComponents])
 		SHADER_PARAMETER(FUintVector2, InNumTiles)
+		SHADER_PARAMETER(FIntPoint, InTileGridOrigin)
 		SHADER_PARAMETER(uint32, InLandscapeComponentResolution)
 		SHADER_PARAMETER(int32, InNumExportedGrassTypes)
-		SHADER_PARAMETER(FVector3f, InLandscapeGridScale)
-		SHADER_PARAMETER(float, InLandscapeLocationZ)
+		SHADER_PARAMETER(float, InLandscapeScaleZ)
+		SHADER_PARAMETER(float, InLandscapeRootZ)
 		SHADER_PARAMETER(uint32, InOutputHeight)
-		SHADER_PARAMETER(FUintVector4, InWindow)
+		SHADER_PARAMETER(FUintVector2, InDispatchOffset)
+		SHADER_PARAMETER(FUintVector2, InDispatchSize)
 	END_SHADER_PARAMETER_STRUCT()
 
 public:
diff --git a/Engine/Source/Runtime/Landscape/Internal/LandscapeGrassWeightExporter.h b/Engine/Source/Runtime/Landscape/Internal/LandscapeGrassWeightExporter.h
--- a/Engine/Source/Runtime/Landscape/Internal/LandscapeGrassWeightExporter.h
+++ b/Engine/Source/Runtime/Landscape/Internal/LandscapeGrassWeightExporter.h
@@ -61,7 +61,7 @@ private:
 		TConstArrayView<int32> InHeightMips, UE::Landscape::Grass::EGrassWeightExporterFlags InFlags);
 
 public:
-	virtual ~FLandscapeGrassWeightExporter_RenderThread();
+	LANDSCAPE_API virtual ~FLandscapeGrassWeightExporter_RenderThread();
 
 	const FIntPoint& GetTargetSize() const { return TargetSize; }
 
