[PCG] Editor: move toolbar extension to application mode

#Synced-CL 49926118

diff --git a/Engine/Plugins/Experimental/ProceduralVegetationEditor/Source/ProceduralVegetationEditor/Private/PVEditor.cpp b/Engine/Plugins/Experimental/ProceduralVegetationEditor/Source/ProceduralVegetationEditor/Private/PVEditor.cpp
--- a/Engine/Plugins/Experimental/ProceduralVegetationEditor/Source/ProceduralVegetationEditor/Private/PVEditor.cpp
+++ b/Engine/Plugins/Experimental/ProceduralVegetationEditor/Source/ProceduralVegetationEditor/Private/PVEditor.cpp
@@ -272,17 +272,15 @@ bool FPVEditor::IsPanelAvailable(const FName PanelID) const
 	return SupportedIDs.Contains(PanelID);
 }
 
-void FPVEditor::RegisterToolbarInternal(FToolMenuSection& PCGSection) const
+void FPVEditor::RegisterToolbarInternal(FToolBarBuilder& ToolbarBuilder) const
 {
-	RegisterToolbarButton(PCGSection, EPCGToolbarButtons::Find);
-	RegisterToolbarButton(PCGSection, EPCGToolbarButtons::ForceRegen);
+	RegisterToolbarButton(ToolbarBuilder, EPCGToolbarButtons::Find);
+	RegisterToolbarButton(ToolbarBuilder, EPCGToolbarButtons::ForceRegen);
 	
 	const FPVEditorCommands& Commands = FPVEditorCommands::Get();
-	PCGSection.AddEntry(FToolMenuEntry::InitToolBarButton(
-			Commands.Export,
-			TAttribute<FText>(),
-			TAttribute<FText>(),
-			FSlateIcon(FAppStyle::GetAppStyleSetName(), "Icons.Toolbar.Export")));
+
+	ToolbarBuilder.AddToolBarButton(Commands.Export, NAME_None, TAttribute<FText>(), TAttribute<FText>(),
+		FSlateIcon(FAppStyle::GetAppStyleSetName(), "Icons.Toolbar.Export"));
 }
 
 void FPVEditor::BindCommands()
diff --git a/Engine/Plugins/Experimental/ProceduralVegetationEditor/Source/ProceduralVegetationEditor/Private/PVEditor.h b/Engine/Plugins/Experimental/ProceduralVegetationEditor/Source/ProceduralVegetationEditor/Private/PVEditor.h
--- a/Engine/Plugins/Experimental/ProceduralVegetationEditor/Source/ProceduralVegetationEditor/Private/PVEditor.h
+++ b/Engine/Plugins/Experimental/ProceduralVegetationEditor/Source/ProceduralVegetationEditor/Private/PVEditor.h
@@ -41,7 +41,7 @@ protected:
 	virtual TAttribute<FGraphAppearanceInfo> GetAppearanceInfo() const override;
 	virtual TSharedRef<FTabManager::FLayout> GetDefaultLayout() const override;
 	virtual bool IsPanelAvailable(const FName PanelID) const override;
-	virtual void RegisterToolbarInternal(FToolMenuSection& PCGSection) const override;
+	virtual void RegisterToolbarInternal(FToolBarBuilder& ToolbarBuilder) const override;
 	virtual void BindCommands() override;
 	virtual void OnSelectedNodesChanged(const TSet<UObject*>& InNewSelection) override;
 	// ~End FPCGEditor interface
diff --git a/Engine/Plugins/PCG/Source/PCGEditor/Private/PCGEditor.cpp b/Engine/Plugins/PCG/Source/PCGEditor/Private/PCGEditor.cpp
--- a/Engine/Plugins/PCG/Source/PCGEditor/Private/PCGEditor.cpp
+++ b/Engine/Plugins/PCG/Source/PCGEditor/Private/PCGEditor.cpp
@@ -171,7 +171,6 @@ void FPCGEditor::Initialize(const EToolkitMode::Type InMode, const TSharedPtr<cl
 	GetOrCreateNodeSourceWidget();
 	
 	BindCommands();
-	RegisterToolbar();
 
 	PCGGraphBeingEdited->OnGraphChangedDelegate.AddRaw(this, &FPCGEditor::OnGraphChanged);
 	PCGGraphBeingEdited->OnNodeSourceCompiledDelegate.AddRaw(this, &FPCGEditor::OnNodeSourceCompiled);
@@ -662,35 +661,28 @@ FString FPCGEditor::GetWorldCentricTabPrefix() const
 	return LOCTEXT("WorldCentricTabPrefix", "PCG ").ToString();
 }
 
-void FPCGEditor::RegisterToolbar() const
+void FPCGEditor::RegisterToolbar(FToolBarBuilder& ToolbarBuilder) const
 {
-	UToolMenus* ToolMenus = UToolMenus::Get();
-	FName ParentName;
-	const FName ToolbarName = GetToolMenuToolbarName(ParentName);
-	if (!ToolMenus->IsMenuRegistered(ToolbarName))
-	{
-		UToolMenu* ToolBar = ToolMenus->RegisterMenu(ToolbarName, ParentName, EMultiBoxType::ToolBar);
-		FToolMenuSection& Section = ToolBar->AddSection("PCGToolbar", TAttribute<FText>());
-
-		RegisterToolbarInternal(Section);
-	}
+	ToolbarBuilder.BeginSection("PCGToolbar");
+	RegisterToolbarInternal(ToolbarBuilder);
+	ToolbarBuilder.EndSection();
 }
 
-void FPCGEditor::RegisterToolbarInternal(FToolMenuSection& PCGSection) const
+void FPCGEditor::RegisterToolbarInternal(FToolBarBuilder& ToolbarBuilder) const
 {
-	RegisterToolbarButton(PCGSection, EPCGToolbarButtons::Find);
-	RegisterToolbarButton(PCGSection, EPCGToolbarButtons::PauseRegen);
-	RegisterToolbarButton(PCGSection, EPCGToolbarButtons::ForceRegen);
-	RegisterToolbarButton(PCGSection, EPCGToolbarButtons::CancelExecution);
-	RegisterToolbarButton(PCGSection, EPCGToolbarButtons::OpenDebugObjectTreeTab);
+	RegisterToolbarButton(ToolbarBuilder, EPCGToolbarButtons::Find);
+	RegisterToolbarButton(ToolbarBuilder, EPCGToolbarButtons::PauseRegen);
+	RegisterToolbarButton(ToolbarBuilder, EPCGToolbarButtons::ForceRegen);
+	RegisterToolbarButton(ToolbarBuilder, EPCGToolbarButtons::CancelExecution);
+	RegisterToolbarButton(ToolbarBuilder, EPCGToolbarButtons::OpenDebugObjectTreeTab);
 
-	PCGSection.AddSeparator(NAME_None);
+	ToolbarBuilder.AddSeparator(NAME_None);
 
-	RegisterToolbarButton(PCGSection, EPCGToolbarButtons::GraphParams);
-	RegisterToolbarButton(PCGSection, EPCGToolbarButtons::GraphSettings);
+	RegisterToolbarButton(ToolbarBuilder, EPCGToolbarButtons::GraphParams);
+	RegisterToolbarButton(ToolbarBuilder, EPCGToolbarButtons::GraphSettings);
 }
 
-void FPCGEditor::RegisterToolbarButton(FToolMenuSection& Section, EPCGToolbarButtons Button) const
+void FPCGEditor::RegisterToolbarButton(FToolBarBuilder& ToolbarBuilder, EPCGToolbarButtons Button) const
 {
 	const FPCGEditorCommands& PCGEditorCommands = FPCGEditorCommands::Get();
 	
@@ -698,82 +690,55 @@ void FPCGEditor::RegisterToolbarButton(FToolMenuSection& Section, EPCGToolbarBut
 	{
 	case EPCGToolbarButtons::Find:
 		{
-			Section.AddEntry(FToolMenuEntry::InitToolBarButton(
-				PCGEditorCommands.Find,
-				TAttribute<FText>(),
-				TAttribute<FText>(),
-				FSlateIcon(FPCGEditorStyle::Get().GetStyleSetName(), "PCG.Command.Find")));
-
+			ToolbarBuilder.AddToolBarButton(PCGEditorCommands.Find, NAME_None, TAttribute<FText>(), TAttribute<FText>(),
+				FSlateIcon(FPCGEditorStyle::Get().GetStyleSetName(), "PCG.Command.Find"));
 			break;
 		}
 	case EPCGToolbarButtons::PauseRegen:
 		{
-			Section.AddEntry(FToolMenuEntry::InitToolBarButton(
-				FPCGEditorCommands::Get().PauseAutoRegeneration,
-				TAttribute<FText>(),
-				TAttribute<FText>(),
-				FSlateIcon(FPCGEditorStyle::Get().GetStyleSetName(), "PCG.Command.PauseRegen")));
-
+			ToolbarBuilder.AddToolBarButton(PCGEditorCommands.PauseAutoRegeneration, NAME_None, TAttribute<FText>(), TAttribute<FText>(),
+				FSlateIcon(FPCGEditorStyle::Get().GetStyleSetName(), "PCG.Command.PauseRegen"));
 			break;
 		}
 	case EPCGToolbarButtons::ForceRegen:
 		{
-			Section.AddEntry(FToolMenuEntry::InitToolBarButton(
-				FPCGEditorCommands::Get().ForceGraphRegeneration,
-				TAttribute<FText>(),
-				TAttribute<FText>(),
+			ToolbarBuilder.AddToolBarButton(PCGEditorCommands.ForceGraphRegeneration, NAME_None, TAttribute<FText>(), TAttribute<FText>(),
 				TAttribute<FSlateIcon>::CreateLambda([]()
 				{
 					static const FSlateIcon ForceRegen = FSlateIcon(FPCGEditorStyle::Get().GetStyleSetName(), "PCG.Command.ForceRegen");
 					static const FSlateIcon ForceRegenClearCache = FSlateIcon(FPCGEditorStyle::Get().GetStyleSetName(), "PCG.Command.ForceRegenClearCache");
-								
+
 					FModifierKeysState ModifierKeys = FSlateApplication::Get().GetModifierKeys();
 					return ModifierKeys.IsControlDown() ? ForceRegenClearCache : ForceRegen;
-				})));
-
+				}));
 			break;
 		}
 	case EPCGToolbarButtons::CancelExecution:
 		{
-			Section.AddEntry(FToolMenuEntry::InitToolBarButton(
-				FPCGEditorCommands::Get().CancelExecution,
-				TAttribute<FText>(),
-				TAttribute<FText>(),
-				FSlateIcon(FPCGEditorStyle::Get().GetStyleSetName(), "PCG.Command.StopRegen")));
-
+			ToolbarBuilder.AddToolBarButton(PCGEditorCommands.CancelExecution, NAME_None, TAttribute<FText>(), TAttribute<FText>(),
+				FSlateIcon(FPCGEditorStyle::Get().GetStyleSetName(), "PCG.Command.StopRegen"));
 			break;
 		}
 	case EPCGToolbarButtons::OpenDebugObjectTreeTab:
 		{
 			if (IsPanelAvailable(PCGEditorTabs::DebugObjectID))
 			{
-				Section.AddEntry(FToolMenuEntry::InitToolBarButton(
-					FPCGEditorCommands::Get().OpenDebugObjectTreeTab,
-					TAttribute<FText>(),
-					TAttribute<FText>(),
-					FSlateIcon(FPCGEditorStyle::Get().GetStyleSetName(), "PCG.Command.OpenDebugTreeTab")));
+				ToolbarBuilder.AddToolBarButton(PCGEditorCommands.OpenDebugObjectTreeTab, NAME_None, TAttribute<FText>(), TAttribute<FText>(),
+					FSlateIcon(FPCGEditorStyle::Get().GetStyleSetName(), "PCG.Command.OpenDebugTreeTab"));
 			}
 
 			break;
 		}
 	case EPCGToolbarButtons::GraphParams:
 		{
-			Section.AddEntry(FToolMenuEntry::InitToolBarButton(
-				 PCGEditorCommands.ToggleGraphParams,
-				 TAttribute<FText>(),
-				 TAttribute<FText>(),
-				 FSlateIcon(FPCGEditorStyle::Get().GetStyleSetName(), "PCG.Command.OpenGraphParams")));
-
+			ToolbarBuilder.AddToolBarButton(PCGEditorCommands.ToggleGraphParams, NAME_None, TAttribute<FText>(), TAttribute<FText>(),
+				FSlateIcon(FPCGEditorStyle::Get().GetStyleSetName(), "PCG.Command.OpenGraphParams"));
 			break;
 		}
 	case EPCGToolbarButtons::GraphSettings:
 		{
-			Section.AddEntry(FToolMenuEntry::InitToolBarButton(
-				 PCGEditorCommands.EditGraphSettings,
-				 TAttribute<FText>(),
-				 TAttribute<FText>(),
-				 FSlateIcon(FPCGEditorStyle::Get().GetStyleSetName(), "PCG.Command.GraphSettings")));
-
+			ToolbarBuilder.AddToolBarButton(PCGEditorCommands.EditGraphSettings, NAME_None, TAttribute<FText>(), TAttribute<FText>(),
+				FSlateIcon(FPCGEditorStyle::Get().GetStyleSetName(), "PCG.Command.GraphSettings"));
 			break;
 		}
 	}
@@ -905,10 +870,13 @@ void FPCGEditor::BindCommands()
 
 void FPCGEditor::OnFind()
 {
-	if (TabManager.IsValid() && FindWidget.IsValid())
+	if (TabManager.IsValid())
 	{
 		TabManager->TryInvokeTab(FTabId(PCGEditorTabs::FindID));
-		FindWidget->FocusForUse();
+		if (FindWidget.IsValid())
+		{
+			FindWidget->FocusForUse();
+		}
 	}
 }
 
@@ -921,7 +889,7 @@ void FPCGEditor::OpenDetailsView()
 			for (int DetailsViewIndex = 0; DetailsViewIndex < PCGEditorTabs::NumPropertyDetailTabs; ++DetailsViewIndex)
 			{
 				TSharedPtr<SPCGEditorGraphDetailsView> DetailsView = PropertyDetailsWidgets[DetailsViewIndex];
-				if (DetailsView.IsValid() && !DetailsView->IsLocked())
+				if (!DetailsView.IsValid() || !DetailsView->IsLocked())
 				{
 					if (!bVisibleOnly || TabManager->FindExistingLiveTab(FTabId(PCGEditorTabs::PropertyDetailsID[DetailsViewIndex])))
 					{
@@ -940,10 +908,7 @@ void FPCGEditor::OpenDetailsView()
 		}
 
 		// Default to first if they are all locked
-		if (PropertyDetailsWidgets[0].IsValid())
-		{
-			TabManager->TryInvokeTab(FTabId(PCGEditorTabs::PropertyDetailsID[0]));
-		}
+		TabManager->TryInvokeTab(FTabId(PCGEditorTabs::PropertyDetailsID[0]));
 	}
 }
 
diff --git a/Engine/Plugins/PCG/Source/PCGEditor/Private/PCGEditorDefaultMode.cpp b/Engine/Plugins/PCG/Source/PCGEditor/Private/PCGEditorDefaultMode.cpp
--- a/Engine/Plugins/PCG/Source/PCGEditor/Private/PCGEditorDefaultMode.cpp
+++ b/Engine/Plugins/PCG/Source/PCGEditor/Private/PCGEditorDefaultMode.cpp
@@ -136,6 +136,14 @@ FPCGEditorDefaultMode::FPCGEditorDefaultMode(TSharedPtr<FPCGEditor> InPCGEditor)
 
 		TabFactories.RegisterFactory(Info.CreateFactory.Execute());
 	}
+
+	ToolbarExtender = MakeShared<FExtender>();
+	ToolbarExtender->AddToolBarExtension(
+		"Asset",
+		EExtensionHook::After,
+		InPCGEditor->GetToolkitCommands(),
+		FToolBarExtensionDelegate::CreateSP(InPCGEditor.Get(), &FPCGEditor::RegisterToolbar)
+	);
 }
 
 FName FPCGEditorDefaultMode::StaticName()
diff --git a/Engine/Plugins/PCG/Source/PCGEditor/Public/PCGEditor.h b/Engine/Plugins/PCG/Source/PCGEditor/Public/PCGEditor.h
--- a/Engine/Plugins/PCG/Source/PCGEditor/Public/PCGEditor.h
+++ b/Engine/Plugins/PCG/Source/PCGEditor/Public/PCGEditor.h
@@ -252,11 +252,17 @@ PRAGMA_ENABLE_DEPRECATION_WARNINGS
 protected:
 	PCGEDITOR_API virtual TAttribute<FGraphAppearanceInfo> GetAppearanceInfo() const;
 
+	UE_DEPRECATED(5.8, "Implement version with FToolbarBuilder parameter instead")
+	PCGEDITOR_API virtual void RegisterToolbarInternal(FToolMenuSection& PCGSection) const {}
+
 	/** Register PCG specific toolbar for the editor */
-	PCGEDITOR_API virtual void RegisterToolbarInternal(FToolMenuSection& PCGSection) const;
+	PCGEDITOR_API virtual void RegisterToolbarInternal(FToolBarBuilder& ToolbarBuilder) const;
+
+	UE_DEPRECATED(5.8, "Use version with FToolBarBuilder parameter instead")
+	PCGEDITOR_API void RegisterToolbarButton(FToolMenuSection& Section, EPCGToolbarButtons Button) const {}
 
 	/** Register PCG default editor toolbar button */
-	PCGEDITOR_API void RegisterToolbarButton(FToolMenuSection& Section, EPCGToolbarButtons Button) const;
+	PCGEDITOR_API void RegisterToolbarButton(FToolBarBuilder& ToolbarBuilder, EPCGToolbarButtons Button) const;
 
 	/** Bind commands to delegates */
 	PCGEDITOR_API virtual void BindCommands();
@@ -305,7 +311,7 @@ protected:
 
 private:
 	/** Register PCG specific toolbar for the editor */
-	void RegisterToolbar() const;
+	void RegisterToolbar(FToolBarBuilder& InToolbarBuilder) const;
 	
 	/** Bring up the find tab */
 	void OnFind();
