Fix build

#rnx

#Synced-CL 49992548

diff --git a/Engine/Plugins/PCG/Source/PCGEditor/Public/PCGEditorTabFactories.h b/Engine/Plugins/PCG/Source/PCGEditor/Public/PCGEditorTabFactories.h
--- a/Engine/Plugins/PCG/Source/PCGEditor/Public/PCGEditorTabFactories.h
+++ b/Engine/Plugins/PCG/Source/PCGEditor/Public/PCGEditorTabFactories.h
@@ -37,7 +37,7 @@ struct FPCGGenericTabFactory : public FWorkflowTabFactory
 public:
 	PCGEDITOR_API FPCGGenericTabFactory(const FPCGGenericTabFactoryParams& InParams);
 
-	virtual TSharedRef<SWidget> CreateTabBody(const FWorkflowTabSpawnInfo& Info) const override;
+	PCGEDITOR_API virtual TSharedRef<SWidget> CreateTabBody(const FWorkflowTabSpawnInfo& Info) const override;
 
 private:
 	TFunction<TSharedRef<SWidget>()> CreateTabBodyFunc;
