[PCG] Fix cell that generation source is in getting lower priority than the next cell in front of camera which has a higher distance weight.

#rnx

#Synced-CL 49952186

diff --git a/Engine/Plugins/PCG/Source/PCG/Private/RuntimeGen/SchedulingPolicies/PCGSchedulingPolicyDistanceAndDirection.cpp b/Engine/Plugins/PCG/Source/PCG/Private/RuntimeGen/SchedulingPolicies/PCGSchedulingPolicyDistanceAndDirection.cpp
--- a/Engine/Plugins/PCG/Source/PCG/Private/RuntimeGen/SchedulingPolicies/PCGSchedulingPolicyDistanceAndDirection.cpp
+++ b/Engine/Plugins/PCG/Source/PCG/Private/RuntimeGen/SchedulingPolicies/PCGSchedulingPolicyDistanceAndDirection.cpp
@@ -30,6 +30,12 @@ double UPCGSchedulingPolicyDistanceAndDirection::CalculatePriority(const IPCGGen
 
 			const FVector GenSourceToComponent = GenerationBounds.GetClosestPointTo(GenSourcePosition) - GenSourcePosition;
 
+			if (GenSourceToComponent.IsNearlyZero())
+			{
+				// If the gen source is inside the generation bounds then always give maximum priority.
+				return 1.0;
+			}
+
 			if (bUseDistance)
 			{
 				// Normalize via grid size so that priority calculations are well behaved across different grid sizes.
