[PCG] Add Runtime Gen Execution Source
* Added UPCGRuntimeGenExecutionSource, a lightweight alternative to PCG Components that is completely transient and cannot add actors/components to the level. It is ideal when used in high performance contexts with primitive-componentless GPU spawning.

#Synced-CL 49994026

diff --git a/Engine/Plugins/PCG/Source/PCG/Internal/RuntimeGen/PCGRuntimeGenExecutionSource.h b/Engine/Plugins/PCG/Source/PCG/Internal/RuntimeGen/PCGRuntimeGenExecutionSource.h
new file mode 100644
--- /dev/null
+++ b/Engine/Plugins/PCG/Source/PCG/Internal/RuntimeGen/PCGRuntimeGenExecutionSource.h
@@ -0,0 +1,147 @@
+// Copyright Epic Games, Inc. All Rights Reserved.
+
+#pragma once
+
+#include "PCGGraphExecutionInspection.h"
+#include "PCGManagedResourceContainer.h"
+#include "PCGRuntimeGenExecutionState.h"
+#include "ChangeTracking/PCGChangeTracker.h"
+
+#include "PCGRuntimeGenExecutionSource.generated.h"
+
+/**
+ * Runtime generated execution source, managed by the RuntimeGenScheduler. This execution source is never serialized
+ * and does not produce actors/actor-components. It is owned by a concrete original execution source (e.g. a UPCGComponent in the level)
+ * which is responsible for providing state (e.g. World, Seed, Transform, etc.).
+ */
+UCLASS(MinimalAPI, BlueprintType, ClassGroup = (Procedural))
+class UPCGRuntimeGenExecutionSource : public UObject, public IPCGGraphExecutionSource
+{
+	GENERATED_BODY()
+
+public:
+	/** ~Begin UObject interface */
+	static void AddReferencedObjects(UObject* InThis, FReferenceCollector& Collector);
+	/** ~End UObject interface */
+
+	/** ~Begin IPCGGraphExecutionSource interface */
+	virtual IPCGGraphExecutionState& GetExecutionState() override { return ExecutionState; }
+	virtual const IPCGGraphExecutionState& GetExecutionState() const override { return ExecutionState; }
+	/** ~End IPCGGraphExecutionSource interface */
+
+	UPCGRuntimeGenExecutionSource(const FObjectInitializer& InObjectInitializer);
+
+	void Register(IPCGGraphExecutionSource* InOwningSource);
+	void Unregister();
+	void SetOwningSource(IPCGGraphExecutionSource* InOwningSource);
+
+	/** Begin mirror of FPCGRuntimeGenExecutionState */
+	UPCGData* GetSelfData() const;
+	int32 GetSeed() const;
+	FString GetDebugName() const;
+	UWorld* GetWorld() const;
+	UObject* GetTarget() const;
+	bool IsActive() const;
+	bool HasAuthority() const;
+	FTransform GetTransform() const;
+	FTransform GetOriginalTransform() const;
+	FBox GetBounds() const;
+	FBox GetOriginalBounds() const;
+	FBox GetLocalSpaceBounds() const;
+	FBox GetOriginalLocalSpaceBounds() const;
+	UPCGGraph* GetGraph() const;
+	UPCGGraphInstance* GetGraphInstance() const;
+	void Cancel();
+	void OnGraphExecutionAborted(bool bQuiet = false, bool bCleanupUnusedResources = true);
+	bool IsGenerating() const;
+	bool IsGenerated() const;
+	FPCGTaskId GetGenerationTaskId() const;
+	FPCGTaskId Generate(const IPCGGraphExecutionState::FGenerateParams& InGenerateParams);
+	FPCGTaskId Cleanup(const IPCGGraphExecutionState::FCleanupParams& InCleanupParams);
+	void StoreOutputDataForPin(const FString& InResourceKey, const FPCGDataCollection& InData);
+	const FPCGDataCollection* RetrieveOutputDataForPin(const FString& InResourceKey) const;
+	FPCGManagedResourceContainer* GetManagedResourceContainer();
+	FTransactionallySafeCriticalSection* GetManagedResourceContainerLock();
+
+#if WITH_EDITOR
+	const PCGUtils::FExtraCapture& GetExtraCapture() const;
+	PCGUtils::FExtraCapture& GetExtraCapture();
+	const FPCGGraphExecutionInspection& GetInspection() const;
+	FPCGGraphExecutionInspection& GetInspection();
+	FPCGDynamicTrackingPriority GetDynamicTrackingPriority() const;
+	void RegisterDynamicTracking(const UPCGSettings* InSettings, const TArrayView<TPair<FPCGSelectionKey, bool>>& InDynamicKeysAndCulling);
+	void RegisterDynamicTracking(const FPCGSelectionKeyToSettingsMap& InKeysToSettings);
+	TArray<FPCGSelectionKey> GatherTrackingKeys() const;
+	bool IsKeyTrackedAndCulled(const FPCGSelectionKey& Key, bool& bOutIsCulled) const;
+	void ForEachSettingTrackingKey(const FPCGSelectionKey& InKey, const TFunctionRef<void(const FPCGSelectionKey&, const FPCGSettingsAndCulling&)> InCallback) const;
+	void AddLocalSourceDynamicTrackingKeys(const FPCGSelectionKeyToSettingsMap& LocalSourceDynamicTrackingKeys);
+	void AddLocalSourceCurrentTrackingKeys(const FPCGSelectionKeyToSettingsMap& LocalSourceCurrentExecutionTrackingKeys, const TSet<const UPCGSettings*>& LocalSourceCurrentExecutionTrackingSettings);
+	bool IsRefreshInProgress() const;
+#endif
+
+	IPCGGraphExecutionSource* GetOriginalSource() const;
+	IPCGGraphExecutionSource* GetLocalSource(const FPCGGridDescriptor& InGridDescriptor, const FIntVector& InGridCoords) const;
+	bool IsPartitioned() const;
+	bool IsLocalSource() const;
+	bool Use2DGrid() const;
+	EPCGHiGenGrid GetGenerationGrid() const;
+	uint32 GetGenerationGridSize() const;
+	FIntVector GetGenerationGridCoords() const;
+	bool IsManagedByRuntimeGenSystem() const;
+	const UPCGSchedulingPolicyBase* GetRuntimeGenSchedulingPolicy() const;
+	const FPCGRuntimeGenerationRadii& GetGenerationRadii() const;
+	double GetGenerationRadiusFromGrid(EPCGHiGenGrid InGrid) const;
+	double GetCleanupRadiusFromGrid(EPCGHiGenGrid InGrid) const;
+	/** End mirror of FPCGRuntimeGenExecutionState */
+
+	/** Creates a local execution source that is automatically setup and registered with this original execution source. Nullptr if this is not a partitioned original source. */
+	UPCGRuntimeGenExecutionSource* CreateLocalSource(uint32 InGridSize, const FIntVector& InGridCoords);
+
+	/** Removes a local source from the tracking set allowing it to be garbage collected. */
+	void DestroyLocalSource(UPCGRuntimeGenExecutionSource* InLocalSource);
+
+	/** The owning execution source drives the properties for the RuntimeGen execution source. */
+	IPCGGraphExecutionSource* GetOwningSource();
+
+private:
+	/** Clear any data stored for any pins. */
+	void ClearPerPinGeneratedOutput();
+
+	void CreateSelfData();
+
+#if WITH_EDITOR
+	bool UpdateTrackingCache(TArray<FPCGSelectionKey>* InOptionalChangedKeys = nullptr);
+#endif
+
+private:
+	FPCGRuntimeGenExecutionState ExecutionState;
+
+	IPCGGraphExecutionSource* OwningSource = nullptr;
+	UPCGRuntimeGenExecutionSource* OriginalSource = nullptr;
+	uint32 GridSize = PCGHiGenGrid::UnboundedGridSize();
+	FIntVector GridCoords = FIntVector::ZeroValue;
+
+	UPROPERTY()
+	TSet<TObjectPtr<UPCGRuntimeGenExecutionSource>> LocalSources;
+
+	FPCGTaskId GenerationTaskId = InvalidPCGTaskId;
+	bool bGenerated = false;
+
+	/** If any graph edges cross execution grid sizes, data on the edge is stored / retrieved from this map. */
+	UPROPERTY(Transient)
+	TMap<FString, FPCGDataCollection> PerPinGeneratedOutput;
+	mutable FTransactionallySafeRWLock PerPinGeneratedOutputLock;
+
+	UPROPERTY()
+	FPCGManagedResourceContainer ManagedResourceContainer;
+	mutable FTransactionallySafeCriticalSection ManagedResourcesLock;
+
+#if WITH_EDITOR
+	FPCGGraphExecutionInspection ExecutionInspection;
+	FPCGChangeTracker ChangeTracker;
+	mutable PCGUtils::FExtraCapture ExtraCapture;
+#endif
+
+	UPROPERTY()
+	TObjectPtr<UPCGData> SelfData = nullptr;
+};
diff --git a/Engine/Plugins/PCG/Source/PCG/Internal/RuntimeGen/PCGRuntimeGenExecutionState.h b/Engine/Plugins/PCG/Source/PCG/Internal/RuntimeGen/PCGRuntimeGenExecutionState.h
new file mode 100644
--- /dev/null
+++ b/Engine/Plugins/PCG/Source/PCG/Internal/RuntimeGen/PCGRuntimeGenExecutionState.h
@@ -0,0 +1,75 @@
+// Copyright Epic Games, Inc. All Rights Reserved.
+
+#pragma once
+
+#include "PCGGraphExecutionStateInterface.h"
+
+class UPCGRuntimeGenExecutionSource;
+
+class FPCGRuntimeGenExecutionState : public IPCGGraphExecutionState
+{
+	/** ~Begin IPCGGraphExecutionState interface */
+public:
+	virtual UPCGData* GetSelfData() const override;
+	virtual int32 GetSeed() const override;
+	virtual FString GetDebugName() const override;
+	virtual UWorld* GetWorld() const override;
+	virtual UObject* GetTarget() const override;
+	virtual bool IsActive() const override;
+	virtual bool HasAuthority() const override;
+	virtual FTransform GetTransform() const override;
+	virtual FTransform GetOriginalTransform() const override;
+	virtual FBox GetBounds() const override;
+	virtual FBox GetOriginalBounds() const override;
+	virtual FBox GetLocalSpaceBounds() const override;
+	virtual FBox GetOriginalLocalSpaceBounds() const override;
+	virtual UPCGGraph* GetGraph() const override;
+	virtual UPCGGraphInstance* GetGraphInstance() const override;
+	virtual void Cancel() override;
+	virtual void OnGraphExecutionAborted(bool bQuiet = false, bool bCleanupUnusedResources = true) override;
+	virtual bool IsGenerating() const override;
+	virtual bool IsGenerated() const override;
+	virtual FPCGTaskId GetGenerationTaskId() const override;
+	virtual FPCGTaskId Generate(const FGenerateParams& InGenerateParams) override;
+	virtual FPCGTaskId Cleanup(const FCleanupParams& InCleanupParams) override;
+	virtual void StoreOutputDataForPin(const FString& InResourceKey, const FPCGDataCollection& InData) const override;
+	virtual const FPCGDataCollection* RetrieveOutputDataForPin(const FString& InResourceKey) const override;
+	virtual FPCGManagedResourceContainer* GetManagedResourceContainer() override;
+	virtual FTransactionallySafeCriticalSection* GetManagedResourceContainerLock() override;
+
+#if WITH_EDITOR
+	virtual const PCGUtils::FExtraCapture& GetExtraCapture() const override;
+	virtual PCGUtils::FExtraCapture& GetExtraCapture() override;
+	virtual const FPCGGraphExecutionInspection& GetInspection() const override;
+	virtual FPCGGraphExecutionInspection& GetInspection() override;
+	PCG_API virtual FPCGDynamicTrackingPriority GetDynamicTrackingPriority() const override;
+	virtual void RegisterDynamicTracking(const UPCGSettings* InSettings, const TArrayView<TPair<FPCGSelectionKey, bool>>& InDynamicKeysAndCulling) override;
+	virtual void RegisterDynamicTracking(const FPCGSelectionKeyToSettingsMap& InKeysToSettings) override;
+	virtual TArray<FPCGSelectionKey> GatherTrackingKeys() const override;
+	virtual bool IsKeyTrackedAndCulled(const FPCGSelectionKey& Key, bool& bOutIsCulled) const override;
+	virtual void ForEachSettingTrackingKey(const FPCGSelectionKey& InKey, const TFunctionRef<void(const FPCGSelectionKey&, const FPCGSettingsAndCulling&)> InCallback) const override;
+	virtual void AddLocalSourceDynamicTrackingKeys(const FPCGSelectionKeyToSettingsMap& LocalSourceDynamicTrackingKeys) override;
+	virtual void AddLocalSourceCurrentTrackingKeys(const FPCGSelectionKeyToSettingsMap& LocalSourceCurrentExecutionTrackingKeys, const TSet<const UPCGSettings*>& LocalSourceCurrentExecutionTrackingSettings) override;
+	virtual bool IsRefreshInProgress() const override;
+#endif
+
+	virtual IPCGGraphExecutionSource* GetOriginalSource() const override;
+	virtual IPCGGraphExecutionSource* GetLocalSource(const FPCGGridDescriptor& GridDescriptor, const FIntVector& CellCoords) const override;
+	virtual bool IsPartitioned() const override;
+	virtual bool IsLocalSource() const override;
+	virtual bool Use2DGrid() const override;
+	virtual EPCGHiGenGrid GetGenerationGrid() const override;
+	virtual uint32 GetGenerationGridSize() const override;
+	virtual FIntVector GetGenerationGridCoords() const override;
+	virtual bool IsManagedByRuntimeGenSystem() const override;
+	virtual const UPCGSchedulingPolicyBase* GetRuntimeGenSchedulingPolicy() const override;
+	virtual const FPCGRuntimeGenerationRadii& GetGenerationRadii() const override;
+	virtual double GetGenerationRadiusFromGrid(EPCGHiGenGrid InGrid) const override;
+	virtual double GetCleanupRadiusFromGrid(EPCGHiGenGrid InGrid) const override;
+	/** ~End IPCGGraphExecutionState interface */
+
+	void SetSource(UPCGRuntimeGenExecutionSource* InSource);
+
+private:
+	UPCGRuntimeGenExecutionSource* Source = nullptr;
+};
diff --git a/Engine/Plugins/PCG/Source/PCG/Private/PCGComponentExecutionState.cpp b/Engine/Plugins/PCG/Source/PCG/Private/PCGComponentExecutionState.cpp
--- a/Engine/Plugins/PCG/Source/PCG/Private/PCGComponentExecutionState.cpp
+++ b/Engine/Plugins/PCG/Source/PCG/Private/PCGComponentExecutionState.cpp
@@ -4,6 +4,7 @@
 
 #include "PCGComponent.h"
 #include "PCGSubsystem.h"
+#include "Grid/PCGPartitionActor.h"
 
 #include "LandscapeProxy.h"
 #include "GameFramework/Actor.h"
@@ -56,6 +57,12 @@ UObject* FPCGComponentExecutionState::GetTarget() const
 	return Component->GetOwner();
 }
 
+bool FPCGComponentExecutionState::IsActive() const
+{
+	check(Component);
+	return Component->bActivated;
+}
+
 bool FPCGComponentExecutionState::HasAuthority() const
 {
 	check(Component && Component->GetOwner());
@@ -146,6 +153,30 @@ bool FPCGComponentExecutionState::IsManagedByRuntimeGenSystem() const
 	return Component->IsManagedByRuntimeGenSystem();
 }
 
+const UPCGSchedulingPolicyBase* FPCGComponentExecutionState::GetRuntimeGenSchedulingPolicy() const
+{
+	check(Component);
+	return Component->GetRuntimeGenSchedulingPolicy();
+}
+
+const FPCGRuntimeGenerationRadii& FPCGComponentExecutionState::GetGenerationRadii() const
+{
+	check(Component);
+	return Component->GetGenerationRadii();
+}
+
+double FPCGComponentExecutionState::GetGenerationRadiusFromGrid(EPCGHiGenGrid InGrid) const
+{
+	check(Component);
+	return Component->GetGenerationRadiusFromGrid(InGrid);
+}
+
+double FPCGComponentExecutionState::GetCleanupRadiusFromGrid(EPCGHiGenGrid InGrid) const
+{
+	check(Component);
+	return Component->GetCleanupRadiusFromGrid(InGrid);
+}
+
 IPCGGraphExecutionSource* FPCGComponentExecutionState::GetOriginalSource() const
 {
 	check(Component);
@@ -194,6 +225,18 @@ uint32 FPCGComponentExecutionState::GetGenerationGridSize() const
 	return (IsPartitioned() || IsLocalSource()) ? Component->GetGenerationGridSize() : PCGHiGenGrid::UninitializedGridSize();
 }
 
+FIntVector FPCGComponentExecutionState::GetGenerationGridCoords() const
+{
+	check(Component);
+
+	if (APCGPartitionActor* PartitionActor = Cast<APCGPartitionActor>(Component->GetOwner()))
+	{
+		return PartitionActor->GetGridCoord();
+	}
+
+	return FIntVector::ZeroValue;
+}
+
 bool FPCGComponentExecutionState::IsInPreviewMode() const
 {
 	check(Component);
diff --git a/Engine/Plugins/PCG/Source/PCG/Private/PCGGraphExecutionStateInterface.cpp b/Engine/Plugins/PCG/Source/PCG/Private/PCGGraphExecutionStateInterface.cpp
--- a/Engine/Plugins/PCG/Source/PCG/Private/PCGGraphExecutionStateInterface.cpp
+++ b/Engine/Plugins/PCG/Source/PCG/Private/PCGGraphExecutionStateInterface.cpp
@@ -46,6 +46,12 @@ FPCGGridDescriptor IPCGGraphExecutionState::GetGridDescriptor(uint32 InGridSize)
 		.SetIsRuntime(IsManagedByRuntimeGenSystem());
 }
 
+const FPCGRuntimeGenerationRadii& IPCGGraphExecutionState::GetGenerationRadii() const
+{
+	static const FPCGRuntimeGenerationRadii DefaultGenerationRadii;
+	return DefaultGenerationRadii;
+}
+
 #if WITH_EDITOR
 
 FPCGDynamicTrackingPriority IPCGGraphExecutionState::GetDynamicTrackingPriority() const
diff --git a/Engine/Plugins/PCG/Source/PCG/Private/RuntimeGen/PCGRuntimeGenExecutionSource.cpp b/Engine/Plugins/PCG/Source/PCG/Private/RuntimeGen/PCGRuntimeGenExecutionSource.cpp
new file mode 100644
--- /dev/null
+++ b/Engine/Plugins/PCG/Source/PCG/Private/RuntimeGen/PCGRuntimeGenExecutionSource.cpp
@@ -0,0 +1,601 @@
+// Copyright Epic Games, Inc. All Rights Reserved.
+
+#include "RuntimeGen/PCGRuntimeGenExecutionSource.h"
+
+#include "PCGContext.h"
+#include "PCGData.h"
+#include "PCGGraph.h"
+#include "PCGManagedResourceContainer.h"
+#include "PCGSubsystem.h"
+#include "Data/PCGVolumeData.h"
+#include "Utils/PCGGeneratedResourcesLogging.h"
+#include "Utils/PCGGraphExecutionLogging.h"
+
+#include UE_INLINE_GENERATED_CPP_BY_NAME(PCGRuntimeGenExecutionSource)
+
+void UPCGRuntimeGenExecutionSource::AddReferencedObjects(UObject* InThis, FReferenceCollector& Collector)
+{
+	Super::AddReferencedObjects(InThis, Collector);
+
+	UPCGRuntimeGenExecutionSource* This = CastChecked<UPCGRuntimeGenExecutionSource>(InThis);
+
+#if WITH_EDITOR
+	This->ExecutionInspection.AddReferencedObjects(Collector);
+#endif
+}
+
+UPCGRuntimeGenExecutionSource::UPCGRuntimeGenExecutionSource(const FObjectInitializer& InObjectInitializer)
+	: Super(InObjectInitializer)
+{
+	ExecutionState.SetSource(this);
+	ManagedResourceContainer.SetOwner(this);
+}
+
+void UPCGRuntimeGenExecutionSource::Register(IPCGGraphExecutionSource* InOwningSource)
+{
+	OriginalSource = this;
+	GridSize = PCGHiGenGrid::UnboundedGridSize();
+
+	SetOwningSource(InOwningSource);
+
+	if (UPCGSubsystem* Subsystem = Cast<UPCGSubsystem>(ExecutionState.GetSubsystem()))
+	{
+		Subsystem->RegisterOrUpdateExecutionSource(this, /*bDoPartitionMapping=*/false);
+	}
+}
+
+void UPCGRuntimeGenExecutionSource::Unregister()
+{
+	if (UPCGSubsystem* Subsystem = Cast<UPCGSubsystem>(ExecutionState.GetSubsystem()))
+	{
+		Subsystem->UnregisterExecutionSource(this);
+	}
+}
+
+void UPCGRuntimeGenExecutionSource::SetOwningSource(IPCGGraphExecutionSource* InOwningSource)
+{
+	check(InOwningSource);
+
+	if (IsLocalSource())
+	{
+		return;
+	}
+
+	OwningSource = InOwningSource;
+
+	CreateSelfData();
+
+#if WITH_EDITOR
+	UpdateTrackingCache();
+#endif
+
+	for (UPCGRuntimeGenExecutionSource* LocalSource : LocalSources)
+	{
+		if (!ensure(LocalSource))
+		{
+			continue;
+		}
+
+		LocalSource->OwningSource = OwningSource;
+		LocalSource->CreateSelfData();
+	}
+}
+
+UPCGData* UPCGRuntimeGenExecutionSource::GetSelfData() const
+{
+	return SelfData;
+}
+
+int32 UPCGRuntimeGenExecutionSource::GetSeed() const
+{
+	check(OwningSource);
+	return OwningSource->GetExecutionState().GetSeed();
+}
+
+FString UPCGRuntimeGenExecutionSource::GetDebugName() const
+{
+	if (!IsLocalSource())
+	{
+		check(OwningSource);
+		return OwningSource->GetExecutionState().GetDebugName();
+	}
+
+	TStringBuilderWithBuffer<TCHAR, 1024> StringBuilder;
+
+	StringBuilder += FString::Printf(TEXT("PCGRuntimeGenExecutionSource_%d_"), GridSize);
+
+	if (Use2DGrid())
+	{
+		StringBuilder += FString::Printf(TEXT("%d_%d"), GridCoords.X, GridCoords.Y);
+	}
+	else
+	{
+		StringBuilder += FString::Printf(TEXT("%d_%d_%d"), GridCoords.X, GridCoords.Y, GridCoords.Z);
+	}
+
+	return StringBuilder.ToString();
+}
+
+UWorld* UPCGRuntimeGenExecutionSource::GetWorld() const
+{
+	check(OwningSource);
+	return OwningSource->GetExecutionState().GetWorld();
+}
+
+UObject* UPCGRuntimeGenExecutionSource::GetTarget() const
+{
+	check(OwningSource);
+
+	// Local runtime gen execution sources don't have any target, but for original sources we can fallback to the original PCG component's actor.
+	return IsLocalSource() ? nullptr : OwningSource->GetExecutionState().GetTarget();
+}
+
+bool UPCGRuntimeGenExecutionSource::IsActive() const
+{
+	check(OwningSource);
+	return OwningSource->GetExecutionState().IsActive();
+}
+
+bool UPCGRuntimeGenExecutionSource::HasAuthority() const
+{
+	check(OwningSource);
+	return OwningSource->GetExecutionState().HasAuthority();
+}
+
+FTransform UPCGRuntimeGenExecutionSource::GetTransform() const
+{
+	return IsLocalSource() ? FTransform(FQuat::Identity, (FVector(GridCoords) + FVector(0.5)) * GetGenerationGridSize(), FVector::OneVector) : GetOriginalTransform();
+}
+
+FTransform UPCGRuntimeGenExecutionSource::GetOriginalTransform() const
+{
+	check(OwningSource);
+	return OwningSource->GetExecutionState().GetOriginalTransform();
+}
+
+FBox UPCGRuntimeGenExecutionSource::GetBounds() const
+{
+	if (IsLocalSource())
+	{
+		const FVector Center = GetTransform().GetLocation();
+		const FVector::FReal HalfGridSize = GridSize / 2.0;
+
+		FVector Extent(HalfGridSize, HalfGridSize, HalfGridSize);
+
+		// In case of 2D grid, it's like the actor has infinite bounds on the Z axis
+		if (Use2DGrid())
+		{
+			Extent.Z = HALF_WORLD_MAX1;
+		}
+
+		return FBox(Center - Extent, Center + Extent).Overlap(GetOriginalBounds());
+	}
+	else
+	{
+		return GetOriginalBounds();
+	}
+}
+
+FBox UPCGRuntimeGenExecutionSource::GetOriginalBounds() const
+{
+	check(OwningSource);
+	return OwningSource->GetExecutionState().GetOriginalBounds();
+}
+
+FBox UPCGRuntimeGenExecutionSource::GetLocalSpaceBounds() const
+{
+	return GetBounds();
+}
+
+FBox UPCGRuntimeGenExecutionSource::GetOriginalLocalSpaceBounds() const
+{
+	check(OwningSource);
+	return OwningSource->GetExecutionState().GetOriginalLocalSpaceBounds();
+}
+
+UPCGGraph* UPCGRuntimeGenExecutionSource::GetGraph() const
+{
+	check(OwningSource);
+	return OwningSource->GetExecutionState().GetGraph();
+}
+
+UPCGGraphInstance* UPCGRuntimeGenExecutionSource::GetGraphInstance() const
+{
+	check(OwningSource);
+	return OwningSource->GetExecutionState().GetGraphInstance();
+}
+
+void UPCGRuntimeGenExecutionSource::Cancel()
+{
+	if (GenerationTaskId != InvalidPCGTaskId && ExecutionState.GetSubsystem())
+	{
+		ExecutionState.GetSubsystem()->CancelGeneration(this);
+	}
+}
+
+void UPCGRuntimeGenExecutionSource::OnGraphExecutionAborted(bool bQuiet, bool bCleanupUnusedResources)
+{
+	GenerationTaskId = InvalidPCGTaskId;
+
+	if (bCleanupUnusedResources)
+	{
+		ManagedResourceContainer.CleanupUnusedManagedResources();
+	}
+
+#if WITH_EDITOR
+	if (IPCGBaseSubsystem* Subsystem = GetExecutionState().GetSubsystem())
+	{
+		Subsystem->OnPCGSourceGenerationDone(this, EPCGGenerationStatus::Completed);
+	}
+#endif
+}
+
+bool UPCGRuntimeGenExecutionSource::IsGenerating() const
+{
+	return GenerationTaskId != InvalidPCGTaskId;
+}
+
+bool UPCGRuntimeGenExecutionSource::IsGenerated() const
+{
+	return bGenerated;
+}
+
+FPCGTaskId UPCGRuntimeGenExecutionSource::GetGenerationTaskId() const
+{
+	return GenerationTaskId;
+}
+
+FPCGTaskId UPCGRuntimeGenExecutionSource::Generate(const IPCGGraphExecutionState::FGenerateParams& InGenerateParams)
+{
+	ClearPerPinGeneratedOutput();
+
+	if (UPCGSubsystem* Subsystem = Cast<UPCGSubsystem>(ExecutionState.GetSubsystem()))
+	{
+		const FPCGTaskId NewGenerationTaskId = Subsystem->ScheduleGraph(this, InGenerateParams.Dependencies);
+
+		if (NewGenerationTaskId != InvalidPCGTaskId)
+		{
+			auto PostProcessTask = [WeakThis = TWeakObjectPtr<UPCGRuntimeGenExecutionSource>(this)]()
+			{
+				UPCGRuntimeGenExecutionSource* This = WeakThis.Get();
+
+				if (!IsValid(This))
+				{
+					return true;
+				}
+
+				PCGGraphExecutionLogging::LogPostProcessGraph(This);
+
+				This->bGenerated = true;
+				This->GenerationTaskId = InvalidPCGTaskId;
+
+#if WITH_EDITOR
+				if (IPCGBaseSubsystem* Subsystem = This->GetExecutionState().GetSubsystem())
+				{
+					Subsystem->OnPCGSourceGenerationDone(This, EPCGGenerationStatus::Completed);
+				}
+#endif
+
+				return true;
+			};
+
+			// GenerationTaskId should actually be the post-process task, not the generation task itself.
+			GenerationTaskId = Subsystem->ScheduleGeneric(PostProcessTask, this, /*Dependencies=*/{ NewGenerationTaskId });
+			return GenerationTaskId;
+		}
+	}
+
+	return InvalidPCGTaskId;
+}
+
+FPCGTaskId UPCGRuntimeGenExecutionSource::Cleanup(const IPCGGraphExecutionState::FCleanupParams& InCleanupParams)
+{
+	TRACE_CPUPROFILER_EVENT_SCOPE(UPCGRuntimeGenExecutionSource::Cleanup);
+
+	// Runtime gen cleanup is always immediate (for now), so we only implement the immediate path here.
+	if (!InCleanupParams.bImmediate)
+	{
+		return InvalidPCGTaskId;
+	}
+
+	PCGGeneratedResourcesLogging::LogCleanupLocalImmediate(this, InCleanupParams.bReleaseManagedResources, ManagedResourceContainer.GetManagedResourcesCopy());
+
+	Cancel();
+
+	ManagedResourceContainer.Cleanup(InCleanupParams.bReleaseManagedResources);
+
+	if (InCleanupParams.bReleaseManagedResources)
+	{
+		bGenerated = false;
+
+		ClearPerPinGeneratedOutput();
+	}
+
+	if (!IsLocalSource() && InCleanupParams.bCleanupLocalSources)
+	{
+		for (UPCGRuntimeGenExecutionSource* LocalSource : LocalSources)
+		{
+			if (LocalSource)
+			{
+				LocalSource->GetExecutionState().Cleanup(InCleanupParams);
+			}
+		}
+	}
+
+	PCGGeneratedResourcesLogging::LogCleanupLocalImmediateFinished(this, ManagedResourceContainer.GetManagedResourcesCopy());
+
+	return InvalidPCGTaskId;
+}
+
+void UPCGRuntimeGenExecutionSource::StoreOutputDataForPin(const FString& InResourceKey, const FPCGDataCollection& InData)
+{
+	UE::TWriteScopeLock ScopedWriteLock(PerPinGeneratedOutputLock);
+
+	InData.MarkUsage(EPCGDataUsage::ComponentPerPinOutputData);
+
+	FPCGDataCollection* FoundExistingData = PerPinGeneratedOutput.Find(InResourceKey);
+
+	// For all existing data items, clear their usage which may release transient resources if the data is not present in the new data collection.
+	if (FoundExistingData)
+	{
+		for (const FPCGTaggedData& ExistingData : FoundExistingData->TaggedData)
+		{
+			if (ExistingData.Data && !InData.TaggedData.FindByPredicate([&ExistingData](const FPCGTaggedData& NewData) { return ExistingData.Data == NewData.Data; }))
+			{
+				ExistingData.Data->ClearUsage(EPCGDataUsage::ComponentPerPinOutputData);
+			}
+		}
+
+		*FoundExistingData = InData;
+	}
+	else
+	{
+		PerPinGeneratedOutput.Add(InResourceKey, InData);
+	}
+}
+
+const FPCGDataCollection* UPCGRuntimeGenExecutionSource::RetrieveOutputDataForPin(const FString& InResourceKey) const
+{
+	UE::TReadScopeLock ScopedReadLock(PerPinGeneratedOutputLock);
+	return PerPinGeneratedOutput.Find(InResourceKey);
+}
+
+FPCGManagedResourceContainer* UPCGRuntimeGenExecutionSource::GetManagedResourceContainer()
+{
+	return &ManagedResourceContainer;
+}
+
+FTransactionallySafeCriticalSection* UPCGRuntimeGenExecutionSource::GetManagedResourceContainerLock()
+{
+	return &ManagedResourcesLock;
+}
+
+#if WITH_EDITOR
+const PCGUtils::FExtraCapture& UPCGRuntimeGenExecutionSource::GetExtraCapture() const
+{
+	return ExtraCapture;
+}
+
+PCGUtils::FExtraCapture& UPCGRuntimeGenExecutionSource::GetExtraCapture()
+{
+	return ExtraCapture;
+}
+
+const FPCGGraphExecutionInspection& UPCGRuntimeGenExecutionSource::GetInspection() const
+{
+	return ExecutionInspection;
+}
+
+FPCGGraphExecutionInspection& UPCGRuntimeGenExecutionSource::GetInspection()
+{
+	return ExecutionInspection;
+}
+
+FPCGDynamicTrackingPriority UPCGRuntimeGenExecutionSource::GetDynamicTrackingPriority() const
+{
+	return FPCGDynamicTrackingPriority(); 
+}
+
+void UPCGRuntimeGenExecutionSource::RegisterDynamicTracking(const UPCGSettings* InSettings, const TArrayView<TPair<FPCGSelectionKey, bool>>& InDynamicKeysAndCulling)
+{
+	ChangeTracker.RegisterDynamicTracking(InSettings, InDynamicKeysAndCulling);
+}
+
+void UPCGRuntimeGenExecutionSource::RegisterDynamicTracking(const FPCGSelectionKeyToSettingsMap& InKeysToSettings)
+{
+	ChangeTracker.RegisterDynamicTracking(InKeysToSettings);
+}
+
+TArray<FPCGSelectionKey> UPCGRuntimeGenExecutionSource::GatherTrackingKeys() const
+{
+	return ChangeTracker.GatherTrackingKeys();
+}
+
+bool UPCGRuntimeGenExecutionSource::IsKeyTrackedAndCulled(const FPCGSelectionKey& Key, bool& bOutIsCulled) const
+{
+	return ChangeTracker.IsKeyTrackedAndCulled(Key, bOutIsCulled);
+}
+
+void UPCGRuntimeGenExecutionSource::ForEachSettingTrackingKey(const FPCGSelectionKey& InKey, const TFunctionRef<void(const FPCGSelectionKey&, const FPCGSettingsAndCulling&)> InCallback) const
+{
+	ChangeTracker.ForEachSettingTrackingKey(InKey, InCallback);
+}
+
+void UPCGRuntimeGenExecutionSource::AddLocalSourceDynamicTrackingKeys(const FPCGSelectionKeyToSettingsMap& LocalSourceDynamicTrackingKeys)
+{
+	ChangeTracker.AddLocalSourceDynamicTrackingKeys(LocalSourceDynamicTrackingKeys);
+}
+
+void UPCGRuntimeGenExecutionSource::AddLocalSourceCurrentTrackingKeys(const FPCGSelectionKeyToSettingsMap& LocalSourceCurrentExecutionTrackingKeys, const TSet<const UPCGSettings*>& LocalSourceCurrentExecutionTrackingSettings)
+{
+	ChangeTracker.AddLocalSourceCurrentTrackingKeys(LocalSourceCurrentExecutionTrackingKeys, LocalSourceCurrentExecutionTrackingSettings);
+}
+
+bool UPCGRuntimeGenExecutionSource::IsRefreshInProgress() const
+{
+	return false;
+}
+#endif // WITH_EDITOR
+
+IPCGGraphExecutionSource* UPCGRuntimeGenExecutionSource::GetOriginalSource() const
+{
+	return const_cast<UPCGRuntimeGenExecutionSource*>(OriginalSource);
+}
+
+IPCGGraphExecutionSource* UPCGRuntimeGenExecutionSource::GetLocalSource(const FPCGGridDescriptor& InGridDescriptor, const FIntVector& InGridCoords) const
+{
+	const uint32 LocalGridSize = InGridDescriptor.GetGridSize();
+
+	// @todo_pcg: Could have a hashed lookup instead.
+	for (UPCGRuntimeGenExecutionSource* LocalSource : OriginalSource->LocalSources)
+	{
+		if (LocalSource && LocalSource->GetGenerationGridSize() == LocalGridSize && LocalSource->GridCoords == InGridCoords)
+		{
+			return LocalSource;
+		}
+	}
+
+	return nullptr;
+}
+
+bool UPCGRuntimeGenExecutionSource::IsPartitioned() const
+{
+	check(OwningSource);
+	return OwningSource->GetExecutionState().IsPartitioned();
+}
+
+bool UPCGRuntimeGenExecutionSource::IsLocalSource() const
+{
+	return this != OriginalSource;
+}
+
+bool UPCGRuntimeGenExecutionSource::Use2DGrid() const
+{
+	check(OwningSource);
+	return OwningSource->GetExecutionState().Use2DGrid();
+}
+
+EPCGHiGenGrid UPCGRuntimeGenExecutionSource::GetGenerationGrid() const
+{
+	if (PCGHiGenGrid::IsValidGridSize(GridSize))
+	{
+		return PCGHiGenGrid::GridSizeToGrid(GridSize);
+	}
+	else if (GridSize == PCGHiGenGrid::UnboundedGridSize())
+	{
+		return EPCGHiGenGrid::Unbounded;
+	}
+	else
+	{
+		return PCGHiGenGrid::GridSizeToGrid(GetGraph()->GetDefaultGridSize());
+	}
+}
+
+uint32 UPCGRuntimeGenExecutionSource::GetGenerationGridSize() const
+{
+	return (IsPartitioned() || IsLocalSource()) ? GridSize : PCGHiGenGrid::UninitializedGridSize();
+}
+
+FIntVector UPCGRuntimeGenExecutionSource::GetGenerationGridCoords() const
+{
+	return GridCoords;
+}
+
+bool UPCGRuntimeGenExecutionSource::IsManagedByRuntimeGenSystem() const
+{
+	return true;
+}
+
+const UPCGSchedulingPolicyBase* UPCGRuntimeGenExecutionSource::GetRuntimeGenSchedulingPolicy() const
+{
+	check(OwningSource);
+	return OwningSource->GetExecutionState().GetRuntimeGenSchedulingPolicy();
+}
+
+const FPCGRuntimeGenerationRadii& UPCGRuntimeGenExecutionSource::GetGenerationRadii() const
+{
+	check(OwningSource);
+
+	static const FPCGRuntimeGenerationRadii DefaultGenerationRadii;
+	return OwningSource->GetExecutionState().GetGenerationRadii();
+}
+
+double UPCGRuntimeGenExecutionSource::GetGenerationRadiusFromGrid(EPCGHiGenGrid InGrid) const
+{
+	check(OwningSource);
+	return OwningSource->GetExecutionState().GetGenerationRadiusFromGrid(InGrid);
+}
+
+double UPCGRuntimeGenExecutionSource::GetCleanupRadiusFromGrid(EPCGHiGenGrid InGrid) const
+{
+	check(OwningSource);
+	return OwningSource->GetExecutionState().GetCleanupRadiusFromGrid(InGrid);
+}
+
+UPCGRuntimeGenExecutionSource* UPCGRuntimeGenExecutionSource::CreateLocalSource(uint32 InGridSize, const FIntVector& InGridCoords)
+{
+	check(IsInGameThread());
+
+	if (IsLocalSource() || !IsPartitioned())
+	{
+		return nullptr;
+	}
+
+	UPCGRuntimeGenExecutionSource* LocalSource = NewObject<UPCGRuntimeGenExecutionSource>();
+	LocalSource->OwningSource = OwningSource;
+	LocalSource->OriginalSource = this;
+	LocalSource->GridSize = InGridSize;
+	LocalSource->GridCoords = InGridCoords;
+	LocalSource->CreateSelfData();
+
+	LocalSources.Add(LocalSource);
+
+	return LocalSource;
+}
+
+void UPCGRuntimeGenExecutionSource::DestroyLocalSource(UPCGRuntimeGenExecutionSource* InLocalSource)
+{
+	if (InLocalSource)
+	{
+		IPCGGraphExecutionState::FCleanupParams CleanupParams;
+		CleanupParams.bImmediate = true;
+		
+		InLocalSource->GetExecutionState().Cleanup(CleanupParams);
+	}
+
+	LocalSources.Remove(InLocalSource);
+}
+
+IPCGGraphExecutionSource* UPCGRuntimeGenExecutionSource::GetOwningSource()
+{
+	return OwningSource;
+}
+
+void UPCGRuntimeGenExecutionSource::ClearPerPinGeneratedOutput()
+{
+	UE::TWriteScopeLock ScopedWriteLock(PerPinGeneratedOutputLock);
+
+	for (TPair<FString, FPCGDataCollection>& Entry : PerPinGeneratedOutput)
+	{
+		Entry.Value.ClearUsage(EPCGDataUsage::ComponentPerPinOutputData);
+	}
+
+	PerPinGeneratedOutput.Reset();
+}
+
+void UPCGRuntimeGenExecutionSource::CreateSelfData()
+{
+	check(IsInGameThread());
+
+	UPCGVolumeData* VolumeData = NewObject<UPCGVolumeData>();
+	VolumeData->Initialize(GetBounds());
+
+	SelfData = VolumeData;
+}
+
+#if WITH_EDITOR
+bool UPCGRuntimeGenExecutionSource::UpdateTrackingCache(TArray<FPCGSelectionKey>* InOptionalChangedKeys)
+{
+	return ChangeTracker.UpdateTrackingCache(GetGraph(), InOptionalChangedKeys);
+}
+#endif
diff --git a/Engine/Plugins/PCG/Source/PCG/Private/RuntimeGen/PCGRuntimeGenExecutionState.cpp b/Engine/Plugins/PCG/Source/PCG/Private/RuntimeGen/PCGRuntimeGenExecutionState.cpp
new file mode 100644
--- /dev/null
+++ b/Engine/Plugins/PCG/Source/PCG/Private/RuntimeGen/PCGRuntimeGenExecutionState.cpp
@@ -0,0 +1,324 @@
+// Copyright Epic Games, Inc. All Rights Reserved.
+
+#include "RuntimeGen/PCGRuntimeGenExecutionState.h"
+
+#include "RuntimeGen/PCGRuntimeGenExecutionSource.h"
+
+UPCGData* FPCGRuntimeGenExecutionState::GetSelfData() const
+{
+	check(Source);
+	return Source->GetSelfData();
+}
+
+int32 FPCGRuntimeGenExecutionState::GetSeed() const
+{
+	check(Source);
+	return Source->GetSeed();
+}
+
+FString FPCGRuntimeGenExecutionState::GetDebugName() const
+{
+	check(Source);
+	return Source->GetDebugName();
+}
+
+UWorld* FPCGRuntimeGenExecutionState::GetWorld() const
+{
+	check(Source);
+	return Source->GetWorld();
+}
+
+UObject* FPCGRuntimeGenExecutionState::GetTarget() const
+{
+	check(Source);
+	return Source->GetTarget();
+}
+
+bool FPCGRuntimeGenExecutionState::IsActive() const
+{
+	check(Source);
+	return Source->IsActive();
+}
+
+bool FPCGRuntimeGenExecutionState::HasAuthority() const
+{
+	check(Source);
+	return Source->HasAuthority();
+}
+
+FTransform FPCGRuntimeGenExecutionState::GetTransform() const
+{
+	check(Source);
+	return Source->GetTransform();
+}
+
+FTransform FPCGRuntimeGenExecutionState::GetOriginalTransform() const
+{
+	check(Source);
+	return Source->GetOriginalTransform();
+}
+
+FBox FPCGRuntimeGenExecutionState::GetBounds() const
+{
+	check(Source);
+	return Source->GetBounds();
+}
+
+FBox FPCGRuntimeGenExecutionState::GetOriginalBounds() const
+{
+	check(Source);
+	return Source->GetOriginalBounds();
+}
+
+FBox FPCGRuntimeGenExecutionState::GetLocalSpaceBounds() const
+{
+	check(Source);
+	return Source->GetLocalSpaceBounds();
+}
+
+FBox FPCGRuntimeGenExecutionState::GetOriginalLocalSpaceBounds() const
+{
+	check(Source);
+	return Source->GetOriginalLocalSpaceBounds();
+}
+
+UPCGGraph* FPCGRuntimeGenExecutionState::GetGraph() const
+{
+	check(Source);
+	return Source->GetGraph();
+}
+
+UPCGGraphInstance* FPCGRuntimeGenExecutionState::GetGraphInstance() const
+{
+	check(Source);
+	return Source->GetGraphInstance();
+}
+
+void FPCGRuntimeGenExecutionState::Cancel()
+{
+	check(Source);
+	Source->Cancel();
+}
+
+void FPCGRuntimeGenExecutionState::OnGraphExecutionAborted(bool bQuiet, bool bCleanupUnusedResources)
+{
+	check(Source);
+	Source->OnGraphExecutionAborted(bQuiet, bCleanupUnusedResources);
+}
+
+bool FPCGRuntimeGenExecutionState::IsGenerating() const
+{
+	check(Source);
+	return Source->IsGenerating();
+}
+
+bool FPCGRuntimeGenExecutionState::IsGenerated() const
+{
+	check(Source);
+	return Source->IsGenerated();
+}
+
+FPCGTaskId FPCGRuntimeGenExecutionState::GetGenerationTaskId() const
+{
+	check(Source);
+	return Source->GetGenerationTaskId();
+}
+
+FPCGTaskId FPCGRuntimeGenExecutionState::Generate(const FGenerateParams& InGenerateParams)
+{
+	check(Source);
+	return Source->Generate(InGenerateParams);
+}
+
+FPCGTaskId FPCGRuntimeGenExecutionState::Cleanup(const FCleanupParams& InCleanupParams)
+{
+	check(Source);
+	return Source->Cleanup(InCleanupParams);
+}
+
+void FPCGRuntimeGenExecutionState::StoreOutputDataForPin(const FString& InResourceKey, const FPCGDataCollection& InData) const
+{
+	check(Source);
+	Source->StoreOutputDataForPin(InResourceKey, InData);
+}
+
+const FPCGDataCollection* FPCGRuntimeGenExecutionState::RetrieveOutputDataForPin(const FString& InResourceKey) const
+{
+	check(Source);
+	return Source->RetrieveOutputDataForPin(InResourceKey);
+}
+
+FPCGManagedResourceContainer* FPCGRuntimeGenExecutionState::GetManagedResourceContainer()
+{
+	check(Source);
+	return Source->GetManagedResourceContainer();
+}
+
+FTransactionallySafeCriticalSection* FPCGRuntimeGenExecutionState::GetManagedResourceContainerLock()
+{
+	check(Source);
+	return Source->GetManagedResourceContainerLock();
+}
+
+#if WITH_EDITOR
+const PCGUtils::FExtraCapture& FPCGRuntimeGenExecutionState::GetExtraCapture() const
+{
+	check(Source);
+	return Source->GetExtraCapture();
+}
+
+PCGUtils::FExtraCapture& FPCGRuntimeGenExecutionState::GetExtraCapture()
+{
+	check(Source);
+	return Source->GetExtraCapture();
+}
+
+const FPCGGraphExecutionInspection& FPCGRuntimeGenExecutionState::GetInspection() const
+{
+	check(Source);
+	return Source->GetInspection();
+}
+
+FPCGGraphExecutionInspection& FPCGRuntimeGenExecutionState::GetInspection()
+{
+	check(Source);
+	return Source->GetInspection();
+}
+
+void FPCGRuntimeGenExecutionState::RegisterDynamicTracking(const UPCGSettings* InSettings, const TArrayView<TPair<FPCGSelectionKey, bool>>& InDynamicKeysAndCulling)
+{
+	check(Source);
+	Source->RegisterDynamicTracking(InSettings, InDynamicKeysAndCulling);
+}
+
+FPCGDynamicTrackingPriority FPCGRuntimeGenExecutionState::GetDynamicTrackingPriority() const
+{
+	check(Source);
+	return Source->GetDynamicTrackingPriority();
+}
+
+void FPCGRuntimeGenExecutionState::RegisterDynamicTracking(const FPCGSelectionKeyToSettingsMap& InKeysToSettings)
+{
+	check(Source);
+	Source->RegisterDynamicTracking(InKeysToSettings);
+}
+
+TArray<FPCGSelectionKey> FPCGRuntimeGenExecutionState::GatherTrackingKeys() const
+{
+	check(Source);
+	return Source->GatherTrackingKeys();
+}
+
+bool FPCGRuntimeGenExecutionState::IsKeyTrackedAndCulled(const FPCGSelectionKey& Key, bool& bOutIsCulled) const
+{
+	check(Source);
+	return Source->IsKeyTrackedAndCulled(Key, bOutIsCulled);
+}
+
+void FPCGRuntimeGenExecutionState::ForEachSettingTrackingKey(const FPCGSelectionKey& InKey, const TFunctionRef<void(const FPCGSelectionKey&, const FPCGSettingsAndCulling&)> InCallback) const
+{
+	check(Source);
+	Source->ForEachSettingTrackingKey(InKey, InCallback);
+}
+
+void FPCGRuntimeGenExecutionState::AddLocalSourceDynamicTrackingKeys(const FPCGSelectionKeyToSettingsMap& LocalSourceDynamicTrackingKeys)
+{
+	check(Source);
+	Source->AddLocalSourceDynamicTrackingKeys(LocalSourceDynamicTrackingKeys);
+}
+
+void FPCGRuntimeGenExecutionState::AddLocalSourceCurrentTrackingKeys(const FPCGSelectionKeyToSettingsMap& LocalSourceCurrentExecutionTrackingKeys, const TSet<const UPCGSettings*>& LocalSourceCurrentExecutionTrackingSettings)
+{
+	check(Source);
+	Source->AddLocalSourceCurrentTrackingKeys(LocalSourceCurrentExecutionTrackingKeys, LocalSourceCurrentExecutionTrackingSettings);
+}
+
+bool FPCGRuntimeGenExecutionState::IsRefreshInProgress() const
+{
+	check(Source);
+	return Source->IsRefreshInProgress();
+}
+#endif // WITH_EDITOR
+
+IPCGGraphExecutionSource* FPCGRuntimeGenExecutionState::GetOriginalSource() const
+{
+	check(Source);
+	return Source->GetOriginalSource();
+}
+
+IPCGGraphExecutionSource* FPCGRuntimeGenExecutionState::GetLocalSource(const FPCGGridDescriptor& GridDescriptor, const FIntVector& CellCoords) const
+{
+	check(Source);
+	return Source->GetLocalSource(GridDescriptor, CellCoords);
+}
+
+bool FPCGRuntimeGenExecutionState::IsPartitioned() const
+{
+	check(Source);
+	return Source->IsPartitioned();
+}
+
+bool FPCGRuntimeGenExecutionState::IsLocalSource() const
+{
+	check(Source);
+	return Source->IsLocalSource();
+}
+
+bool FPCGRuntimeGenExecutionState::Use2DGrid() const
+{
+	check(Source);
+	return Source->Use2DGrid();
+}
+
+EPCGHiGenGrid FPCGRuntimeGenExecutionState::GetGenerationGrid() const
+{
+	check(Source);
+	return Source->GetGenerationGrid();
+}
+
+uint32 FPCGRuntimeGenExecutionState::GetGenerationGridSize() const
+{
+	check(Source);
+	return Source->GetGenerationGridSize();
+}
+
+FIntVector FPCGRuntimeGenExecutionState::GetGenerationGridCoords() const
+{
+	check(Source);
+	return Source->GetGenerationGridCoords();
+}
+
+bool FPCGRuntimeGenExecutionState::IsManagedByRuntimeGenSystem() const
+{
+	check(Source);
+	return Source->IsManagedByRuntimeGenSystem();
+}
+
+const UPCGSchedulingPolicyBase* FPCGRuntimeGenExecutionState::GetRuntimeGenSchedulingPolicy() const
+{
+	check(Source);
+	return Source->GetRuntimeGenSchedulingPolicy();
+}
+
+const FPCGRuntimeGenerationRadii& FPCGRuntimeGenExecutionState::GetGenerationRadii() const
+{
+	check(Source);
+	return Source->GetGenerationRadii();
+}
+
+double FPCGRuntimeGenExecutionState::GetGenerationRadiusFromGrid(EPCGHiGenGrid InGrid) const
+{
+	check(Source);
+	return Source->GetGenerationRadiusFromGrid(InGrid);
+}
+
+double FPCGRuntimeGenExecutionState::GetCleanupRadiusFromGrid(EPCGHiGenGrid InGrid) const
+{
+	check(Source);
+	return Source->GetCleanupRadiusFromGrid(InGrid);
+}
+
+void FPCGRuntimeGenExecutionState::SetSource(UPCGRuntimeGenExecutionSource* InSource)
+{
+	Source = InSource;
+}
diff --git a/Engine/Plugins/PCG/Source/PCG/Public/PCGComponentExecutionState.h b/Engine/Plugins/PCG/Source/PCG/Public/PCGComponentExecutionState.h
--- a/Engine/Plugins/PCG/Source/PCG/Public/PCGComponentExecutionState.h
+++ b/Engine/Plugins/PCG/Source/PCG/Public/PCGComponentExecutionState.h
@@ -18,6 +18,7 @@ public:
 	virtual FTransform GetOriginalTransform() const override;
 	virtual UWorld* GetWorld() const override;
 	virtual UObject* GetTarget() const override;
+	virtual bool IsActive() const override;
 	virtual bool HasAuthority() const override;
 	virtual FBox GetBounds() const override;
 	virtual FBox GetOriginalBounds() const override;
@@ -33,6 +34,10 @@ public:
 	virtual FBox GetGeneratedBounds() const override;
 	virtual void ExecutePreGraph(FPCGContext* InContext) override;
 	virtual bool IsManagedByRuntimeGenSystem() const override;
+	virtual const UPCGSchedulingPolicyBase* GetRuntimeGenSchedulingPolicy() const override;
+	virtual const FPCGRuntimeGenerationRadii& GetGenerationRadii() const override;
+	virtual double GetGenerationRadiusFromGrid(EPCGHiGenGrid InGrid) const override;
+	virtual double GetCleanupRadiusFromGrid(EPCGHiGenGrid InGrid) const override;
 	virtual IPCGGraphExecutionSource* GetOriginalSource() const override;
 	virtual IPCGGraphExecutionSource* GetLocalSource(const FPCGGridDescriptor& GridDescriptor, const FIntVector& CellCoords) const override;
 	virtual bool IsPartitioned() const override;
@@ -40,6 +45,7 @@ public:
 	virtual FPCGGridDescriptor GetGridDescriptor(uint32 InGridSize) const override;
 	virtual EPCGHiGenGrid GetGenerationGrid() const override;
 	virtual uint32 GetGenerationGridSize() const override;
+	virtual FIntVector GetGenerationGridCoords() const override;
 	virtual bool IsInPreviewMode() const override;
 	virtual FPCGTaskId GetGenerationTaskId() const override;
 	virtual void StoreOutputDataForPin(const FString& InResourceKey, const FPCGDataCollection& InData) const override;
diff --git a/Engine/Plugins/PCG/Source/PCG/Public/PCGGraphExecutionStateInterface.h b/Engine/Plugins/PCG/Source/PCG/Public/PCGGraphExecutionStateInterface.h
--- a/Engine/Plugins/PCG/Source/PCG/Public/PCGGraphExecutionStateInterface.h
+++ b/Engine/Plugins/PCG/Source/PCG/Public/PCGGraphExecutionStateInterface.h
@@ -18,8 +18,10 @@ class UPCGData;
 class UPCGGraph;
 class UPCGGraphInstance;
 class UPCGManagedResource;
+class UPCGSchedulingPolicyBase;
 struct FPCGGridDescriptor;
 struct FPCGManagedResourceContainer;
+struct FPCGRuntimeGenerationRadii;
 
 namespace PCGUtils
 {
@@ -84,6 +86,9 @@ public:
 		return Cast<TargetType>(GetTarget());
 	}
 
+	/** True if the ExecutionState is active. */
+	virtual bool IsActive() const { return true; }
+
 	/** Returns true if the ExecutionState has network authority */
 	virtual bool HasAuthority() const = 0;
 
@@ -212,10 +217,6 @@ public:
 	/** Returns true if this execution source is a local source of a partitioned execution source. */
 	virtual bool IsLocalSource() const { return false; }
 
-	// @todo_pcg: The execution source should probably just have an execution domain instead of directly querying runtime gen vs etc.
-	/** Returns true if the execution source is managed by the runtime generation system. Nothing else should generate or cleanup this execution source. */
-	virtual bool IsManagedByRuntimeGenSystem() const { return false; }
-
 	/** Returns true if execution source should output on a 2D grid. */
 	virtual bool Use2DGrid() const { return false; }
 
@@ -228,6 +229,25 @@ public:
 	/** Get the size for the HiGen grid this execution source executes on. */
 	virtual uint32 GetGenerationGridSize() const { return PCGHiGenGrid::UninitializedGridSize(); }
 
+	/** Get the grid cell coordinates this execution source executes on. */
+	virtual FIntVector GetGenerationGridCoords() const { return FIntVector::ZeroValue; }
+
+	// @todo_pcg: The execution source should probably just have an execution domain instead of directly querying runtime gen vs etc.
+	/** Returns true if the execution source is managed by the runtime generation system. Nothing else should generate or cleanup this execution source. */
+	virtual bool IsManagedByRuntimeGenSystem() const { return false; }
+
+	/** Gets the policy which dictates behavior when using runtime generation. */
+	virtual const UPCGSchedulingPolicyBase* GetRuntimeGenSchedulingPolicy() const { return nullptr; }
+
+	/** Gets the radii inside of which this execution source will be generated when using runtime generation. */
+	PCG_API virtual const FPCGRuntimeGenerationRadii& GetGenerationRadii() const;
+
+	/** Gets the radius inside of which this execution source will be generated when using runtime generation. */
+	virtual double GetGenerationRadiusFromGrid(EPCGHiGenGrid InGrid) const { return 0.0f; }
+
+	/** Gets the radius outside of which this execution source will be cleaned up when using runtime generation. */
+	virtual double GetCleanupRadiusFromGrid(EPCGHiGenGrid InGrid) const { return 0.0f; }
+
 	/** Returns true if the generated artifacts are transient */
 	virtual bool IsInPreviewMode() const { return false; }
 };
