[PCG] Bughawk fix to PCGRuntimeGenerationExecutionSource.cpp

#rnx
#Synced-CL 49994026

diff --git a/Engine/Plugins/PCG/Source/PCG/Private/RuntimeGen/PCGRuntimeGenExecutionSource.cpp b/Engine/Plugins/PCG/Source/PCG/Private/RuntimeGen/PCGRuntimeGenExecutionSource.cpp
--- a/Engine/Plugins/PCG/Source/PCG/Private/RuntimeGen/PCGRuntimeGenExecutionSource.cpp
+++ b/Engine/Plugins/PCG/Source/PCG/Private/RuntimeGen/PCGRuntimeGenExecutionSource.cpp
@@ -225,7 +225,7 @@ void UPCGRuntimeGenExecutionSource::OnGraphExecutionAborted(bool bQuiet, bool bC
 #if WITH_EDITOR
 	if (IPCGBaseSubsystem* Subsystem = GetExecutionState().GetSubsystem())
 	{
-		Subsystem->OnPCGSourceGenerationDone(this, EPCGGenerationStatus::Completed);
+		Subsystem->OnPCGSourceGenerationDone(this, EPCGGenerationStatus::Aborted);
 	}
 #endif
 }
