[PCG] Fixed replace tags node test because of different compact set behavior.

#rnx

#Synced-CL 49997914

diff --git a/Engine/Plugins/PCG/Source/PCG/Private/Tests/Elements/PCGReplaceTagTest.cpp b/Engine/Plugins/PCG/Source/PCG/Private/Tests/Elements/PCGReplaceTagTest.cpp
--- a/Engine/Plugins/PCG/Source/PCG/Private/Tests/Elements/PCGReplaceTagTest.cpp
+++ b/Engine/Plugins/PCG/Source/PCG/Private/Tests/Elements/PCGReplaceTagTest.cpp
@@ -50,7 +50,7 @@ public:
 			const TSet<FString> OutputTagSet = Context->OutputData.TaggedData[0].Tags;
 
 			UTEST_EQUAL("Tag count", OutputTagSet.Num(), Parameters.ExpectedOutputTags.Num());
-			UTEST_TRUE("Output Tags", OutputTagSet.Array() == Parameters.ExpectedOutputTags.Array());
+			UTEST_TRUE("Output Tags", OutputTagSet.Includes(Parameters.ExpectedOutputTags) && Parameters.ExpectedOutputTags.Includes(OutputTagSet));
 		}
 
 		return true;
