[PCG] Fix flushing of instances during global component reregister

#rnx
#Synced-CL 49994026

diff --git a/Engine/Plugins/PCG/Source/PCG/Private/Subsystems/PCGSubsystem.cpp b/Engine/Plugins/PCG/Source/PCG/Private/Subsystems/PCGSubsystem.cpp
--- a/Engine/Plugins/PCG/Source/PCG/Private/Subsystems/PCGSubsystem.cpp
+++ b/Engine/Plugins/PCG/Source/PCG/Private/Subsystems/PCGSubsystem.cpp
@@ -19,6 +19,7 @@
 #include "RuntimeGen/PCGGenSourceManager.h"
 #include "RuntimeGen/PCGRuntimeGenScheduler.h"
 
+#include "ComponentReregisterContext.h"
 #include "Algo/Transform.h"
 #include "Engine/Level.h"
 
@@ -210,6 +211,8 @@ void UPCGSubsystem::Deinitialize()
 
 	TrackingManager->Deinitialize();
 
+	FGlobalComponentReregisterContext::OnGlobalComponentReregisterContextDestroyed().RemoveAll(this);
+
 #if WITH_EDITOR
 	FCoreUObjectDelegates::OnObjectsReplaced.RemoveAll(this);
 
@@ -278,6 +281,8 @@ void UPCGSubsystem::PostInitialize()
 		FWorldDelegates::LevelRemovedFromWorld.AddUObject(this, &UPCGSubsystem::OnLevelRemovedFromWorld);
 	}
 #endif
+
+	FGlobalComponentReregisterContext::OnGlobalComponentReregisterContextDestroyed().AddUObject(this, &UPCGSubsystem::OnGlobalComponentReregisterContextDestroyed);
 }
 
 UPCGSubsystem* UPCGSubsystem::GetInstance(UWorld* World)
@@ -2172,6 +2177,11 @@ void UPCGSubsystem::ExecuteBeginTickActions()
 	}
 }
 
+void UPCGSubsystem::OnGlobalComponentReregisterContextDestroyed()
+{
+	RefreshAllRuntimeGenComponents();
+}
+
 void UPCGSubsystem::RegisterPartitionActor(APCGPartitionActor* InActor)
 {
 	TrackingManager->RegisterPartitionActor(InActor);
diff --git a/Engine/Plugins/PCG/Source/PCG/Public/Subsystems/PCGSubsystem.h b/Engine/Plugins/PCG/Source/PCG/Public/Subsystems/PCGSubsystem.h
--- a/Engine/Plugins/PCG/Source/PCG/Public/Subsystems/PCGSubsystem.h
+++ b/Engine/Plugins/PCG/Source/PCG/Public/Subsystems/PCGSubsystem.h
@@ -416,6 +416,8 @@ private:
 
 	void ExecuteBeginTickActions();
 
+	void OnGlobalComponentReregisterContextDestroyed();
+
 	APCGWorldActor* PCGWorldActor = nullptr;
 	FPCGRuntimeGenScheduler* RuntimeGenScheduler = nullptr;
 	bool bHasTickedOnce = false;
