[PCG] Refresh when scalability quality levels change

#Synced-CL 49994026

#Synced-CL 49994026

diff --git a/Engine/Plugins/PCG/Source/PCG/Private/Subsystems/PCGSubsystem.cpp b/Engine/Plugins/PCG/Source/PCG/Private/Subsystems/PCGSubsystem.cpp
--- a/Engine/Plugins/PCG/Source/PCG/Private/Subsystems/PCGSubsystem.cpp
+++ b/Engine/Plugins/PCG/Source/PCG/Private/Subsystems/PCGSubsystem.cpp
@@ -211,6 +211,7 @@ void UPCGSubsystem::Deinitialize()
 
 	TrackingManager->Deinitialize();
 
+	Scalability::OnScalabilitySettingsChanged.RemoveAll(this);
 	FGlobalComponentReregisterContext::OnGlobalComponentReregisterContextDestroyed().RemoveAll(this);
 
 #if WITH_EDITOR
@@ -283,6 +284,8 @@ void UPCGSubsystem::PostInitialize()
 #endif
 
 	FGlobalComponentReregisterContext::OnGlobalComponentReregisterContextDestroyed().AddUObject(this, &UPCGSubsystem::OnGlobalComponentReregisterContextDestroyed);
+
+	Scalability::OnScalabilitySettingsChanged.AddUObject(this, &UPCGSubsystem::OnScalabilitySettingsChanged);
 }
 
 UPCGSubsystem* UPCGSubsystem::GetInstance(UWorld* World)
@@ -480,6 +483,11 @@ void UPCGSubsystem::OnPCGQualityLevelChanged()
 	});
 }
 
+void UPCGSubsystem::OnScalabilitySettingsChanged(const Scalability::FQualityLevels& QualityLevels)
+{
+	RefreshAllRuntimeGenComponents(EPCGChangeType::GenerationGrid);
+}
+
 #if WITH_EDITOR
 void UPCGSubsystem::DestroyCurrentPCGWorldActor()
 {
diff --git a/Engine/Plugins/PCG/Source/PCG/Public/Subsystems/PCGSubsystem.h b/Engine/Plugins/PCG/Source/PCG/Public/Subsystems/PCGSubsystem.h
--- a/Engine/Plugins/PCG/Source/PCG/Public/Subsystems/PCGSubsystem.h
+++ b/Engine/Plugins/PCG/Source/PCG/Public/Subsystems/PCGSubsystem.h
@@ -10,6 +10,7 @@
 #include "Grid/PCGGridDescriptor.h"
 #include "Utils/PCGNodeVisualLogs.h"
 
+#include "Scalability.h"
 #include "UObject/ObjectKey.h"
 
 #include "PCGSubsystem.generated.h"
@@ -132,6 +133,8 @@ public:
 	static PCG_API int32 GetPCGQualityLevel();
 	PCG_API void OnPCGQualityLevelChanged();
 
+	void OnScalabilitySettingsChanged(const Scalability::FQualityLevels& QualityLevels);
+
 #if WITH_EDITOR
 	PCG_API void DestroyAllPCGWorldActors();
 	PCG_API void DestroyCurrentPCGWorldActor();
