[PCG] Support Actor-Componentless runtime generation
* Actor-Componentless generation mode can be enabled on the PCG Graph or overridden on the PCG Component. Can only be used with runtime generation.
* Generalize several areas in the plugin to support execution sources instead of only PCG Components.

#Synced-CL 49994026

diff --git a/Engine/Plugins/PCG/Source/PCG/Internal/PCGTrackingManager.h b/Engine/Plugins/PCG/Source/PCG/Internal/PCGTrackingManager.h
--- a/Engine/Plugins/PCG/Source/PCG/Internal/PCGTrackingManager.h
+++ b/Engine/Plugins/PCG/Source/PCG/Internal/PCGTrackingManager.h
@@ -3,6 +3,7 @@
 #pragma once
 
 #include "PCGCommon.h"
+#include "ChangeTracking/PCGChangeTrackingRegistry.h"
 #include "Elements/PCGActorSelector.h"
 #include "Grid/PCGGridDescriptor.h"
 #include "Grid/PCGExecutionSourceOctree.h"
diff --git a/Engine/Plugins/PCG/Source/PCG/Internal/RuntimeGen/PCGRuntimeGenChangeDetection.h b/Engine/Plugins/PCG/Source/PCG/Internal/RuntimeGen/PCGRuntimeGenChangeDetection.h
--- a/Engine/Plugins/PCG/Source/PCG/Internal/RuntimeGen/PCGRuntimeGenChangeDetection.h
+++ b/Engine/Plugins/PCG/Source/PCG/Internal/RuntimeGen/PCGRuntimeGenChangeDetection.h
@@ -3,25 +3,24 @@
 #pragma once
 
 #include "PCGCommon.h"
-#include "PCGComponent.h"
+#include "RuntimeGen/PCGRuntimeGenExecutionSource.h"
 
 #include "UObject/ObjectKey.h"
 
 class IPCGGenSourceBase;
-class UPCGComponent;
 
 namespace PCGRuntimeGenChangeDetection
 {
-	using FDetectionInputsBase = TTuple<TObjectKey<UPCGComponent>, const IPCGGenSourceBase*, FPCGRuntimeGenerationRadii, /*RadiusMultiplier*/float, /*UseGenSourceDirection*/bool, /*MinGridSize*/uint32>;
+	using FDetectionInputsBase = TTuple</*OriginalExecutionSource*/TObjectKey<UObject>, const IPCGGenSourceBase*, FPCGRuntimeGenerationRadii, /*RadiusMultiplier*/float, /*UseGenSourceDirection*/bool, /*MinGridSize*/uint32>;
 
 	/** The key used for the state cache. */
 	struct FDetectionInputs : FDetectionInputsBase
 	{
-		FDetectionInputs(const UPCGComponent* InOriginalComponent, const IPCGGenSourceBase* InGenSource, const FPCGRuntimeGenerationRadii& InGenerationRadii, float InRadiusMultiplier, bool bInUseGenSourceDirection, uint32 InMinGridSize)
-			: FDetectionInputsBase(InOriginalComponent, InGenSource, InGenerationRadii, InRadiusMultiplier, bInUseGenSourceDirection, InMinGridSize)
+		FDetectionInputs(const IPCGGraphExecutionSource* InOriginalExecutionSource, const IPCGGenSourceBase* InGenSource, const FPCGRuntimeGenerationRadii& InGenerationRadii, float InRadiusMultiplier, bool bInUseGenSourceDirection, uint32 InMinGridSize)
+			: FDetectionInputsBase(Cast<UObject>(InOriginalExecutionSource), InGenSource, InGenerationRadii, InRadiusMultiplier, bInUseGenSourceDirection, InMinGridSize)
 		{}
 
-		const UPCGComponent* GetOriginalComponent() const { return Get<0>().ResolveObjectPtr(); }
+		const IPCGGraphExecutionSource* GetOriginalExecutionSource() const { return Cast<IPCGGraphExecutionSource>(Get<0>().ResolveObjectPtr()); }
 		const IPCGGenSourceBase* GetGenSource() const { return Get<1>(); }
 		const FPCGRuntimeGenerationRadii& GetGenerationRadii() const { return Get<2>(); }
 		float GetGenerationRadiusMultiplier() const { return Get<3>(); }
@@ -54,7 +53,7 @@ namespace PCGRuntimeGenChangeDetection
 		void OnCellsScanned(const FDetectionInputs& InInputs);
 
 		/** Flush any recorded state, resulting in a rescan of the component on next tick. If nullptr provided all components will rescan. */
-		void FlushCachedState(const UPCGComponent* InOptionalOriginalComponent = nullptr);
+		void FlushCachedState(const IPCGGraphExecutionSource* InOptionalOriginalExecutionSource = nullptr);
 
 	private:
 		TMap<FDetectionInputs, FCachedState> StateCache;
diff --git a/Engine/Plugins/PCG/Source/PCG/Internal/RuntimeGen/PCGRuntimeGenExecutionSource.h b/Engine/Plugins/PCG/Source/PCG/Internal/RuntimeGen/PCGRuntimeGenExecutionSource.h
--- a/Engine/Plugins/PCG/Source/PCG/Internal/RuntimeGen/PCGRuntimeGenExecutionSource.h
+++ b/Engine/Plugins/PCG/Source/PCG/Internal/RuntimeGen/PCGRuntimeGenExecutionSource.h
@@ -77,6 +77,7 @@ public:
 	void AddLocalSourceDynamicTrackingKeys(const FPCGSelectionKeyToSettingsMap& LocalSourceDynamicTrackingKeys);
 	void AddLocalSourceCurrentTrackingKeys(const FPCGSelectionKeyToSettingsMap& LocalSourceCurrentExecutionTrackingKeys, const TSet<const UPCGSettings*>& LocalSourceCurrentExecutionTrackingSettings);
 	bool IsRefreshInProgress() const;
+	bool WasGeneratedThisSession() const;
 #endif
 
 	IPCGGraphExecutionSource* GetOriginalSource() const;
diff --git a/Engine/Plugins/PCG/Source/PCG/Internal/RuntimeGen/PCGRuntimeGenExecutionState.h b/Engine/Plugins/PCG/Source/PCG/Internal/RuntimeGen/PCGRuntimeGenExecutionState.h
--- a/Engine/Plugins/PCG/Source/PCG/Internal/RuntimeGen/PCGRuntimeGenExecutionState.h
+++ b/Engine/Plugins/PCG/Source/PCG/Internal/RuntimeGen/PCGRuntimeGenExecutionState.h
@@ -51,6 +51,7 @@ public:
 	virtual void AddLocalSourceDynamicTrackingKeys(const FPCGSelectionKeyToSettingsMap& LocalSourceDynamicTrackingKeys) override;
 	virtual void AddLocalSourceCurrentTrackingKeys(const FPCGSelectionKeyToSettingsMap& LocalSourceCurrentExecutionTrackingKeys, const TSet<const UPCGSettings*>& LocalSourceCurrentExecutionTrackingSettings) override;
 	virtual bool IsRefreshInProgress() const override;
+	virtual bool WasGeneratedThisSession() const override;
 #endif
 
 	virtual IPCGGraphExecutionSource* GetOriginalSource() const override;
diff --git a/Engine/Plugins/PCG/Source/PCG/Private/Data/Registry/PCGGetDataFunctionRegistry.cpp b/Engine/Plugins/PCG/Source/PCG/Private/Data/Registry/PCGGetDataFunctionRegistry.cpp
--- a/Engine/Plugins/PCG/Source/PCG/Private/Data/Registry/PCGGetDataFunctionRegistry.cpp
+++ b/Engine/Plugins/PCG/Source/PCG/Private/Data/Registry/PCGGetDataFunctionRegistry.cpp
@@ -7,6 +7,7 @@
 #include "PCGContext.h"
 #include "PCGData.h"
 #include "PCGGraph.h"
+#include "PCGGraphExecutionStateInterface.h"
 #include "Compute/PCGComputeCommon.h"
 #include "Data/PCGBasePointData.h"
 #include "Data/PCGCollisionShapeData.h"
@@ -39,6 +40,8 @@
 #include "GameFramework/Volume.h"
 #include "VT/RuntimeVirtualTexture.h"
 
+#define LOCTEXT_NAMESPACE "PCGGetDataFunctionRegistry"
+
 FPCGGetDataFunctionRegistry::FFunctionHandle FPCGGetDataFunctionRegistry::RegisterDataFromActorFunction(const FPCGGetDataFunctionRegistry::FDataFromActorFunction& InFunction)
 {
 	check(IsInGameThread());
@@ -107,7 +110,7 @@ int FPCGGetDataFunctionRegistry::DefaultDataFromActor(FPCGContext* InContext, co
 
 	const EPCGDataType InDataFilter = InParams.DataTypeFilter;
 	const bool bParseActor = InParams.bParseActor;
-	const UPCGComponent* Component = InParams.SourceComponent;
+	const IPCGGraphExecutionSource* Source = InParams.Source;
 
 	auto NameTagsToStringTags = [](const FName& InName) { return InName.ToString(); };
 	TSet<FString> ActorTags;
@@ -130,7 +133,15 @@ int FPCGGetDataFunctionRegistry::DefaultDataFromActor(FPCGContext* InContext, co
 	}
 	else if (PartitionActor)
 	{
-		check(!Component || Component->GetOwner() == Actor); // Invalid processing otherwise because of the this usage
+		// If we have a Partition Actor, the execution source must be a PCG Component.
+		const UPCGComponent* Component = Cast<UPCGComponent>(Source);
+
+		// Invalid processing otherwise because of the this usage
+		if (!Component || Component->GetOwner() != Actor)
+		{
+			PCGLog::LogErrorOnGraph(FText::Format(LOCTEXT("InvalidComponent", "Could not get data from actor {0}, execution source is invalid."), FText::FromString(PartitionActor->GetActorNameOrLabel())), InContext);
+			return 0;
+		}
 
 		UPCGVolumeData* VolumeData = FPCGContext::NewObject_AnyThread<UPCGVolumeData>(InContext);
 		UPCGSpatialData* Result = VolumeData;
@@ -162,14 +173,14 @@ int FPCGGetDataFunctionRegistry::DefaultDataFromActor(FPCGContext* InContext, co
 	else if (LandscapeActor)
 	{
 		UPCGLandscapeData* Data = FPCGContext::NewObject_AnyThread<UPCGLandscapeData>(InContext);
-		const UPCGGraph* PCGGraph = Component ? Component->GetGraph() : nullptr;
+		const UPCGGraph* PCGGraph = Source ? Source->GetExecutionState().GetGraph() : nullptr;
 
 		FPCGLandscapeDataProps LandscapeDataProps;
 		LandscapeDataProps.bGetHeightOnly = false;
 		LandscapeDataProps.bGetLayerWeights = (!PCGGraph || PCGGraph->bLandscapeUsesMetadata);
 		LandscapeDataProps.bSampleVirtualTextures = false;
 
-		Data->Initialize({ LandscapeActor }, PCGHelpers::GetGridBounds(Actor, Component), LandscapeDataProps);
+		Data->Initialize({ LandscapeActor }, PCGHelpers::GetGridBounds(Actor, Source), LandscapeDataProps);
 
 		FPCGTaggedData& TaggedData = Output.Collection.TaggedData.Emplace_GetRef();
 		TaggedData.Data = Data;
@@ -329,4 +340,6 @@ int FPCGGetDataFunctionRegistry::DefaultDataFromComponent(FPCGContext* InContext
 	{
 		return 0;
 	}
-}
\ No newline at end of file
+}
+
+#undef LOCTEXT_NAMESPACE
\ No newline at end of file
diff --git a/Engine/Plugins/PCG/Source/PCG/Private/Elements/PCGActorSelector.cpp b/Engine/Plugins/PCG/Source/PCG/Private/Elements/PCGActorSelector.cpp
--- a/Engine/Plugins/PCG/Source/PCG/Private/Elements/PCGActorSelector.cpp
+++ b/Engine/Plugins/PCG/Source/PCG/Private/Elements/PCGActorSelector.cpp
@@ -800,7 +800,10 @@ void FPCGComponentSelectorSettings::PrepareForFiltering(bool bForce) const
 
 bool FPCGComponentSelectorSettings::FilterComponent(UActorComponent* InComponent) const
 {
-	check(InComponent);
+	if (!InComponent)
+	{
+		return true;
+	}
 
 	if (!ComponentList.IsEmpty() && !ComponentList.Contains(InComponent))
 	{
diff --git a/Engine/Plugins/PCG/Source/PCG/Private/Elements/PCGDataFromActor.cpp b/Engine/Plugins/PCG/Source/PCG/Private/Elements/PCGDataFromActor.cpp
--- a/Engine/Plugins/PCG/Source/PCG/Private/Elements/PCGDataFromActor.cpp
+++ b/Engine/Plugins/PCG/Source/PCG/Private/Elements/PCGDataFromActor.cpp
@@ -380,7 +380,7 @@ bool FPCGDataFromActorElement::ExecuteInternal(FPCGContext* InContext) const
 			FBox Bounds = ExecutionSource->GetExecutionState().GetBounds();
 			BoundsCheck = [Settings, Bounds, ExecutionSource](const AActor* OtherActor) -> bool
 			{
-				const FBox OtherActorBounds = OtherActor ? PCGHelpers::GetGridBounds(OtherActor, Cast<UPCGComponent>(ExecutionSource)) : FBox(EForceInit::ForceInit);
+				const FBox OtherActorBounds = OtherActor ? PCGHelpers::GetGridBounds(OtherActor, ExecutionSource) : FBox(EForceInit::ForceInit);
 				return OtherActorBounds.IsValid && Bounds.IsValid && Bounds.Intersect(OtherActorBounds);
 			};
 		}
@@ -388,6 +388,7 @@ bool FPCGDataFromActorElement::ExecuteInternal(FPCGContext* InContext) const
 		TFunction<bool(const AActor*)> SelfIgnoreCheck = [](const AActor*) -> bool { return true; };
 		if (Settings->ActorSelector.bIgnoreSelfAndChildren)
 		{
+			// @todo_pcg: RG source should support self ignore check
 			SelfIgnoreCheck = [Self](const AActor* OtherActor) -> bool
 			{
 				// Check if OtherActor is a child of self
@@ -624,7 +625,7 @@ bool FPCGDataFromActorElement::ExecuteInternal(FPCGContext* InContext) const
 	return true;
 }
 
-bool ShouldIgnorePCGComponent(FPCGDataFromActorContext* Context, const UPCGDataFromActorSettings* Settings, UPCGComponent* Component, const UPCGComponent* SourceOriginalComponent)
+bool ShouldIgnorePCGComponent(FPCGDataFromActorContext* Context, const UPCGDataFromActorSettings* Settings, UPCGComponent* Component, const IPCGGraphExecutionSource* DownstreamOriginalSource)
 {
 	check(Context && Settings);
 	if (!Context->ComponentSelector.FilterComponent(Component))
@@ -633,11 +634,11 @@ bool ShouldIgnorePCGComponent(FPCGDataFromActorContext* Context, const UPCGDataF
 	}
 
 	const UPCGComponent* OriginalComponent = Component ? Component->GetOriginalComponent() : nullptr;
-	const AActor* SourceOriginalOwner = SourceOriginalComponent->GetOwner();
+	const UObject* DownstreamOriginalOwner = DownstreamOriginalSource ? DownstreamOriginalSource->GetExecutionState().GetTarget() : nullptr;
 
 	return !OriginalComponent
-		|| OriginalComponent == SourceOriginalComponent
-		|| (Settings->ActorSelector.bIgnoreSelfAndChildren && OriginalComponent->GetOwner() == SourceOriginalOwner);
+		|| OriginalComponent == DownstreamOriginalSource
+		|| (Settings->ActorSelector.bIgnoreSelfAndChildren && OriginalComponent->GetExecutionState().GetTarget() == DownstreamOriginalOwner);
 }
 
 void FPCGDataFromActorElement::GatherWaitTasks(AActor* FoundActor, FPCGContext* InContext, TArray<FPCGTaskId>& OutWaitTasks) const
@@ -653,58 +654,58 @@ void FPCGDataFromActorElement::GatherWaitTasks(AActor* FoundActor, FPCGContext*
 	const UPCGDataFromActorSettings* Settings = Context->GetInputSettings<UPCGDataFromActorSettings>();
 	check(Settings);
 
-	UPCGComponent* SourceComponent = Cast<UPCGComponent>(Context->ExecutionSource.Get());
-	UPCGComponent* SourceOriginalComponent = SourceComponent ? SourceComponent->GetOriginalComponent() : nullptr;
+	const IPCGGraphExecutionSource* Source = Context->ExecutionSource.Get();
+	IPCGGraphExecutionSource* OriginalSource = Source ? Source->GetExecutionState().GetOriginalSource() : nullptr;
 
-	if (!SourceOriginalComponent)
+	if (!OriginalSource)
 	{
 		return;
 	}
 
-	// We will prevent gathering the current execution - this task cannot wait on itself
-	const AActor* SourceOwner = SourceOriginalComponent->GetOwner();
-
-	TInlineComponentArray<UPCGComponent*, 1> PCGComponents = PCGDataFromActorHelpers::GetPCGComponentsFromActor(
+	TInlineComponentArray<UPCGComponent*, 1> UpstreamComponents = PCGDataFromActorHelpers::GetPCGComponentsFromActor(
 		FoundActor,
-		SourceComponent->GetSubsystem(),
+		Cast<UPCGSubsystem>(Source->GetExecutionState().GetSubsystem()),
 		/*bGetLocalComponents=*/true,
 		Settings->bGetDataOnAllGrids,
 		Settings->AllowedGrids,
 		Settings->bComponentsMustOverlapSelf,
-		Settings->bComponentsMustOverlapSelf ? SourceComponent->GetGridBounds() : FBox());
+		Settings->bComponentsMustOverlapSelf ? Source->GetExecutionState().GetBounds() : FBox());
 
-	for (UPCGComponent* Component : PCGComponents)
+	for (UPCGComponent* UpstreamComponent : UpstreamComponents)
 	{
 		// Avoid waiting on our own execution (including local components).
-		if (ShouldIgnorePCGComponent(Context, Settings, Component, SourceOriginalComponent))
+		if (ShouldIgnorePCGComponent(Context, Settings, UpstreamComponent, OriginalSource))
 		{
 			continue;
 		}
 
-		if (Component->IsGenerating())
+		if (UpstreamComponent->IsGenerating())
 		{
-			OutWaitTasks.Add(Component->GetGenerationTaskId());
+			OutWaitTasks.Add(UpstreamComponent->GetGenerationTaskId());
 		}
-		else if (!Component->bGenerated && Component->bActivated && (Component->GetSerializedEditingMode() == EPCGEditorDirtyMode::Preview) && Component->GetOwner())
+		else if (!UpstreamComponent->bGenerated && UpstreamComponent->bActivated && (UpstreamComponent->GetSerializedEditingMode() == EPCGEditorDirtyMode::Preview) && UpstreamComponent->GetOwner())
 		{
 #if WITH_EDITOR
 			// Signal that any change notifications from generating upstream component should not trigger re-executions of this component.
 			// Such change notifications can cancel the current execution.
 			// Note: Uses owner because FPCGTrackingManager::OnPCGGraphGeneratedOrCleaned reports change on owner.
-			SourceOriginalComponent->StartIgnoringChangeOriginDuringGeneration(Component->GetOwner());
-			Context->IgnoredChangeOrigins.Add(Component->GetOwner());
+			if (UPCGComponent* OriginalComponent = Cast<UPCGComponent>(OriginalSource))
+			{
+				OriginalComponent->StartIgnoringChangeOriginDuringGeneration(UpstreamComponent->GetOwner());
+			}
+			Context->IgnoredChangeOrigins.Add(UpstreamComponent->GetOwner());
 #endif
 
-			const FPCGTaskId GenerateTask = Component->GenerateLocalGetTaskId(EPCGComponentGenerationTrigger::GenerateOnDemand, /*bForce=*/false);
+			const FPCGTaskId GenerateTask = UpstreamComponent->GenerateLocalGetTaskId(EPCGComponentGenerationTrigger::GenerateOnDemand, /*bForce=*/false);
 			if (GenerateTask != InvalidPCGTaskId)
 			{
-				PCGGraphExecutionLogging::LogGraphScheduleDependency(Component, Context->GetStack());
+				PCGGraphExecutionLogging::LogGraphScheduleDependency(UpstreamComponent, Context->GetStack());
 
 				OutWaitTasks.Add(GenerateTask);
 			}
 			else
 			{
-				PCGGraphExecutionLogging::LogGraphScheduleDependencyFailed(Component, Context->GetStack());
+				PCGGraphExecutionLogging::LogGraphScheduleDependencyFailed(UpstreamComponent, Context->GetStack());
 			}
 		}
 	}
@@ -910,40 +911,40 @@ void FPCGDataFromActorElement::ProcessActor(FPCGContext* InContext, const UPCGDa
 	FPCGDataFromActorContext* Context = static_cast<FPCGDataFromActorContext*>(InContext);
 	check(Settings);
 
-	UPCGComponent* SourceComponent = Cast<UPCGComponent>(Context->ExecutionSource.Get());
-	const UPCGComponent* SourceOriginalComponent = SourceComponent ? SourceComponent->GetOriginalComponent() : nullptr;
+	const IPCGGraphExecutionSource* Source = Context->ExecutionSource.Get();
+	const IPCGGraphExecutionSource* OriginalSource = Source ? Source->GetExecutionState().GetOriginalSource() : nullptr;
 
-	if (!FoundActor || !IsValid(FoundActor) || !SourceOriginalComponent)
+	if (!FoundActor || !IsValid(FoundActor) || !OriginalSource)
 	{
 		return;
 	}
 
-	const AActor* SourceOwner = SourceOriginalComponent->GetOwner();
-	TInlineComponentArray<UPCGComponent*, 1> PCGComponents;
+	const UObject* SourceOwner = OriginalSource->GetExecutionState().GetTarget();
+	TInlineComponentArray<UPCGComponent*, 1> UpstreamComponents;
 	bool bHasGeneratedPCGData = false;
 	FProperty* FoundProperty = nullptr;
 
 	if (Settings->Mode == EPCGGetDataFromActorMode::GetDataFromPCGComponent || Settings->Mode == EPCGGetDataFromActorMode::GetDataFromPCGComponentOrParseComponents)
 	{
-		PCGComponents = PCGDataFromActorHelpers::GetPCGComponentsFromActor(
+		UpstreamComponents = PCGDataFromActorHelpers::GetPCGComponentsFromActor(
 			FoundActor,
-			SourceComponent->GetSubsystem(),
+			Cast<UPCGSubsystem>(Source->GetExecutionState().GetSubsystem()),
 			/*bGetLocalComponents=*/true,
 			Settings->bGetDataOnAllGrids,
 			Settings->AllowedGrids,
 			Settings->bComponentsMustOverlapSelf,
-			Settings->bComponentsMustOverlapSelf ? SourceComponent->GetGridBounds() : FBox());
+			Settings->bComponentsMustOverlapSelf ? Source->GetExecutionState().GetBounds() : FBox());
 
 		// Remove any PCG components that are filtered from the component selector &
 		// Remove any PCG components that don't belong to an external execution context (i.e. share the same original component), or that share a common root actor.
-		PCGComponents.RemoveAllSwap([Context, Settings, SourceOriginalComponent](UPCGComponent* Component)
+		UpstreamComponents.RemoveAllSwap([Context, Settings, OriginalSource](UPCGComponent* UpstreamComponent)
 		{
-			return ShouldIgnorePCGComponent(Context, Settings, Component, SourceOriginalComponent);
+			return ShouldIgnorePCGComponent(Context, Settings, UpstreamComponent, OriginalSource);
 		});
 
-		for (UPCGComponent* Component : PCGComponents)
+		for (UPCGComponent* UpstreamComponent : UpstreamComponents)
 		{
-			bHasGeneratedPCGData |= !Component->GetGeneratedGraphOutput().TaggedData.IsEmpty();
+			bHasGeneratedPCGData |= !UpstreamComponent->GetGeneratedGraphOutput().TaggedData.IsEmpty();
 		}
 		if (APCGIsolatedActor* IsolatedActor = Cast<APCGIsolatedActor>(FoundActor))
 		{
@@ -961,7 +962,7 @@ void FPCGDataFromActorElement::ProcessActor(FPCGContext* InContext, const UPCGDa
 	// Some additional validation
 	if (Settings->Mode == EPCGGetDataFromActorMode::GetDataFromPCGComponent && !bHasGeneratedPCGData)
 	{
-		if (!PCGComponents.IsEmpty())
+		if (!UpstreamComponents.IsEmpty())
 		{
 			PCGE_LOG(Log, GraphAndLog, FText::Format(LOCTEXT("ActorHasNoGeneratedData", "Actor '{0}' does not have any previously generated data, or all its components were filtered out."), FText::FromName(FoundActor->GetFName())));
 		}
@@ -974,10 +975,10 @@ void FPCGDataFromActorElement::ProcessActor(FPCGContext* InContext, const UPCGDa
 		return;
 	}
 
-	auto GetActorDataCollection = [SourceComponent, Context, Settings](AActor* Actor, EPCGDataType DataType, bool bParseActor) -> FPCGDataCollection
+	auto GetActorDataCollection = [Source, Context, Settings](AActor* Actor, EPCGDataType DataType, bool bParseActor) -> FPCGDataCollection
 	{
 		FPCGGetDataFunctionRegistryParams GDFRParams;
-		GDFRParams.SourceComponent = SourceComponent;
+		GDFRParams.Source = Source;
 		GDFRParams.ComponentSelector = &Context->ComponentSelector;
 		GDFRParams.bParseActor = bParseActor;
 		GDFRParams.DataTypeFilter = DataType;
@@ -998,7 +999,7 @@ void FPCGDataFromActorElement::ProcessActor(FPCGContext* InContext, const UPCGDa
 
 	if (bHasGeneratedPCGData)
 	{
-		for (UPCGComponent* Component : PCGComponents)
+		for (UPCGComponent* Component : UpstreamComponents)
 		{
 			for (const FPCGTaggedData& TaggedData : Component->GetGeneratedGraphOutput().TaggedData)
 			{
diff --git a/Engine/Plugins/PCG/Source/PCG/Private/Helpers/PCGBlueprintHelpers.cpp b/Engine/Plugins/PCG/Source/PCG/Private/Helpers/PCGBlueprintHelpers.cpp
--- a/Engine/Plugins/PCG/Source/PCG/Private/Helpers/PCGBlueprintHelpers.cpp
+++ b/Engine/Plugins/PCG/Source/PCG/Private/Helpers/PCGBlueprintHelpers.cpp
@@ -376,7 +376,7 @@ void UPCGBlueprintHelpers::RefreshPCGRuntimeComponent(UPCGComponent* InComponent
 			PCGSubsystem->FlushCache();
 		}
 
-		PCGSubsystem->RefreshRuntimeGenComponent(InComponent, EPCGChangeType::GenerationGrid);
+		PCGSubsystem->RefreshRuntimeGenExecutionSource(InComponent, EPCGChangeType::GenerationGrid);
 	}
 }
 
diff --git a/Engine/Plugins/PCG/Source/PCG/Private/Helpers/PCGHelpers.cpp b/Engine/Plugins/PCG/Source/PCG/Private/Helpers/PCGHelpers.cpp
--- a/Engine/Plugins/PCG/Source/PCG/Private/Helpers/PCGHelpers.cpp
+++ b/Engine/Plugins/PCG/Source/PCG/Private/Helpers/PCGHelpers.cpp
@@ -117,7 +117,7 @@ namespace PCGHelpers
 		}
 	}
 
-	FBox GetGridBounds(const AActor* Actor, const UPCGComponent* Component)
+	FBox GetGridBounds(const AActor* Actor, const IPCGGraphExecutionSource* Source)
 	{
 		FBox Bounds(EForceInit::ForceInit);
 
@@ -126,12 +126,12 @@ namespace PCGHelpers
 			// First, get the bounds from the partition actor
 			Bounds = PartitionActor->GetFixedBounds();
 
-			const UPCGComponent* OriginalComponent = Component ? PartitionActor->GetOriginalComponent(Component) : nullptr;
-			if (OriginalComponent)
+			const IPCGGraphExecutionSource* OriginalSource = Source ? Source->GetExecutionState().GetOriginalSource() : nullptr;
+			if (OriginalSource)
 			{
-				if (OriginalComponent->GetOwner() != PartitionActor)
+				if (OriginalSource->GetExecutionState().GetTarget() != PartitionActor)
 				{
-					Bounds = Bounds.Overlap(OriginalComponent->GetGridBounds());
+					Bounds = Bounds.Overlap(OriginalSource->GetExecutionState().GetBounds());
 				}
 			}
 		}
diff --git a/Engine/Plugins/PCG/Source/PCG/Private/PCGComponent.cpp b/Engine/Plugins/PCG/Source/PCG/Private/PCGComponent.cpp
--- a/Engine/Plugins/PCG/Source/PCG/Private/PCGComponent.cpp
+++ b/Engine/Plugins/PCG/Source/PCG/Private/PCGComponent.cpp
@@ -1335,6 +1335,16 @@ double UPCGComponent::GetCleanupRadiusFromGrid(EPCGHiGenGrid Grid) const
 	return 0;
 }
 
+bool UPCGComponent::UseActorComponentlessGeneration() const
+{
+	if (bOverrideUseActorComponentlessGeneration)
+	{
+		return bUseActorComponentlessGeneration;
+	}
+
+	return GetGraph() && GetGraph()->UseActorComponentlessGeneration();
+}
+
 bool UPCGComponent::MoveResourcesToNewActor(const FPCGMoveResourceParams& Params, bool bCreateChild)
 {
 	// Don't move resources if we are generating or cleaning up
@@ -2154,7 +2164,7 @@ void UPCGComponent::RefreshAfterGraphChanged(UPCGGraphInterface* InGraph, EPCGCh
 	{
 		if (UPCGSubsystem* Subsystem = GetSubsystem())
 		{
-			Subsystem->RefreshRuntimeGenComponent(this, ChangeType);
+			Subsystem->RefreshRuntimeGenExecutionSource(this, ChangeType);
 		}
 	}
 	else
@@ -2190,7 +2200,7 @@ void UPCGComponent::PreEditChange(FProperty* PropertyAboutToChange)
 			if (UPCGSubsystem* Subsystem = GetSubsystem())
 			{
 				// When toggling off of GenerateAtRuntime, we should flush the RuntimeGenScheduler state for this component.
-				Subsystem->RefreshRuntimeGenComponent(this, EPCGChangeType::GenerationGrid);
+				Subsystem->RefreshRuntimeGenExecutionSource(this, EPCGChangeType::GenerationGrid);
 			}
 
 			// Reset to the the editing mode we were in before entering GenerateAtRuntime mode.
@@ -2204,7 +2214,7 @@ void UPCGComponent::PreEditChange(FProperty* PropertyAboutToChange)
 			if (UPCGSubsystem* Subsystem = GetSubsystem())
 			{
 				// When toggling IsPartitioned, we should proactively flush the RuntimeGenScheduler.
-				Subsystem->RefreshRuntimeGenComponent(this, EPCGChangeType::GenerationGrid);
+				Subsystem->RefreshRuntimeGenExecutionSource(this, EPCGChangeType::GenerationGrid);
 			}
 		}
 	}
@@ -2413,7 +2423,7 @@ void UPCGComponent::PreEditUndo()
 		// Always flush the runtime gen state before undo/redo to avoid leaking resources or locking grid cells from future generation.
 		if (UPCGSubsystem* Subsystem = GetSubsystem())
 		{
-			Subsystem->RefreshRuntimeGenComponent(this, EPCGChangeType::GenerationGrid);
+			Subsystem->RefreshRuntimeGenExecutionSource(this, EPCGChangeType::GenerationGrid);
 		}
 	}
 
@@ -2664,7 +2674,7 @@ void UPCGComponent::Refresh(EPCGChangeType ChangeType, bool bCancelExistingRefre
 	{
 		if (UPCGSubsystem* Subsystem = GetSubsystem())
 		{
-			Subsystem->RefreshRuntimeGenComponent(this, ChangeType);
+			Subsystem->RefreshRuntimeGenExecutionSource(this, ChangeType);
 		}
 
 		return;
@@ -3093,7 +3103,7 @@ FPCGDataCollection UPCGComponent::CreateActorPCGDataCollection(AActor* Actor, co
 	TRACE_CPUPROFILER_EVENT_SCOPE(UPCGComponent::CreateActorPCGData);
 
 	FPCGGetDataFunctionRegistryParams Params;
-	Params.SourceComponent = Component;
+	Params.Source = Component;
 	Params.bParseActor = bParseActor;
 	Params.DataTypeFilter = InDataFilter;
 
diff --git a/Engine/Plugins/PCG/Source/PCG/Private/PCGComponentExecutionState.cpp b/Engine/Plugins/PCG/Source/PCG/Private/PCGComponentExecutionState.cpp
--- a/Engine/Plugins/PCG/Source/PCG/Private/PCGComponentExecutionState.cpp
+++ b/Engine/Plugins/PCG/Source/PCG/Private/PCGComponentExecutionState.cpp
@@ -391,6 +391,12 @@ bool FPCGComponentExecutionState::IsRefreshInProgress() const
 	return Component->IsRefreshInProgress();
 }
 
+bool FPCGComponentExecutionState::WasGeneratedThisSession() const
+{
+	check(Component);
+	return Component->WasGeneratedThisSession();
+}
+
 FPCGDynamicTrackingPriority FPCGComponentExecutionState::GetDynamicTrackingPriority() const
 {
 	check(Component);
diff --git a/Engine/Plugins/PCG/Source/PCG/Private/PCGGraph.cpp b/Engine/Plugins/PCG/Source/PCG/Private/PCGGraph.cpp
--- a/Engine/Plugins/PCG/Source/PCG/Private/PCGGraph.cpp
+++ b/Engine/Plugins/PCG/Source/PCG/Private/PCGGraph.cpp
@@ -62,6 +62,11 @@ namespace PCGGraph
 		TEXT("pcg.Graph.CompileTasksToDetectChanges"),
 		true,
 		TEXT("Recompile graph after edits to detect if compiled tasks changed, to avoid unnecessary graph execution"));
+
+	static TAutoConsoleVariable<bool> CVarEnableActorComponentlessGeneration(
+		TEXT("pcg.RuntimeGeneration.EnableActorComponentlessGeneration"),
+		true,
+		TEXT("Allows executing runtime PCG graphs without partition actors and local PCG components."));
 }
 
 namespace PCGGraphUtils
@@ -1673,6 +1678,11 @@ double UPCGGraph::GetGridCleanupRadiusFromGrid(EPCGHiGenGrid Grid) const
 	}
 }
 
+bool UPCGGraph::UseActorComponentlessGeneration() const
+{
+	return bUseActorComponentlessGeneration && PCGGraph::CVarEnableActorComponentlessGeneration.GetValueOnAnyThread();
+}
+
 #if WITH_EDITOR
 void UPCGGraph::DisableNotificationsForEditor()
 {
@@ -2113,7 +2123,8 @@ void UPCGGraph::PostEditChangeChainProperty(FPropertyChangedChainEvent& Property
 	else if (PropertyName == GET_MEMBER_NAME_CHECKED(UPCGGraph, HiGenGridSize)
 		|| PropertyName == GET_MEMBER_NAME_CHECKED(UPCGGraph, bUseHierarchicalGeneration)
 		|| PropertyName == GET_MEMBER_NAME_CHECKED(UPCGGraph, bUse2DGrid)
-		|| PropertyName == GET_MEMBER_NAME_CHECKED(UPCGGraph, HiGenExponential))
+		|| PropertyName == GET_MEMBER_NAME_CHECKED(UPCGGraph, HiGenExponential)
+		|| PropertyName == GET_MEMBER_NAME_CHECKED(UPCGGraph, bUseActorComponentlessGeneration))
 	{
 		// The higen settings change the structure of the graph (presence or absence of links between grid levels).
 		NotifyGraphChanged(EPCGChangeType::Structural | EPCGChangeType::GenerationGrid);
diff --git a/Engine/Plugins/PCG/Source/PCG/Private/PCGTrackingManager.cpp b/Engine/Plugins/PCG/Source/PCG/Private/PCGTrackingManager.cpp
--- a/Engine/Plugins/PCG/Source/PCG/Private/PCGTrackingManager.cpp
+++ b/Engine/Plugins/PCG/Source/PCG/Private/PCGTrackingManager.cpp
@@ -438,6 +438,8 @@ bool FPCGTrackingManager::RemapExecutionSource(const IPCGGraphExecutionSource* I
 	RemapTracking(InOldExecutionSource, InNewExecutionSource);
 #endif // WITH_EDITOR
 
+	PCGSubsystem->OnOriginalExecutionSourceReplaced(InOldExecutionSource, InNewExecutionSource);
+
 	return true;
 }
 
@@ -1473,7 +1475,7 @@ void FPCGTrackingManager::OnSelectionKeyChanged(const FPCGSelectionKey& InSelect
 void FPCGTrackingManager::OnSelectionKeyChanged(const FPCGSelectionKey& InSelectionKey, const UObject* InOriginatingChangeObject, const TSet<FName>& InRemovedTags, bool bInNoRefreshOwner, TArrayView<const FBox> InBounds)
 {
 	TRACE_CPUPROFILER_EVENT_SCOPE(FPCGTrackingManager::OnSelectionKeyChanged);
-		
+
 	/** First gather all the execution sources that are tracking this key */
 	TSet<IPCGGraphExecutionSource*> CulledTrackedExecutionSources;
 	TSet<IPCGGraphExecutionSource*> AlwaysTrackedExecutionSources;
@@ -1487,11 +1489,11 @@ void FPCGTrackingManager::OnSelectionKeyChanged(const FPCGSelectionKey& InSelect
 	
 	auto Gather = [InSelectionKey, InOriginatingSelectionKey, &InRemovedTags, &Change](const TMap<FPCGSelectionKey, TSet<IPCGGraphExecutionSource*>>& InMap, TSet<IPCGGraphExecutionSource*>& OutSet)
 	{
-		for (const auto& It : InMap)
+		for (const auto& [SelectionKey, ExecutionSource] : InMap)
 		{
-			if (It.Key.IsMatching(InSelectionKey, InRemovedTags, It.Value, &OutSet) || It.Key.IsMatching(InOriginatingSelectionKey, InRemovedTags, It.Value, &OutSet))
+			if (SelectionKey.IsMatching(InSelectionKey, InRemovedTags, ExecutionSource, &OutSet) || SelectionKey.IsMatching(InOriginatingSelectionKey, InRemovedTags, ExecutionSource, &OutSet))
 			{
-				Change->MatchedKeys.Add(It.Key);
+				Change->MatchedKeys.Add(SelectionKey);
 			}
 		}
 	};
diff --git a/Engine/Plugins/PCG/Source/PCG/Private/RuntimeGen/PCGRuntimeGenChangeDetection.cpp b/Engine/Plugins/PCG/Source/PCG/Private/RuntimeGen/PCGRuntimeGenChangeDetection.cpp
--- a/Engine/Plugins/PCG/Source/PCG/Private/RuntimeGen/PCGRuntimeGenChangeDetection.cpp
+++ b/Engine/Plugins/PCG/Source/PCG/Private/RuntimeGen/PCGRuntimeGenChangeDetection.cpp
@@ -2,7 +2,6 @@
 
 #include "RuntimeGen/PCGRuntimeGenChangeDetection.h"
 
-#include "PCGComponent.h"
 #include "RuntimeGen/GenSources/PCGGenSourceBase.h"
 
 namespace PCGRuntimeGenChangeDetection
@@ -114,16 +113,16 @@ namespace PCGRuntimeGenChangeDetection
 		StateCache.FindOrAdd(InInputs).UpdateFromInputs(InInputs);
 	}
 
-	void FDetector::FlushCachedState(const UPCGComponent* InOptionalOriginalComponent)
+	void FDetector::FlushCachedState(const IPCGGraphExecutionSource* InOptionalOriginalExecutionSource)
 	{
-		if (InOptionalOriginalComponent)
+		if (InOptionalOriginalExecutionSource)
 		{
 			TArray<FDetectionInputs, TInlineAllocator<32>> KeysToRemove;
 
 			for (const TPair<FDetectionInputs, FCachedState>& Entry : StateCache)
 			{
-				const UPCGComponent* Component = Entry.Key.GetOriginalComponent();
-				if (!Component || Component == InOptionalOriginalComponent)
+				const IPCGGraphExecutionSource* ExecutionSource = Entry.Key.GetOriginalExecutionSource();
+				if (!ExecutionSource || ExecutionSource == InOptionalOriginalExecutionSource)
 				{
 					KeysToRemove.Add(Entry.Key);
 				}
diff --git a/Engine/Plugins/PCG/Source/PCG/Private/RuntimeGen/PCGRuntimeGenExecutionSource.cpp b/Engine/Plugins/PCG/Source/PCG/Private/RuntimeGen/PCGRuntimeGenExecutionSource.cpp
--- a/Engine/Plugins/PCG/Source/PCG/Private/RuntimeGen/PCGRuntimeGenExecutionSource.cpp
+++ b/Engine/Plugins/PCG/Source/PCG/Private/RuntimeGen/PCGRuntimeGenExecutionSource.cpp
@@ -435,6 +435,12 @@ bool UPCGRuntimeGenExecutionSource::IsRefreshInProgress() const
 {
 	return false;
 }
+
+bool UPCGRuntimeGenExecutionSource::WasGeneratedThisSession() const
+{
+	// Runtime gen execution sources are generated right away, so this is always true.
+	return true;
+}
 #endif // WITH_EDITOR
 
 IPCGGraphExecutionSource* UPCGRuntimeGenExecutionSource::GetOriginalSource() const
diff --git a/Engine/Plugins/PCG/Source/PCG/Private/RuntimeGen/PCGRuntimeGenExecutionState.cpp b/Engine/Plugins/PCG/Source/PCG/Private/RuntimeGen/PCGRuntimeGenExecutionState.cpp
--- a/Engine/Plugins/PCG/Source/PCG/Private/RuntimeGen/PCGRuntimeGenExecutionState.cpp
+++ b/Engine/Plugins/PCG/Source/PCG/Private/RuntimeGen/PCGRuntimeGenExecutionState.cpp
@@ -238,6 +238,12 @@ bool FPCGRuntimeGenExecutionState::IsRefreshInProgress() const
 	check(Source);
 	return Source->IsRefreshInProgress();
 }
+
+bool FPCGRuntimeGenExecutionState::WasGeneratedThisSession() const
+{
+	check(Source);
+	return Source->WasGeneratedThisSession();
+}
 #endif // WITH_EDITOR
 
 IPCGGraphExecutionSource* FPCGRuntimeGenExecutionState::GetOriginalSource() const
diff --git a/Engine/Plugins/PCG/Source/PCG/Private/RuntimeGen/PCGRuntimeGenScheduler.cpp b/Engine/Plugins/PCG/Source/PCG/Private/RuntimeGen/PCGRuntimeGenScheduler.cpp
--- a/Engine/Plugins/PCG/Source/PCG/Private/RuntimeGen/PCGRuntimeGenScheduler.cpp
+++ b/Engine/Plugins/PCG/Source/PCG/Private/RuntimeGen/PCGRuntimeGenScheduler.cpp
@@ -14,9 +14,10 @@
 #include "Grid/PCGPartitionActor.h"
 #include "Helpers/PCGActorHelpers.h"
 #include "Helpers/PCGHelpers.h"
+#include "RuntimeGen/PCGGenSourceManager.h"
+#include "RuntimeGen/PCGRuntimeGenExecutionSource.h"
 #include "RuntimeGen/GenSources/PCGGenSourceBase.h"
 #include "RuntimeGen/GenSources/PCGGenSourceComponent.h"
-#include "RuntimeGen/PCGGenSourceManager.h"
 
 #include "DrawDebugHelpers.h"
 #include "EngineDefines.h" // For UE_ENABLE_DEBUG_DRAWING
@@ -213,6 +214,36 @@ FPCGGridDescriptor FPCGRuntimeGenScheduler::FGridGenerationKey::GetGridDescripto
 	return FPCGGridDescriptor().SetGridSize(GetGridSize()).SetIsRuntime(true).SetIs2DGrid(Use2DGrid());
 }
 
+// DEPRECATED 5.8
+UPCGComponent* FPCGRuntimeGenScheduler::FGridGenerationKey::GetOriginalComponent() const
+{
+	return Cast<UPCGComponent>(GetOriginalSource());
+}
+
+// DEPRECATED 5.8
+UPCGComponent* FPCGRuntimeGenScheduler::FGridGenerationKey::GetOriginalComponentEvenIfGarbage() const
+{
+	return Cast<UPCGComponent>(GetOriginalSourceEvenIfGarbage());
+}
+
+// DEPRECATED 5.8
+void FPCGRuntimeGenScheduler::FGridGenerationKey::SetOriginalComponent(UPCGComponent* InComponent)
+{
+	SetOriginalSource(InComponent);
+}
+
+// DEPRECATED 5.8
+IPCGGraphExecutionSource* FPCGRuntimeGenScheduler::FGridGenerationKey::GetCachedLocalComponent() const
+{
+	return GetCachedLocalSource();
+}
+
+// DEPRECATED 5.8
+void FPCGRuntimeGenScheduler::FGridGenerationKey::SetCachedLocalComponent(IPCGGraphExecutionSource* InLocalComponent) const
+{
+	SetCachedLocalSource(InLocalComponent);
+}
+
 FPCGRuntimeGenScheduler::FPCGRuntimeGenScheduler(UWorld* InWorld, FPCGTrackingManager* InTrackingManager)
 {
 	check(InWorld && InTrackingManager);
@@ -311,7 +342,7 @@ void FPCGRuntimeGenScheduler::Tick(APCGWorldActor* InPCGWorldActor, double InEnd
 	}
 
 	// Initialize RuntimeGen PA pool if necessary. If PoolSize is 0, then we have not initialized the pool yet.
-	if (!GenSources.IsEmpty() || !GeneratedComponents.IsEmpty())
+	if (!GenSources.IsEmpty() || !GeneratedSources.IsEmpty())
 	{
 		if (PartitionActorPoolSize == 0 && PCGRuntimeGenSchedulerHelpers::CVarRuntimeGenerationEnablePooling.GetValueOnAnyThread())
 		{
@@ -327,40 +358,40 @@ void FPCGRuntimeGenScheduler::Tick(APCGWorldActor* InPCGWorldActor, double InEnd
 		return;
 	}
 
-	CleanupDelayedRefreshComponents();
+	CleanupDelayedRefreshSources();
 
-	// 1. Queue nearby components for generation.
+	// 1. Queue nearby execution sources for generation.
 	
 	if (!GenSources.IsEmpty())
 	{
-		FTickQueueComponentsForGenerationInputs Inputs;
+		FTickQueueSourcesForGenerationInputs Inputs;
 		Inputs.GenSources = &GenSources;
 		Inputs.PCGWorldActor = InPCGWorldActor;
 		Inputs.AllPartitionedExecutionSources = TrackingManager->GetAllRegisteredPartitionedExecutionSources();
 		Inputs.AllNonPartitionedExecutionSources = TrackingManager->GetAllRegisteredNonPartitionedExecutionSources();
-		Inputs.GeneratedComponents = &GeneratedComponents;
+		Inputs.GeneratedSources = &GeneratedSources;
 
 		ChangeDetector.PreTick();
-		TickQueueComponentsForGeneration(Inputs, ComponentsToGenerate);
+		TickQueueSourcesForGeneration(Inputs, SourcesToGenerate);
 		ChangeDetector.PostTick();
 	}
 
-	// 2. Schedule cleanup on components that become out of range.
+	// 2. Schedule cleanup on execution sources that become out of range.
 
-	if (!GeneratedComponents.IsEmpty() && (GenSources.IsEmpty() || bCleanupScanRequired || !IsChangeDetectionEnabled()))
+	if (!GeneratedSources.IsEmpty() && (GenSources.IsEmpty() || bCleanupScanRequired || !IsChangeDetectionEnabled()))
 	{
 		TickCleanup(GenSources, InPCGWorldActor, InEndTime);
 	}
 
-	// 3. Schedule generation on components in priority order.
-	if (!ComponentsToGenerate.IsEmpty())
+	// 3. Schedule generation on execution sources in priority order.
+	if (!SourcesToGenerate.IsEmpty())
 	{
-		// Sort components by priority (will be generated in descending order).
-		ComponentsToGenerate.ValueSort([](double PrioA, double PrioB)->bool { return PrioA > PrioB; });
+		// Sort execution sources by priority (will be generated in descending order).
+		SourcesToGenerate.ValueSort([](double PrioA, double PrioB)->bool { return PrioA > PrioB; });
 
-		// Only apply time budget to cleanup currently. Currently too easy to introduce latency issues so don't hold back generation new components
+		// Only apply time budget to cleanup currently. Currently too easy to introduce latency issues so don't hold back generation new execution sources
 		// (and we already have CVarNumGeneratingComponentsAtSameTime to throttle new generations).
-		TickScheduleGeneration(ComponentsToGenerate);
+		TickScheduleGeneration(SourcesToGenerate);
 	}
 
 #if PCG_RGS_ONSCREENDEBUGMESSAGES
@@ -370,7 +401,7 @@ void FPCGRuntimeGenScheduler::Tick(APCGWorldActor* InPCGWorldActor, double InEnd
 
 bool FPCGRuntimeGenScheduler::ShouldTick()
 {
-	check(Context.World && TrackingManager);
+	check(Context.World);
 
 	if (!PCGRuntimeGenSchedulerHelpers::CVarRuntimeGenerationEnable.GetValueOnAnyThread())
 	{
@@ -403,8 +434,8 @@ bool FPCGRuntimeGenScheduler::ShouldTick()
 		RefreshContext();
 	}
 
-	// We can stop ticking if there are no runtime gen components alive and there are no generated components that need cleaning up.
-	if (!Context.bAnyRuntimeGenSourcesExist && GeneratedComponents.IsEmpty() && GeneratedComponentsToRemove.IsEmpty())
+	// We can stop ticking if there are no runtime gen execution sources alive and there are no generated execution sources that need cleaning up.
+	if (!Context.bAnyRuntimeGenSourcesExist && GeneratedSources.IsEmpty() && GeneratedSourcesToRemove.IsEmpty())
 	{
 		return false;
 	}
@@ -437,20 +468,20 @@ bool FPCGRuntimeGenScheduler::ShouldTick()
 	return true;
 }
 
-void FPCGRuntimeGenScheduler::TickQueueComponentsForGeneration(
-	const FTickQueueComponentsForGenerationInputs& Inputs,
-	TMap<FGridGenerationKey, double>& OutComponentsToGenerate)
+void FPCGRuntimeGenScheduler::TickQueueSourcesForGeneration(
+	const FTickQueueSourcesForGenerationInputs& Inputs,
+	TMap<FGridGenerationKey, double>& OutSourcesToGenerate)
 {
-	TRACE_CPUPROFILER_EVENT_SCOPE(FPCGRuntimeGenScheduler::TickQueueComponentsForGeneration);
+	TRACE_CPUPROFILER_EVENT_SCOPE(FPCGRuntimeGenScheduler::TickQueueSourcesForGeneration);
 
 	check(Inputs.PCGWorldActor && Inputs.GenSources);
 
-	// TODO: Thought - it would be possible to maintain a global maximum generation distance across all components
-	// perhaps in the actor&comp mapping system, and then do a spatial query to get the components here.
+	// TODO: Thought - it would be possible to maintain a global maximum generation distance across all execution sources
+	// perhaps in the actor&comp mapping system, and then do a spatial query to get the execution sources here.
 
-	auto AddComponentToGenerate = [&OutComponentsToGenerate](FGridGenerationKey& InKey, const IPCGGenSourceBase* InGenSource, const UPCGSchedulingPolicyBase* InPolicy, const FBox& InComponentBounds, bool bInUse2DGrid)
+	auto AddSourceToGenerate = [&OutSourcesToGenerate](FGridGenerationKey& InKey, const IPCGGenSourceBase* InGenSource, const UPCGSchedulingPolicyBase* InPolicy, const FBox& InSourceBounds, bool bInUse2DGrid)
 	{
-		const double PolicyPriority = InPolicy ? InPolicy->CalculatePriority(InGenSource, InComponentBounds, bInUse2DGrid) : 0.0;
+		const double PolicyPriority = InPolicy ? InPolicy->CalculatePriority(InGenSource, InSourceBounds, bInUse2DGrid) : 0.0;
 		double Priority = FMath::Clamp(PolicyPriority, 0.0, 1.0);
 		if (PolicyPriority != Priority)
 		{
@@ -461,10 +492,10 @@ void FPCGRuntimeGenScheduler::TickQueueComponentsForGeneration(
 		const uint32 GridSize = InKey.GetGridSize();
 		Priority += GridSize;
 
-		double* ExistingPriority = OutComponentsToGenerate.Find(InKey);
+		double* ExistingPriority = OutSourcesToGenerate.Find(InKey);
 		if (!ExistingPriority)
 		{
-			OutComponentsToGenerate.Add(InKey, Priority);
+			OutSourcesToGenerate.Add(InKey, Priority);
 		}
 		else if (Priority > *ExistingPriority)
 		{
@@ -520,28 +551,38 @@ void FPCGRuntimeGenScheduler::TickQueueComponentsForGeneration(
 	const bool bDebugDrawGenerationSources = PCGSystemSwitches::CVarPCGDebugDrawGeneratedCells.GetValueOnGameThread();
 #endif
 
-	// Collect local components from all partitioned components.
+	// Collect local execution sources from all partitioned execution sources.
 	for (IPCGGraphExecutionSource* OriginalExecutionSource : Inputs.AllPartitionedExecutionSources)
 	{
-		TRACE_CPUPROFILER_EVENT_SCOPE(FPCGRuntimeGenScheduler::CollectLocalComponents);
+		TRACE_CPUPROFILER_EVENT_SCOPE(FPCGRuntimeGenScheduler::CollectPartitionedSources);
 
-		// @todo_pcg: support execution source
-		UPCGComponent* OriginalComponent = Cast<UPCGComponent>(OriginalExecutionSource);
+		if (!OriginalExecutionSource)
+		{
+			continue;
+		}
 
-		if (!ensure(OriginalComponent) || !OriginalComponent->GetGraph() || !OriginalComponent->bActivated)
+		// Skip PCG components in actor-componentless mode. They should have registered a UPCGRuntimeGenExecutionSource that will handle generation.
+		if (UPCGComponent* OriginalComponent = Cast<UPCGComponent>(OriginalExecutionSource); OriginalComponent && OriginalComponent->UseActorComponentlessGeneration())
 		{
 			continue;
 		}
 
-		const UPCGSchedulingPolicyBase* Policy = OriginalComponent->GetRuntimeGenSchedulingPolicy();
+		const IPCGGraphExecutionState& OriginalExecutionState = OriginalExecutionSource->GetExecutionState();
+
+		if (!OriginalExecutionState.GetGraph() || !OriginalExecutionState.IsActive())
+		{
+			continue;
+		}
+
+		const UPCGSchedulingPolicyBase* Policy = OriginalExecutionState.GetRuntimeGenSchedulingPolicy();
 
 		// todo_pcg: For each execution domain (for now only GenAtRuntime/dynamic), assuming we run Preview through this scheduler, which it seems like we will.
 		// todo_pcg: All this stuff can be hoisted - whether valid, has graph, active, managed by runtime gen..
-		if (OriginalComponent->IsManagedByRuntimeGenSystem() && ensure(Policy))
+		if (OriginalExecutionState.IsManagedByRuntimeGenSystem() && ensure(Policy))
 		{
 			bool bHasUnbounded = false;
 			PCGHiGenGrid::FSizeArray GridSizes;
-			PCGHelpers::GetGenerationGridSizes(OriginalComponent->GetGraph(), Inputs.PCGWorldActor, GridSizes, bHasUnbounded);
+			PCGHelpers::GetGenerationGridSizes(OriginalExecutionState.GetGraph(), Inputs.PCGWorldActor, GridSizes, bHasUnbounded);
 
 			if (GridSizes.IsEmpty() && !bHasUnbounded)
 			{
@@ -560,17 +601,17 @@ void FPCGRuntimeGenScheduler::TickQueueComponentsForGeneration(
 					WorldStreamingQueryGridSizes[GridIndex] = PCGHiGenGrid::UninitializedGridSize();
 
 					PCGHiGenGrid::FSizeArray ParentGridsDescending;
-					OriginalComponent->GetGraph()->GetParentGridSizes(GridSizes[GridIndex], ParentGridsDescending);
+					OriginalExecutionState.GetGraph()->GetParentGridSizes(GridSizes[GridIndex], ParentGridsDescending);
 					for (uint32 ParentGridSize : ParentGridsDescending)
 					{
-						if (OriginalComponent->DoesGridDependOnWorldStreaming(ParentGridSize) && ensure(ParentGridSize > GridSizes[GridIndex]))
+						if (Policy->DoesGridDependOnWorldStreaming(ParentGridSize) && ensure(ParentGridSize > GridSizes[GridIndex]))
 						{
 							WorldStreamingQueryGridSizes[GridIndex] = ParentGridSize;
 							break;
 						}
 					}
 
-					if (WorldStreamingQueryGridSizes[GridIndex] == PCGHiGenGrid::UninitializedGridSize() && OriginalComponent->DoesGridDependOnWorldStreaming(GridSizes[GridIndex]))
+					if (WorldStreamingQueryGridSizes[GridIndex] == PCGHiGenGrid::UninitializedGridSize() && Policy->DoesGridDependOnWorldStreaming(GridSizes[GridIndex]))
 					{
 						WorldStreamingQueryGridSizes[GridIndex] = GridSizes[GridIndex];
 					}
@@ -582,15 +623,15 @@ void FPCGRuntimeGenScheduler::TickQueueComponentsForGeneration(
 			}
 
 			const EPCGHiGenGrid MaxGrid = bHasUnbounded ? EPCGHiGenGrid::Unbounded : PCGHiGenGrid::GridSizeToGrid(GridSizes[0]);
-			const double MaxGenerationRadius = OriginalComponent->GetGenerationRadiusFromGrid(MaxGrid);
+			const double MaxGenerationRadius = OriginalExecutionState.GetGenerationRadiusFromGrid(MaxGrid);
 
 			for (const IPCGGenSourceBase* GenSource : *Inputs.GenSources)
 			{
 				const PCGRuntimeGenChangeDetection::FDetectionInputs ChangeDetectionInputs =
 				{
-					OriginalComponent,
+					OriginalExecutionSource,
 					GenSource,
-					OriginalComponent->GetGenerationRadii(),
+					OriginalExecutionState.GetGenerationRadii(),
 					PCGRuntimeGenSchedulerHelpers::CVarRuntimeGenerationRadiusMultiplier.GetValueOnGameThread(),
 					Policy->CullsBasedOnDirection(),
 					GridSizes.IsEmpty() ? PCGHiGenGrid::UnboundedGridSize() : GridSizes.Last(),
@@ -630,25 +671,25 @@ void FPCGRuntimeGenScheduler::TickQueueComponentsForGeneration(
 					{
 						ensure(PCGHiGenGrid::IsValidGridSize(GridSize));
 
-						const double GenerationRadius = OriginalComponent->GetGenerationRadiusFromGrid(PCGHiGenGrid::GridSizeToGrid(GridSize));
+						const double GenerationRadius = OriginalExecutionState.GetGenerationRadiusFromGrid(PCGHiGenGrid::GridSizeToGrid(GridSize));
 
 						DrawDebugSphere(Context.World, GenSourcePosition, GenerationRadius, 64, FColor::Red, /*bPersistentLines=*/false, /*LifeTime=*/0.0f);
 					}
 				}
 #endif // UE_ENABLE_DEBUG_DRAWING
 
-				const FBox OriginalComponentBounds = OriginalComponent->GetGridBounds();
-				const bool bIs2DGrid = OriginalComponent->Use2DGrid();
+				const FBox OriginalSourceBounds = OriginalExecutionState.GetBounds();
+				const bool bIs2DGrid = OriginalExecutionState.Use2DGrid();
 
 				FVector ModifiedGenSourcePosition = GenSourcePosition;
 				if (bIs2DGrid)
 				{
-					ModifiedGenSourcePosition.Z = OriginalComponentBounds.Min.Z;
+					ModifiedGenSourcePosition.Z = OriginalSourceBounds.Min.Z;
 				}
 
-				const double DistanceSquared = OriginalComponentBounds.ComputeSquaredDistanceToPoint(ModifiedGenSourcePosition);
+				const double DistanceSquared = OriginalSourceBounds.ComputeSquaredDistanceToPoint(ModifiedGenSourcePosition);
 
-				// If GenSource is not within range of the component then skip it.
+				// If GenSource is not within range of the execution source then skip it.
 				if (DistanceSquared > MaxGenerationRadius * MaxGenerationRadius)
 				{
 					if (bDoChangeDetection)
@@ -664,23 +705,23 @@ void FPCGRuntimeGenScheduler::TickQueueComponentsForGeneration(
 					// Poll for streaming to complete on the unbounded grid.
 					if (bCheckStreaming && Policy->DoesGridDependOnWorldStreaming(PCGHiGenGrid::UnboundedGridSize()))
 					{
-						const FVector StreamingLocation = OriginalComponentBounds.GetCenter();
-						const FVector StreamingSize = OriginalComponentBounds.GetSize();
+						const FVector StreamingLocation = OriginalSourceBounds.GetCenter();
+						const FVector StreamingSize = OriginalSourceBounds.GetSize();
 
 						if (!IsWorldStreamingComplete(StreamingLocation, FMath::Max(StreamingSize.X, StreamingSize.Y)))
 						{
-							UE_LOG(LogPCG, VeryVerbose, TEXT("Partitioned original component '%s' rejected as world is not fully loaded."), *OriginalComponent->GetOwner()->GetActorNameOrLabel());
+							UE_LOG(LogPCG, VeryVerbose, TEXT("Partitioned original component '%s' rejected as world is not fully loaded."), *OriginalExecutionState.GetDebugName());
 							continue;
 						}
 					}
 					
-					// Ignore components that have already been generated or marked for generation. Unbounded grid size means not-partitioned.
-					FGridGenerationKey Key(PCGHiGenGrid::UnboundedGridSize(), FIntVector(0), OriginalComponent);
-					if (!Inputs.GeneratedComponents->Contains(Key) &&
-						(!Policy || Policy->ShouldGenerate(GenSource, OriginalComponentBounds, bIs2DGrid)))
+					// Ignore execution sources that have already been generated or marked for generation. Unbounded grid size means not-partitioned.
+					FGridGenerationKey Key(PCGHiGenGrid::UnboundedGridSize(), FIntVector(0), OriginalExecutionSource);
+					if (!Inputs.GeneratedSources->Contains(Key) &&
+						(!Policy || Policy->ShouldGenerate(GenSource, OriginalSourceBounds, bIs2DGrid)))
 					{
 						check(Key.GetGridDescriptor().Is2DGrid() == bIs2DGrid);
-						AddComponentToGenerate(Key, GenSource, Policy, OriginalComponentBounds, bIs2DGrid);
+						AddSourceToGenerate(Key, GenSource, Policy, OriginalSourceBounds, bIs2DGrid);
 					}
 				}
 
@@ -693,7 +734,7 @@ void FPCGRuntimeGenScheduler::TickQueueComponentsForGeneration(
 					const uint32 GridSize = GridSizes[GridIndex];
 
 					const FIntVector GenSourceGridPosition = UPCGActorHelpers::GetCellCoord(GenSourcePosition, GridSize, bIs2DGrid);
-					const double GenerationRadius = OriginalComponent->GetGenerationRadiusFromGrid(PCGHiGenGrid::GridSizeToGrid(GridSize));
+					const double GenerationRadius = OriginalExecutionState.GetGenerationRadiusFromGrid(PCGHiGenGrid::GridSizeToGrid(GridSize));
 					const int32 GridRadius = FMath::CeilToInt32(GenerationRadius / GridSize); // Radius discretized to # of grid cells.
 					const int32 VerticalGridRadius = bIs2DGrid ? 0 : GridRadius; // Flatten the vertical grid radius in the 2D case.
 
@@ -714,10 +755,10 @@ void FPCGRuntimeGenScheduler::TickQueueComponentsForGeneration(
 							for (int32 X = GenSourceGridPosition.X - GridRadius; X <= GenSourceGridPosition.X + GridRadius; ++X)
 							{
 								FIntVector GridCoords(X, Y, Z);
-								FGridGenerationKey Key(GridSize, GridCoords, OriginalComponent);
+								FGridGenerationKey Key(GridSize, GridCoords, OriginalExecutionSource);
 
-								// Ignore components that have already been generated or marked for generation.
-								if (Inputs.GeneratedComponents->Find({ GridSize, GridCoords, OriginalComponent }))
+								// Ignore execution sources that have already been generated or marked for generation.
+								if (Inputs.GeneratedSources->Find({ GridSize, GridCoords, OriginalExecutionSource }))
 								{
 									continue;
 								}
@@ -725,8 +766,8 @@ void FPCGRuntimeGenScheduler::TickQueueComponentsForGeneration(
 								const FVector Center = FVector(GridCoords.X + 0.5, GridCoords.Y + 0.5, GridCoords.Z + 0.5) * GridSize;
 								const FBox CellBounds(Center - HalfExtent, Center + HalfExtent);
 
-								// Overlap cell with the partitioned component.
-								const FBox IntersectedBounds = OriginalComponentBounds.Overlap(CellBounds);
+								// Overlap cell with the partitioned execution source.
+								const FBox IntersectedBounds = OriginalSourceBounds.Overlap(CellBounds);
 								if (!IntersectedBounds.IsValid || IntersectedBounds.GetVolume() <= UE_DOUBLE_SMALL_NUMBER)
 								{
 									continue;
@@ -772,7 +813,7 @@ void FPCGRuntimeGenScheduler::TickQueueComponentsForGeneration(
 
 									if (bStreamingComplete)
 									{
-										AddComponentToGenerate(Key, GenSource, Policy, IntersectedBounds, Key.GetGridDescriptor().Is2DGrid());
+										AddSourceToGenerate(Key, GenSource, Policy, IntersectedBounds, Key.GetGridDescriptor().Is2DGrid());
 									}
 									else
 									{
@@ -785,7 +826,7 @@ void FPCGRuntimeGenScheduler::TickQueueComponentsForGeneration(
 				}
 
 				// Require repeated scanning if we are waiting for level to stream.
-				// todo_pcg: We can definitely do better than this. We should probably have a separate list of keys for components that are awaiting streaming so
+				// todo_pcg: We can definitely do better than this. We should probably have a separate list of keys for execution sources that are awaiting streaming so
 				// that we can keep track of these specifically rather than requiring a full rescan.
 				if (bDoChangeDetection && !bAnyCellGenerationBlockedByStreaming)
 				{
@@ -795,55 +836,65 @@ void FPCGRuntimeGenScheduler::TickQueueComponentsForGeneration(
 		}
 	}
 
-	// Collect all non-partitioned components.
+	// Collect all non-partitioned execution sources.
 	for (IPCGGraphExecutionSource* OriginalExecutionSource : Inputs.AllNonPartitionedExecutionSources)
 	{
-		TRACE_CPUPROFILER_EVENT_SCOPE(FPCGRuntimeGenScheduler::CollectNonPartitionedComponents);
+		TRACE_CPUPROFILER_EVENT_SCOPE(FPCGRuntimeGenScheduler::CollectNonPartitionedSources);
 
-		// @todo_pcg: support execution source
-		UPCGComponent* OriginalComponent = Cast<UPCGComponent>(OriginalExecutionSource);
+		if (!OriginalExecutionSource)
+		{
+			continue;
+		}
 
-		if (!ensure(OriginalComponent) || !OriginalComponent->GetGraph() || !OriginalComponent->bActivated)
+		// Skip PCG components in actor-componentless mode. They should have registered a UPCGRuntimeGenExecutionSource that will handle generation.
+		if (UPCGComponent* OriginalComponent = Cast<UPCGComponent>(OriginalExecutionSource); OriginalComponent && OriginalComponent->UseActorComponentlessGeneration())
+		{
+			continue;
+		}
+
+		const IPCGGraphExecutionState& OriginalExecutionState = OriginalExecutionSource->GetExecutionState();
+
+		if (!OriginalExecutionState.GetGraph() || !OriginalExecutionState.IsActive())
 		{
 			continue;
 		}
 
 		// todo_pcg: Worth adding change detection here? Decide on quantization - perhaps a proportion of minimum volume extent or such?
 		
-		// The generation key for a non-partitioned component should always have unbounded grid size and 0,0,0 cell coord.
-		if (Inputs.GeneratedComponents->Contains({ PCGHiGenGrid::UnboundedGridSize(), FIntVector(0), OriginalComponent }))
+		// The generation key for a non-partitioned execution source should always have unbounded grid size and 0,0,0 cell coord.
+		if (Inputs.GeneratedSources->Contains({ PCGHiGenGrid::UnboundedGridSize(), FIntVector(0), OriginalExecutionSource }))
 		{
-			// todo_pcg: Needs change detection to throttle clean up scans, for now just always cleanup if any runtime gen unbounded component is present.
+			// todo_pcg: Needs change detection to throttle clean up scans, for now just always cleanup if any runtime gen unbounded execution source is present.
 			bCleanupScanRequired = true;
 
 			continue;
 		}
 
-		const UPCGSchedulingPolicyBase* Policy = OriginalComponent->GetRuntimeGenSchedulingPolicy();
+		const UPCGSchedulingPolicyBase* Policy = OriginalExecutionState.GetRuntimeGenSchedulingPolicy();
 
 		// TODO: For each execution domain (for now only GenAtRuntime/dynamic), assuming we run Preview through this scheduler, which it seems like we will.
-		if (OriginalComponent->IsManagedByRuntimeGenSystem() && ensure(Policy))
+		if (OriginalExecutionState.IsManagedByRuntimeGenSystem() && ensure(Policy))
 		{
-			// todo_pcg: Needs change detection to throttle clean up scans, for now just always cleanup if any runtime gen unbounded component is present.
+			// todo_pcg: Needs change detection to throttle clean up scans, for now just always cleanup if any runtime gen unbounded execution source is present.
 			bCleanupScanRequired = true;
 
-			const FBox OriginalComponentBounds = OriginalComponent->GetGridBounds();
+			const FBox OriginalSourceBounds = OriginalExecutionState.GetBounds();
 
 			// Poll for streaming to complete.
 			if (bCheckStreaming && Policy->DoesGridDependOnWorldStreaming(PCGHiGenGrid::UnboundedGridSize()))
 			{
-				const FVector StreamingLocation = OriginalComponentBounds.GetCenter();
-				const FVector StreamingSize = OriginalComponentBounds.GetSize();
+				const FVector StreamingLocation = OriginalSourceBounds.GetCenter();
+				const FVector StreamingSize = OriginalSourceBounds.GetSize();
 
 				if (!IsWorldStreamingComplete(StreamingLocation, FMath::Max(StreamingSize.X, StreamingSize.Y)))
 				{
-					UE_LOG(LogPCG, VeryVerbose, TEXT("Non-partitioned original component '%s' rejected as world is not fully loaded."), *OriginalComponent->GetOwner()->GetActorNameOrLabel());
+					UE_LOG(LogPCG, VeryVerbose, TEXT("Non-partitioned original component '%s' rejected as world is not fully loaded."), *OriginalExecutionState.GetDebugName());
 					continue;
 				}
 			}
 
 			// Unbounded will grab the base GenerationRadius used for non-partitioned and unbounded.
-			const double MaxGenerationRadius = OriginalComponent->GetGenerationRadiusFromGrid(EPCGHiGenGrid::Unbounded);
+			const double MaxGenerationRadius = OriginalExecutionState.GetGenerationRadiusFromGrid(EPCGHiGenGrid::Unbounded);
 
 			for (const IPCGGenSourceBase* GenSource : *Inputs.GenSources)
 			{
@@ -856,20 +907,20 @@ void FPCGRuntimeGenScheduler::TickQueueComponentsForGeneration(
 				const FVector GenSourcePosition = GenSourcePositionOptional.GetValue();
 
 				FVector ModifiedGenSourcePosition = GenSourcePosition;
-				if (OriginalComponent->Use2DGrid())
+				if (OriginalExecutionState.Use2DGrid())
 				{
-					ModifiedGenSourcePosition.Z = OriginalComponentBounds.Min.Z;
+					ModifiedGenSourcePosition.Z = OriginalSourceBounds.Min.Z;
 				}
 
-				const double DistanceSquared = OriginalComponentBounds.ComputeSquaredDistanceToPoint(ModifiedGenSourcePosition);
+				const double DistanceSquared = OriginalSourceBounds.ComputeSquaredDistanceToPoint(ModifiedGenSourcePosition);
 
-				// Max radius for a non-partitioned component is just the base GenerationRadius.
+				// Max radius for a non-partitioned execution source is just the base GenerationRadius.
 				if (DistanceSquared <= MaxGenerationRadius * MaxGenerationRadius && 
-					(!Policy || Policy->ShouldGenerate(GenSource, OriginalComponentBounds, /*bUse2DGrid=*/false)))
+					(!Policy || Policy->ShouldGenerate(GenSource, OriginalSourceBounds, /*bUse2DGrid=*/false)))
 				{
 					// Unbounded grid size means not-partitioned.
-					FGridGenerationKey Key(PCGHiGenGrid::UnboundedGridSize(), FIntVector(0), OriginalComponent);
-					AddComponentToGenerate(Key, GenSource, Policy, OriginalComponentBounds, /*bUse2DGrid=*/false);
+					FGridGenerationKey Key(PCGHiGenGrid::UnboundedGridSize(), FIntVector(0), OriginalExecutionSource);
+					AddSourceToGenerate(Key, GenSource, Policy, OriginalSourceBounds, /*bUse2DGrid=*/false);
 				}
 			}
 		}
@@ -880,7 +931,7 @@ void FPCGRuntimeGenScheduler::TickCleanup(const TSet<IPCGGenSourceBase*>& InGenS
 {
 	TRACE_CPUPROFILER_EVENT_SCOPE(FPCGRuntimeGenScheduler::TickCleanup);
 
-	check(TrackingManager && InPCGWorldActor);
+	check(InPCGWorldActor);
 
 	auto CheckIfAllGenSourcesWantToCleanup = [&InGenSources](const UPCGSchedulingPolicyBase* Policy, const FPCGGridDescriptor& GridDescriptor, const FBox& GridBounds, const double& CleanupRadiusSquared)
 	{
@@ -909,7 +960,7 @@ void FPCGRuntimeGenScheduler::TickCleanup(const TSet<IPCGGenSourceBase*>& InGenS
 
 			const double SquaredDistToGenSource = GridBounds.ComputeSquaredDistanceToPoint(GenSourcePosition);
 
-			// If the distance to the gen source is greater than the cleanup radius, it means this generation source votes for the component to be cleaned up.
+			// If the distance to the gen source is greater than the cleanup radius, it means this generation source votes for the execution source to be cleaned up.
 			// Otherwise, the gen source might still vote for culling regardless.
 			if (SquaredDistToGenSource <= CleanupRadiusSquared && 
 				(!Policy || !Policy->ShouldCull(GenSource, GridBounds, GridDescriptor.Is2DGrid())))
@@ -922,30 +973,30 @@ void FPCGRuntimeGenScheduler::TickCleanup(const TSet<IPCGGenSourceBase*>& InGenS
 		return bAllGenSourcesWantToCleanup;
 	};
 
-	TArray<FGridGenerationKey> GeneratedComponentsArray = GeneratedComponents.Array();
-	const int32 NumGeneratedComponents = GeneratedComponentsArray.Num();
+	TArray<FGridGenerationKey> GeneratedSourcesArray = GeneratedSources.Array();
+	const int32 NumGeneratedSources = GeneratedSourcesArray.Num();
 
-	// Generated Entry Key, Local/Generated Component
-	using PCGComponentToClean = TTuple<FGridGenerationKey, IPCGGraphExecutionSource*>;
-	TArray<PCGComponentToClean, TInlineAllocator<256>> ComponentsToClean;
+	// Generated Entry Key, Local/Generated Source
+	using ExecutionSourceToClean = TTuple<FGridGenerationKey, IPCGGraphExecutionSource*>;
+	TArray<ExecutionSourceToClean, TInlineAllocator<256>> SourcesToClean;
 	TArray<FGridGenerationKey, TInlineAllocator<16>> InvalidKeys;
 
-	ComponentsToClean.SetNumUninitialized(NumGeneratedComponents);
-	InvalidKeys.SetNumUninitialized(NumGeneratedComponents);
+	SourcesToClean.SetNumUninitialized(NumGeneratedSources);
+	InvalidKeys.SetNumUninitialized(NumGeneratedSources);
 
-	std::atomic<int> ComponentsToCleanCounter = 0;
+	std::atomic<int> SourcesToCleanCounter = 0;
 	std::atomic<int> InvalidKeysCounter = 0;
 
 	ParallelFor(
 		TEXT("CleanupLoop"),
-		NumGeneratedComponents,
+		NumGeneratedSources,
 		/*MinBatchSize*/6, // Optimal value found in testing.
-		[&GeneratedComponentsArray, &ComponentsToClean, &InvalidKeys, &ComponentsToCleanCounter, &InvalidKeysCounter, TrackingManager=TrackingManager, &CheckIfAllGenSourcesWantToCleanup](int32 Index)
+		[this, &GeneratedSourcesArray, &SourcesToClean, &InvalidKeys, &SourcesToCleanCounter, &InvalidKeysCounter, &CheckIfAllGenSourcesWantToCleanup](int32 Index)
 		{
-			TRACE_CPUPROFILER_EVENT_SCOPE(SelectComponentForCleanup);
+			TRACE_CPUPROFILER_EVENT_SCOPE(SelectSourceForCleanup);
 			LLM_SCOPE_BYTAG(PCG);
 
-			const FGridGenerationKey& GenerationKey = GeneratedComponentsArray[Index];
+			const FGridGenerationKey& GenerationKey = GeneratedSourcesArray[Index];
 
 			if (!GenerationKey.IsValid())
 			{
@@ -957,112 +1008,127 @@ void FPCGRuntimeGenScheduler::TickCleanup(const TSet<IPCGGenSourceBase*>& InGenS
 			const FPCGGridDescriptor GridDescriptor = GenerationKey.GetGridDescriptor();
 			const EPCGHiGenGrid Grid = GridDescriptor.GetHiGenGrid();
 			const FIntVector& GridCoords = GenerationKey.GetGridCoords();
-			UPCGComponent* OriginalComponent = GenerationKey.GetOriginalComponent();
-			check(OriginalComponent);
+			IPCGGraphExecutionSource* OriginalSource = GenerationKey.GetOriginalSource();
+			check(OriginalSource);
 
-			const UPCGSchedulingPolicyBase* Policy = OriginalComponent->GetRuntimeGenSchedulingPolicy();
+			const UPCGSchedulingPolicyBase* Policy = OriginalSource->GetExecutionState().GetRuntimeGenSchedulingPolicy();
 			ensure(Policy);
 
-			const double CleanupRadius = OriginalComponent->GetCleanupRadiusFromGrid(Grid);
+			const double CleanupRadius = OriginalSource->GetExecutionState().GetCleanupRadiusFromGrid(Grid);
 			const double CleanupRadiusSquared = CleanupRadius * CleanupRadius;
 
-			// If the Grid is unbounded, we have a non-partitioned or unbounded component.
+			// If the Grid is unbounded, we have a non-partitioned or unbounded execution source.
 			if (Grid == EPCGHiGenGrid::Unbounded)
 			{
-				if (!OriginalComponent->bActivated)
+				if (!OriginalSource->GetExecutionState().IsActive())
+				{
+					const int WriteIndex = SourcesToCleanCounter.fetch_add(1);
+					SourcesToClean[WriteIndex] = { GenerationKey, OriginalSource };
+					return;
+				}
+
+				// If the component is still generated while in actor-componentless mode, clean up!
+				if (UPCGComponent* PCGComponent = Cast<UPCGComponent>(OriginalSource); PCGComponent && PCGComponent->UseActorComponentlessGeneration())
 				{
-					const int WriteIndex = ComponentsToCleanCounter.fetch_add(1);
-					ComponentsToClean[WriteIndex] = { GenerationKey, OriginalComponent };
+					const int WriteIndex = SourcesToCleanCounter.fetch_add(1);
+					SourcesToClean[WriteIndex] = { GenerationKey, OriginalSource };
 					return;
 				}
 
-				const FBox GridBounds = GenerationKey.GetCachedBounds().IsSet() ? GenerationKey.GetCachedBounds().GetValue() : OriginalComponent->GetGridBounds();
+				const FBox GridBounds = GenerationKey.GetCachedBounds().IsSet() ? GenerationKey.GetCachedBounds().GetValue() : OriginalSource->GetExecutionState().GetBounds();
 
 				// Only clean up if all generation sources agreed to clean up.
 				if (CheckIfAllGenSourcesWantToCleanup(Policy, GridDescriptor, GridBounds, CleanupRadiusSquared))
 				{
-					const int WriteIndex = ComponentsToCleanCounter.fetch_add(1);
-					ComponentsToClean[WriteIndex] = { GenerationKey, OriginalComponent };
+					const int WriteIndex = SourcesToCleanCounter.fetch_add(1);
+					SourcesToClean[WriteIndex] = { GenerationKey, OriginalSource };
 					return;
 				}
 			}
-			// Otherwise, we have a local component.
+			// Otherwise, we have a local execution source.
 			else
 			{
-				IPCGGraphExecutionSource* LocalComponent = GenerationKey.GetCachedLocalComponent();
-				if (!LocalComponent)
+				IPCGGraphExecutionSource* LocalSource = GenerationKey.GetCachedLocalSource();
+
+				if (!LocalSource)
 				{
-					LocalComponent = TrackingManager->GetLocalComponent(GridDescriptor, GridCoords, OriginalComponent);
+					LocalSource = OriginalSource->GetExecutionState().GetLocalSource(GridDescriptor, GridCoords);
 				}
 
-				APCGPartitionActor* PartitionActor = LocalComponent ? Cast<APCGPartitionActor>(CastChecked<UPCGComponent>(LocalComponent)->GetOwner()) : nullptr;
-				if (!PartitionActor || !OriginalComponent->bActivated)
+				if (!LocalSource || !OriginalSource->GetExecutionState().IsActive())
 				{
-					// Attempt to clean even in failure case to avoid leaking resources.
-					const int WriteIndex = ComponentsToCleanCounter.fetch_add(1);
-					ComponentsToClean[WriteIndex] = { GenerationKey, LocalComponent };
+					const int WriteIndex = SourcesToCleanCounter.fetch_add(1);
+					SourcesToClean[WriteIndex] = { GenerationKey, LocalSource };
 					return;
 				}
 
-				const FBox GridBounds = GenerationKey.GetCachedBounds().IsSet() ? GenerationKey.GetCachedBounds().GetValue() : LocalComponent->GetExecutionState().GetBounds();
+				// If the component is still generated while in actor-componentless mode, clean up!
+				if (UPCGComponent* PCGComponent = Cast<UPCGComponent>(OriginalSource); PCGComponent && PCGComponent->UseActorComponentlessGeneration())
+				{
+					const int WriteIndex = SourcesToCleanCounter.fetch_add(1);
+					SourcesToClean[WriteIndex] = { GenerationKey, LocalSource };
+					return;
+				}
+
+				const FBox GridBounds = GenerationKey.GetCachedBounds().IsSet() ? GenerationKey.GetCachedBounds().GetValue() : LocalSource->GetExecutionState().GetBounds();
 
 				// Only clean up if all generation sources agreed to clean up.
 				if (CheckIfAllGenSourcesWantToCleanup(Policy, GridDescriptor, GridBounds, CleanupRadiusSquared))
 				{
-					const int WriteIndex = ComponentsToCleanCounter.fetch_add(1);
-					ComponentsToClean[WriteIndex] = { GenerationKey, LocalComponent };
+					const int WriteIndex = SourcesToCleanCounter.fetch_add(1);
+					SourcesToClean[WriteIndex] = { GenerationKey, LocalSource };
 					return;
 				}
 			}
 		});
 
 	const int NumInvalidKeys = InvalidKeysCounter.load();
-	const int NumComponentsToClean = ComponentsToCleanCounter.load();
+	const int NumSourcesToClean = SourcesToCleanCounter.load();
 
 	for (int I = 0; I < NumInvalidKeys; ++I)
 	{
-		GeneratedComponents.Remove(InvalidKeys[I]);
+		GeneratedSources.Remove(InvalidKeys[I]);
 	}
 
 	int CleanupIndex = 0;
-	for (; CleanupIndex < NumComponentsToClean; ++CleanupIndex)
+	for (; CleanupIndex < NumSourcesToClean; ++CleanupIndex)
 	{
-		const PCGComponentToClean& ComponentToClean = ComponentsToClean[CleanupIndex];
-		CleanupComponent(ComponentToClean.Get<0>(), ComponentToClean.Get<1>());
+		const ExecutionSourceToClean& SourceToClean = SourcesToClean[CleanupIndex];
+		CleanupSource(SourceToClean.Get<0>(), SourceToClean.Get<1>());
 
 		if (FPlatformTime::Seconds() >= InEndTime)
 		{
-			UE_LOG(LogPCG, Verbose, TEXT("FPCGRuntimeGenScheduler: Time budget exceeded, aborted after cleaning up %d / %d components"), (CleanupIndex + 1), NumComponentsToClean);
+			UE_LOG(LogPCG, Verbose, TEXT("FPCGRuntimeGenScheduler: Time budget exceeded, aborted after cleaning up %d / %d components"), (CleanupIndex + 1), NumSourcesToClean);
 			break;
 		}
 	}
 
-	if (CleanupIndex == NumComponentsToClean)
+	if (CleanupIndex == NumSourcesToClean)
 	{
 		// Record that we completed a cleanup scan.
 		bCleanupScanRequired = false;
 	}
 }
 
-void FPCGRuntimeGenScheduler::TickScheduleGeneration(TMap<FGridGenerationKey, double>& InOutComponentsToGenerate)
+void FPCGRuntimeGenScheduler::TickScheduleGeneration(TMap<FGridGenerationKey, double>& InOutSourcesToGenerate)
 {
 	TRACE_CPUPROFILER_EVENT_SCOPE(FPCGRuntimeGenScheduler::TickScheduleGeneration);
 
 	check(Subsystem && TrackingManager);
 
-	// Count number of currently generating components
+	// Count number of currently generating execution sources
 	int NumGenerating = 0;
 
 	{
 		TRACE_CPUPROFILER_EVENT_SCOPE(FPCGRuntimeGenScheduler::TickScheduleGeneration::CountingCurrentlyGenerating);
-		for (const FGridGenerationKey& Key : GeneratedComponents)
+		for (const FGridGenerationKey& Key : GeneratedSources)
 		{
 			if (!Key.IsValid())
 			{
 				continue;
 			}
 
-			const IPCGGraphExecutionSource* ExecutionSource = (Key.GetGridDescriptor().GetHiGenGrid() == EPCGHiGenGrid::Unbounded) ? Key.GetOriginalComponent() : Key.GetCachedLocalComponent();
+			const IPCGGraphExecutionSource* ExecutionSource = (Key.GetGridDescriptor().GetHiGenGrid() == EPCGHiGenGrid::Unbounded) ? Key.GetOriginalSource() : Key.GetCachedLocalSource();
 
 			if (ExecutionSource && ExecutionSource->GetExecutionState().IsGenerating())
 			{
@@ -1076,24 +1142,24 @@ void FPCGRuntimeGenScheduler::TickScheduleGeneration(TMap<FGridGenerationKey, do
 
 	if (GeneratingCapacity == 0)
 	{
-		// No capacity to start generating new components, try next frame.
+		// No capacity to start generating new execution sources, try next frame.
 		return;
 	}
 
-	GeneratedComponents.Reserve(GeneratedComponents.Num() + InOutComponentsToGenerate.Num());
+	GeneratedSources.Reserve(GeneratedSources.Num() + InOutSourcesToGenerate.Num());
 
-	// Components that we couldn't generate due to missed capacity. Generation will be attempted again next frame.
-	TMap<FGridGenerationKey, double> MissedComponents;
-	MissedComponents.Reserve(FMath::Max(0, InOutComponentsToGenerate.Num() - GeneratingCapacity));
+	// Sources that we couldn't generate due to missed capacity. Generation will be attempted again next frame.
+	TMap<FGridGenerationKey, double> MissedSources;
+	MissedSources.Reserve(FMath::Max(0, InOutSourcesToGenerate.Num() - GeneratingCapacity));
 
-	for (auto It = InOutComponentsToGenerate.CreateConstIterator(); It; ++It)
+	for (auto It = InOutSourcesToGenerate.CreateConstIterator(); It; ++It)
 	{
 		if (NumGenerating >= MaxNumGenerating)
 		{
-			// We've filled generation capacity, record the remaining components to try again next frame.
+			// We've filled generation capacity, record the remaining execution sources to try again next frame.
 			for (; It; ++It)
 			{
-				MissedComponents.Add(*It);
+				MissedSources.Add(*It);
 			}
 
 			break;
@@ -1106,115 +1172,134 @@ void FPCGRuntimeGenScheduler::TickScheduleGeneration(TMap<FGridGenerationKey, do
 		const EPCGHiGenGrid Grid = GridDescriptor.GetHiGenGrid();
 
 		const FIntVector GridCoords = Key.GetGridCoords();
-		UPCGComponent* OriginalComponent = Key.GetOriginalComponent();
-		if (!OriginalComponent)
+		IPCGGraphExecutionSource* OriginalSource = Key.GetOriginalSource();
+		if (!OriginalSource)
 		{
-			UE_LOG(LogPCG, Verbose, TEXT("TickScheduleGeneration: Invalid original component in an InOutComponentsToGenerate entry, generation will be skipped."));
+			UE_LOG(LogPCG, Verbose, TEXT("TickScheduleGeneration: Invalid original component in an InOutSourcesToGenerate entry, generation will be skipped."));
 			continue;
 		}
 
-		// If the Grid is unbounded, we have a non-partitioned or unbounded component.
+		IPCGGraphExecutionState& OriginalExecutionState = OriginalSource->GetExecutionState();
+
+		// If the Grid is unbounded, we have a non-partitioned or unbounded execution source.
 		if (Grid == EPCGHiGenGrid::Unbounded)
 		{
-			if (!OriginalComponent->IsGenerating())
+			if (!OriginalExecutionState.IsGenerating())
 			{
-				if (PCGRuntimeGenSchedulerHelpers::CVarRuntimeGenerationEnableDebugging.GetValueOnAnyThread() && OriginalComponent->GetOwner())
+				if (PCGRuntimeGenSchedulerHelpers::CVarRuntimeGenerationEnableDebugging.GetValueOnAnyThread())
 				{
-					UE_LOG(LogPCG, Warning, TEXT("[RUNTIMEGEN] GENERATE: '%s' (priority %lf)"), *OriginalComponent->GetOwner()->GetActorNameOrLabel(), Priority);
+					UE_LOG(LogPCG, Warning, TEXT("[RUNTIMEGEN] GENERATE: '%s' (priority %lf)"), *OriginalExecutionState.GetDebugName(), Priority);
 				}
 
 #if PCG_RGS_ONSCREENDEBUGMESSAGES
 				++PCGRuntimeGenSchedulerHelpers::Stats.GenerateCallCounter;
 #endif
 
-				// Force to refresh if the component is already generated.
-				OriginalComponent->GenerateLocal(EPCGComponentGenerationTrigger::GenerateAtRuntime, /*bForce=*/true, Grid);
+				IPCGGraphExecutionState::FGenerateParams GenerateParams
+				{
+					.bEvenIfAlreadyGenerated = true,
+					.bGenerateLocalSources = false
+				};
+
+				OriginalSource->GetExecutionState().Generate(GenerateParams);
 			}
 		}
-		// Otherwise we have a local component.
+		// Otherwise we have a local execution source.
 		else
 		{
-			// Grab local component and PA if they exist already.
-			IPCGGraphExecutionSource* LocalComponent = TrackingManager->GetLocalComponent(GridDescriptor, GridCoords, OriginalComponent);
-			TObjectPtr<APCGPartitionActor> PartitionActor = LocalComponent ? Cast<APCGPartitionActor>(CastChecked<UPCGComponent>(LocalComponent)->GetOwner()) : nullptr;
+			IPCGGraphExecutionSource* LocalSource = OriginalExecutionState.GetLocalSource(GridDescriptor, GridCoords);
 
-			if (!LocalComponent || !ensure(PartitionActor))
+			if (!LocalSource)
 			{
-				// Local component & PA do not exist, create them.
-				if (PCGRuntimeGenSchedulerHelpers::CVarRuntimeGenerationEnablePooling.GetValueOnAnyThread())
+				// @todo_pcg: CreateLocalSource could be moved into the graph execution state interface and generalized for all execution sources.
+				if (UPCGRuntimeGenExecutionSource* OriginalRuntimeGenSource = Cast<UPCGRuntimeGenExecutionSource>(OriginalSource))
 				{
-					// Get RuntimeGenPA from pool.
-					PartitionActor = GetPartitionActorFromPool(GridDescriptor, GridCoords);
-
-					if (PCGRuntimeGenSchedulerHelpers::CVarRuntimeGenerationEnableDebugging.GetValueOnAnyThread())
-					{
-						UE_LOG(LogPCG, Warning, TEXT("[RUNTIMEGEN] UNPOOL PARTITION ACTOR: '%s' (priority %lf, %d remaining out of %d)"),
-							*APCGPartitionActor::GetPCGPartitionActorName(GridDescriptor, GridCoords),
-							Priority,
-							PartitionActorPool.Num(),
-							PartitionActorPoolSize);
-					}
+					LocalSource = OriginalRuntimeGenSource->CreateLocalSource(GridDescriptor.GetGridSize(), GridCoords);
 				}
-				else
+				else if (UPCGComponent* OriginalComponent = Cast<UPCGComponent>(OriginalSource))
 				{
-					if (PCGRuntimeGenSchedulerHelpers::CVarRuntimeGenerationEnableDebugging.GetValueOnAnyThread())
+					// Grab local component and PA if they exist already.
+					TObjectPtr<APCGPartitionActor> PartitionActor = LocalSource ? LocalSource->GetExecutionState().GetTypedTarget<APCGPartitionActor>() : nullptr;
+
+					if (!LocalSource || !ensure(PartitionActor))
 					{
-						UE_LOG(LogPCG, Warning, TEXT("[RUNTIMEGEN] CREATE PARTITION ACTOR: '%s' (priority %lf)"),
-							*APCGPartitionActor::GetPCGPartitionActorName(GridDescriptor, GridCoords),
-							Priority);
-					}
+						// Local component & PA do not exist, create them.
+						if (PCGRuntimeGenSchedulerHelpers::CVarRuntimeGenerationEnablePooling.GetValueOnAnyThread())
+						{
+							// Get RuntimeGenPA from pool.
+							PartitionActor = GetPartitionActorFromPool(GridDescriptor, GridCoords);
 
-					// Find or Create RuntimeGenPA.
-					PartitionActor = Subsystem->FindOrCreatePCGPartitionActor(
-						GridDescriptor,
-						GridCoords,
-						/*bCanCreateActor=*/true,
-						PCGRuntimeGenSchedulerHelpers::CVarHideActorsFromOutliner.GetValueOnAnyThread());
-				}
+							if (PCGRuntimeGenSchedulerHelpers::CVarRuntimeGenerationEnableDebugging.GetValueOnAnyThread())
+							{
+								UE_LOG(LogPCG, Warning, TEXT("[RUNTIMEGEN] UNPOOL PARTITION ACTOR: '%s' (priority %lf, %d remaining out of %d)"),
+									*APCGPartitionActor::GetPCGPartitionActorName(GridDescriptor, GridCoords),
+									Priority,
+									PartitionActorPool.Num(),
+									PartitionActorPoolSize);
+							}
+						}
+						else
+						{
+							if (PCGRuntimeGenSchedulerHelpers::CVarRuntimeGenerationEnableDebugging.GetValueOnAnyThread())
+							{
+								UE_LOG(LogPCG, Warning, TEXT("[RUNTIMEGEN] CREATE PARTITION ACTOR: '%s' (priority %lf)"),
+									*APCGPartitionActor::GetPCGPartitionActorName(GridDescriptor, GridCoords),
+									Priority);
+							}
 
-				if (!ensure(PartitionActor))
-				{
-					continue;
-				}
+							// Find or Create RuntimeGenPA.
+							PartitionActor = Subsystem->FindOrCreatePCGPartitionActor(
+								GridDescriptor,
+								GridCoords,
+								/*bCanCreateActor=*/true,
+								PCGRuntimeGenSchedulerHelpers::CVarHideActorsFromOutliner.GetValueOnAnyThread());
+						}
 
-				// Update component mapping for this PA (add local component).
-				{
-					PCG::TUniqueScopeLock WriteLock(TrackingManager->ComponentToPartitionActorsMapLock);
-					TSet<TObjectPtr<APCGPartitionActor>>* PartitionActorsPtr = TrackingManager->ComponentToPartitionActorsMap.Find(OriginalComponent);
+						if (!ensure(PartitionActor))
+						{
+							continue;
+						}
 
-					if (!PartitionActorsPtr)
-					{
-						PartitionActorsPtr = &TrackingManager->ComponentToPartitionActorsMap.Emplace(OriginalComponent);
-					}
+						// Update component mapping for this PA (add local component).
+						{
+							PCG::TUniqueScopeLock WriteLock(TrackingManager->ComponentToPartitionActorsMapLock);
+							TSet<TObjectPtr<APCGPartitionActor>>* PartitionActorsPtr = TrackingManager->ComponentToPartitionActorsMap.Find(OriginalComponent);
+
+							if (!PartitionActorsPtr)
+							{
+								PartitionActorsPtr = &TrackingManager->ComponentToPartitionActorsMap.Emplace(OriginalComponent);
+							}
 
-					// Log this original component before setting up the PA, so that we early out from RefreshComponent if it gets called
-					// in the AddGraphInstance call below.
-					OriginalComponentBeingGenerated = OriginalComponent;
+							// Log this original component before setting up the PA, so that we early out from RefreshExecutionSource if it gets called
+							// in the AddGraphInstance call below.
+							OriginalSourceBeingRegistered = OriginalComponent;
 
-					PartitionActor->AddGraphInstance(OriginalComponent);
+							PartitionActor->AddGraphInstance(OriginalComponent);
 
-					OriginalComponentBeingGenerated = nullptr;
+							OriginalSourceBeingRegistered = nullptr;
 
-					PartitionActorsPtr->Add(PartitionActor);
-				}
+							PartitionActorsPtr->Add(PartitionActor);
+						}
 
-				// Create local component.
-				LocalComponent = PartitionActor->GetLocalComponent(OriginalComponent);
+						// Create local component.
+						LocalSource = PartitionActor->GetLocalComponent(OriginalComponent);
+					}
+				}
 			}
 
-			if (ensure(LocalComponent) && !LocalComponent->GetExecutionState().IsGenerating())
+			if (ensure(LocalSource) && !LocalSource->GetExecutionState().IsGenerating())
 			{
 				if (PCGRuntimeGenSchedulerHelpers::CVarRuntimeGenerationEnableDebugging.GetValueOnAnyThread())
 				{
-					UE_LOG(LogPCG, Warning, TEXT("[RUNTIMEGEN] GENERATE: '%s' (priority %lf)"), *PartitionActor->GetActorNameOrLabel(), Priority);
+					UE_LOG(LogPCG, Warning, TEXT("[RUNTIMEGEN] GENERATE: '%s' (priority %lf)"), *LocalSource->GetExecutionState().GetDebugName(), Priority);
 				}
 
-				// Higen graphs may have data links from original component to local components. The original component will be given a higher priority than local
-				// components and will start generating first. If it is currently generating, local component needs to take a dependency to ensure execution completes.
+				// Higen graphs may have data links from original execution source to local execution sources. The original execution source will be given a higher priority than local
+				// execution sources and will start generating first. If it is currently generating, local execution source needs to take a dependency to ensure execution completes.
 				TArray<FPCGTaskId> Dependencies;
-				if (OriginalComponent->IsGenerating() && OriginalComponent->GetGraph() && OriginalComponent->GetGraph()->IsHierarchicalGenerationEnabled())
+				if (OriginalExecutionState.IsGenerating() && OriginalExecutionState.GetGraph() && OriginalExecutionState.GetGraph()->IsHierarchicalGenerationEnabled())
 				{
-					const FPCGTaskId TaskId = OriginalComponent->GetGenerationTaskId();
+					const FPCGTaskId TaskId = OriginalExecutionState.GetGenerationTaskId();
 
 					if (TaskId != InvalidPCGTaskId)
 					{
@@ -1226,23 +1311,23 @@ void FPCGRuntimeGenScheduler::TickScheduleGeneration(TMap<FGridGenerationKey, do
 				++PCGRuntimeGenSchedulerHelpers::Stats.GenerateCallCounter;
 #endif
 
-				// Force to refresh if the component is already generated.
+				// Force to refresh if the execution source is already generated.
 				IPCGGraphExecutionState::FGenerateParams GenerateParams;
 				GenerateParams.bEvenIfAlreadyGenerated = true;
 				GenerateParams.Dependencies = MoveTemp(Dependencies);
 
-				LocalComponent->GetExecutionState().Generate(GenerateParams);
+				LocalSource->GetExecutionState().Generate(GenerateParams);
 
-				Key.SetCachedLocalComponent(LocalComponent);
+				Key.SetCachedLocalSource(LocalSource);
 			}
 		}
 
-		GeneratedComponents.Add(Key);
+		GeneratedSources.Add(Key);
 		++NumGenerating;
 	}
 
-	// Set any missed components to retry next frame (or clear the pending components if no components were missed).
-	InOutComponentsToGenerate = MoveTemp(MissedComponents);
+	// Set any missed execution sources to retry next frame (or clear the pending execution sources if no execution sources were missed).
+	InOutSourcesToGenerate = MoveTemp(MissedSources);
 
 #if PCG_RGS_ONSCREENDEBUGMESSAGES
 	PCGRuntimeGenSchedulerHelpers::Stats.NumGeneratingComponents = NumGenerating;
@@ -1253,8 +1338,6 @@ void FPCGRuntimeGenScheduler::TickRequestVirtualTexturePriming(const TSet<IPCGGe
 {
 	TRACE_CPUPROFILER_EVENT_SCOPE(FPCGRuntimeGenScheduler::TickRequestVirtualTexturePriming);
 
-	check(TrackingManager);
-
 	TMap<TSoftObjectPtr<URuntimeVirtualTexture>, TArray<URuntimeVirtualTextureComponent*>> VirtualTextureToComponents;
 
 	{
@@ -1383,7 +1466,7 @@ void FPCGRuntimeGenScheduler::TickCVars(const APCGWorldActor* InPCGWorldActor)
 {
 	if (bActorFlushRequested && Subsystem && Subsystem->GetPCGWorldActor())
 	{
-		CleanupLocalComponents(Subsystem->GetPCGWorldActor());
+		CleanupLocalSources(Subsystem->GetPCGWorldActor());
 		ResetPartitionActorPoolToSize(PCGRuntimeGenSchedulerHelpers::CVarRuntimeGenerationBasePoolSize.GetValueOnAnyThread());
 	}
 	bActorFlushRequested = false;
@@ -1393,13 +1476,13 @@ void FPCGRuntimeGenScheduler::TickCVars(const APCGWorldActor* InPCGWorldActor)
 
 	if (bPoolingWasEnabledLastFrame && !bPoolingEnabled)
 	{
-		CleanupLocalComponents(InPCGWorldActor);
+		CleanupLocalSources(InPCGWorldActor);
 		ResetPartitionActorPoolToSize(/*NewPoolSize=*/0);
 	}
 
 	bPoolingWasEnabledLastFrame = bPoolingEnabled;
 
-	// Handle when the base PA PoolSize is modified. Cleanup all local components and reset the pool with the correct number of PAs.
+	// Handle when the base PA PoolSize is modified. Cleanup all local execution sources and reset the pool with the correct number of PAs.
 	if (PCGRuntimeGenSchedulerHelpers::CVarRuntimeGenerationEnablePooling.GetValueOnAnyThread())
 	{
 		// Don't allow a pool size <= 0
@@ -1409,7 +1492,7 @@ void FPCGRuntimeGenScheduler::TickCVars(const APCGWorldActor* InPCGWorldActor)
 		{
 			BasePoolSizeLastFrame = BasePoolSize;
 
-			CleanupLocalComponents(InPCGWorldActor);
+			CleanupLocalSources(InPCGWorldActor);
 			ResetPartitionActorPoolToSize(BasePoolSize);
 		}
 	}
@@ -1430,7 +1513,7 @@ void FPCGRuntimeGenScheduler::TickCVars(const APCGWorldActor* InPCGWorldActor)
 void FPCGRuntimeGenScheduler::OnOriginalComponentRegistered(UPCGComponent* InOriginalComponent)
 {
 	// Ensure we are non-local runtime managed component.
-	if (!InOriginalComponent || !InOriginalComponent->IsManagedByRuntimeGenSystem() || Cast<APCGPartitionActor>(InOriginalComponent->GetOwner()))
+	if (!InOriginalComponent || !InOriginalComponent->IsManagedByRuntimeGenSystem() || InOriginalComponent->IsLocalComponent())
 	{
 		return;
 	}
@@ -1443,17 +1526,18 @@ void FPCGRuntimeGenScheduler::OnOriginalComponentRegistered(UPCGComponent* InOri
 #endif
 
 	CacheVirtualTexturePrimingInfos(InOriginalComponent);
+
+	// Try to register as an RG execution source. If this fails, we will treat the original PCG component as the execution source.
+	RegisterOriginalRuntimeGenExecutionSource(InOriginalComponent);
 }
 
 void FPCGRuntimeGenScheduler::OnOriginalComponentUnregistered(UPCGComponent* InOriginalComponent)
 {
-	if (!InOriginalComponent || Cast<APCGPartitionActor>(InOriginalComponent->GetOwner()))
+	if (!InOriginalComponent || InOriginalComponent->IsLocalComponent())
 	{
 		return;
 	}
 
-	check(TrackingManager);
-
 	// When an original/non-partitioned component is unregistered, we need to dirty the state.
 	bContextDirty = true;
 
@@ -1468,9 +1552,9 @@ void FPCGRuntimeGenScheduler::OnOriginalComponentUnregistered(UPCGComponent* InO
 
 	// Gather all generated components which originated from this original component.
 	TSet<FGridGenerationKey> KeysToCleanup;
-	for (FGridGenerationKey GenerationKey : GeneratedComponents)
+	for (FGridGenerationKey GenerationKey : GeneratedSources)
 	{
-		if (GenerationKey.GetOriginalComponent() == InOriginalComponent)
+		if (GenerationKey.GetOriginalSource() == InOriginalComponent)
 		{
 			KeysToCleanup.Add(GenerationKey);
 		}
@@ -1487,24 +1571,62 @@ void FPCGRuntimeGenScheduler::OnOriginalComponentUnregistered(UPCGComponent* InO
 		}
 
 		const FPCGGridDescriptor GridDescriptor = GenerationKey.GetGridDescriptor();
-		const bool bIsOriginalComponent = GridDescriptor.GetGridSize() == PCGHiGenGrid::UnboundedGridSize();
+		const bool bIsOriginalSource = GridDescriptor.GetGridSize() == PCGHiGenGrid::UnboundedGridSize();
 
-		// Get the generated component for this key (might be a local component).
-		IPCGGraphExecutionSource* ComponentToCleanup = bIsOriginalComponent ? InOriginalComponent : TrackingManager->GetLocalComponent(GridDescriptor, GenerationKey.GetGridCoords(), InOriginalComponent);
+		// Get the generated execution source for this key (might be a local execution source).
+		IPCGGraphExecutionSource* SourceToCleanup = bIsOriginalSource ? InOriginalComponent : InOriginalComponent->GetExecutionState().GetLocalSource(GridDescriptor, GenerationKey.GetGridCoords());
 
-		// It is possible for a PartitionActor's LocalComponent to have been cleaned up by the APCGPartitionActor::EndPlay call depending on the order in which actors get called
-		if (ComponentToCleanup)
+		// It is possible for a PartitionActor's LocalSource to have been cleaned up by the APCGPartitionActor::EndPlay call depending on the order in which actors get called
+		if (SourceToCleanup)
 		{
-			CleanupComponent(GenerationKey, ComponentToCleanup);
+			CleanupSource(GenerationKey, SourceToCleanup);
 		}
 	}
 
 	for (const FGridGenerationKey& InvalidKey : InvalidKeys)
 	{
-		GeneratedComponents.Remove(InvalidKey);
+		GeneratedSources.Remove(InvalidKey);
 	}
 
-	CleanupRemainingComponents(InOriginalComponent);
+	CleanupRemainingSources(InOriginalComponent);
+
+	UnregisterOriginalRuntimeGenExecutionSource(InOriginalComponent);
+}
+
+void FPCGRuntimeGenScheduler::OnOriginalComponentReplaced(const UPCGComponent* InOldComponent, UPCGComponent* InNewComponent)
+{
+	if (!InOldComponent || !InNewComponent)
+	{
+		return;
+	}
+
+	for (FGridGenerationKey& GenerationKey : GeneratedSources)
+	{
+		if (GenerationKey.GetOriginalSourceEvenIfGarbage() == InOldComponent)
+		{
+			GenerationKey.SetOriginalSource(InNewComponent);
+		}
+	}
+
+	TObjectPtr<UPCGRuntimeGenExecutionSource> OriginalExecutionSource = nullptr;
+	if (OriginalExecutionSources.RemoveAndCopyValue(InOldComponent, OriginalExecutionSource))
+	{
+		OriginalExecutionSources.Add(InNewComponent, OriginalExecutionSource);
+
+		// Make sure the runtime gen execution source is using the new PCG component.
+		OriginalExecutionSource->SetOwningSource(InNewComponent);
+	}
+
+	TArray<FPCGVirtualTexturePrimingInfo> PrimingInfos;
+	if (CachedPrimingInfos.RemoveAndCopyValue(InOldComponent, PrimingInfos))
+	{
+		CachedPrimingInfos.Add(InNewComponent, PrimingInfos);
+	}
+
+#if WITH_EDITOR
+	UnregisterGraphParameterChangedEvent(InOldComponent);
+	RegisterGraphParameterChangedEvent(InNewComponent);
+#endif
 }
 
 #if WITH_EDITOR
@@ -1516,88 +1638,70 @@ void FPCGRuntimeGenScheduler::OnObjectsReplaced(const TMap<UObject*, UObject*>&
 		UPCGComponent* NewComponent = Cast<UPCGComponent>(OldToNew.Value);
 
 		// Replace all usage of OldComponent with NewComponent.
-		if (OldComponent && NewComponent)
-		{
-			for (FGridGenerationKey& GenerationKey : GeneratedComponents)
-			{
-				if (GenerationKey.GetOriginalComponentEvenIfGarbage() == OldComponent)
-				{
-					GenerationKey.SetOriginalComponent(NewComponent);
-				}
-			}
-
-			TArray<FPCGVirtualTexturePrimingInfo> PrimingInfos;
-			if (CachedPrimingInfos.RemoveAndCopyValue(OldComponent, PrimingInfos))
-			{
-				CachedPrimingInfos.Add(NewComponent, PrimingInfos);
-			}
-
-			UnregisterGraphParameterChangedEvent(OldComponent);
-			RegisterGraphParameterChangedEvent(NewComponent);
-		}
+		OnOriginalComponentReplaced(OldComponent, NewComponent);
 	}
 }
 #endif
 
-FPCGRuntimeGenScheduler::FGridGenerationKey::FGridGenerationKey(uint32 InGridSize, const FIntVector& InGridCoords, UPCGComponent* InOriginalComponent, IPCGGraphExecutionSource* InLocalComponent)
+FPCGRuntimeGenScheduler::FGridGenerationKey::FGridGenerationKey(uint32 InGridSize, const FIntVector& InGridCoords, IPCGGraphExecutionSource* InOriginalSource, IPCGGraphExecutionSource* InLocalSource)
 	: GridSize(InGridSize)
 	, GridCoords(InGridCoords)
-	, OriginalComponent(InOriginalComponent)
-	, CachedLocalComponent(InLocalComponent)
+	, OriginalSource(Cast<UObject>(InOriginalSource))
+	, CachedLocalSource(InLocalSource)
 {
-	bUse2DGrid = InOriginalComponent ? InOriginalComponent->Use2DGrid() : true;
+	bUse2DGrid = InOriginalSource ? InOriginalSource->GetExecutionState().Use2DGrid() : true;
 
 	if (InGridSize == PCGHiGenGrid::UnboundedGridSize())
 	{
-		const AActor* Owner = InOriginalComponent ? InOriginalComponent->GetOwner() : nullptr;
+		const AActor* Owner = InOriginalSource ? InOriginalSource->GetExecutionState().GetTypedTarget<AActor>() : nullptr;
 		const USceneComponent* RootComponent = Owner ? Owner->GetRootComponent() : nullptr;
 
 		if (RootComponent && RootComponent->GetMobility() != EComponentMobility::Movable)
 		{
-			CachedBounds = InOriginalComponent->GetGridBounds();
+			CachedBounds = InOriginalSource->GetExecutionState().GetBounds();
 		}
 	}
-	else if (InLocalComponent)
+	else if (InLocalSource)
 	{
-		CachedBounds = InLocalComponent->GetExecutionState().GetBounds();
+		CachedBounds = InLocalSource->GetExecutionState().GetBounds();
 	}
 }
 
-void FPCGRuntimeGenScheduler::CleanupRemainingComponents(UPCGComponent* InOriginalComponent)
+void FPCGRuntimeGenScheduler::CleanupRemainingSources(IPCGGraphExecutionSource* InOriginalSource)
 {
 	// Check for remaining PAs to cleanup
-	// There are some cases when on Refresh of the original component that GeneratedComponents doesn't contain all PAs anymore.
+	// There are some cases when on Refresh of the original component that GeneratedSources doesn't contain all PAs anymore.
 	// If the PA is not far enough to be cleaned up but no longer is considered within the Gen Source Radius for Generation.
-	if (InOriginalComponent && InOriginalComponent->IsManagedByRuntimeGenSystem())
+	if (UPCGComponent* OriginalComponent = Cast<UPCGComponent>(InOriginalSource); OriginalComponent && InOriginalSource->GetExecutionState().IsManagedByRuntimeGenSystem())
 	{
-		TSet<TObjectPtr<APCGPartitionActor>> PartitionActors = TrackingManager->GetPCGComponentPartitionActorMappings(InOriginalComponent);
+		TSet<TObjectPtr<APCGPartitionActor>> PartitionActors = TrackingManager->GetPCGComponentPartitionActorMappings(OriginalComponent);
 		for (TObjectPtr<APCGPartitionActor> PartitionActor : PartitionActors)
 		{
-			if (IPCGGraphExecutionSource* ComponentToCleanup = PartitionActor->GetLocalComponent(InOriginalComponent))
+			if (IPCGGraphExecutionSource* ComponentToCleanup = PartitionActor->GetLocalComponent(OriginalComponent))
 			{
-				CleanupLocalComponent(PartitionActor, ComponentToCleanup);
+				CleanupLocalSource(PartitionActor, ComponentToCleanup);
 			}
 		}
 	}
 }
 
-void FPCGRuntimeGenScheduler::CleanupLocalComponents(const APCGWorldActor* InPCGWorldActor)
+void FPCGRuntimeGenScheduler::CleanupLocalSources(const APCGWorldActor* InPCGWorldActor)
 {
-	TRACE_CPUPROFILER_EVENT_SCOPE(FPCGRuntimeGenScheduler::CleanupLocalComponents);
+	TRACE_CPUPROFILER_EVENT_SCOPE(FPCGRuntimeGenScheduler::CleanupLocalSources);
 
-	check(TrackingManager && InPCGWorldActor);
+	check(InPCGWorldActor);
 
-	// Generated Entry Key, Local/Generated Component
-	using PCGComponentToClean = TTuple<FGridGenerationKey, IPCGGraphExecutionSource*>;
-	TArray<PCGComponentToClean, TInlineAllocator<16>> ComponentsToClean;
+	// Generated Entry Key, Local/Generated Source
+	using ExecutionSourceToClean = TTuple<FGridGenerationKey, IPCGGraphExecutionSource*>;
+	TArray<ExecutionSourceToClean, TInlineAllocator<16>> SourcesToClean;
 
-	TSet<UPCGComponent*> OriginalComponents;
-	OriginalComponents.Reserve(16);
+	TSet<IPCGGraphExecutionSource*> OriginalSources;
+	OriginalSources.Reserve(16);
 
 	TArray<FGridGenerationKey, TInlineAllocator<16>> InvalidKeys;
 
-	// Find all generated local components.
-	for (const FGridGenerationKey& GenerationKey : GeneratedComponents)
+	// Find all generated local execution sources.
+	for (const FGridGenerationKey& GenerationKey : GeneratedSources)
 	{
 		if (!GenerationKey.IsValid())
 		{
@@ -1608,31 +1712,31 @@ void FPCGRuntimeGenScheduler::CleanupLocalComponents(const APCGWorldActor* InPCG
 		const FPCGGridDescriptor GridDescriptor = GenerationKey.GetGridDescriptor();
 		const EPCGHiGenGrid Grid = GridDescriptor.GetHiGenGrid();
 		const FIntVector& GridCoords = GenerationKey.GetGridCoords();
-		UPCGComponent* OriginalComponent = GenerationKey.GetOriginalComponent();
-		check(OriginalComponent);
-		OriginalComponents.Add(OriginalComponent);
+		IPCGGraphExecutionSource* OriginalSource = GenerationKey.GetOriginalSource();
+		check(OriginalSource);
+		OriginalSources.Add(OriginalSource);
 
-		// Only operate on LocalComponents.
+		// Only operate on LocalSources.
 		if (Grid != EPCGHiGenGrid::Unbounded)
 		{
-			IPCGGraphExecutionSource* LocalComponent = TrackingManager->GetLocalComponent(GridDescriptor, GridCoords, OriginalComponent);
-			ComponentsToClean.Add({ GenerationKey, LocalComponent });
+			IPCGGraphExecutionSource* LocalSource = OriginalSource->GetExecutionState().GetLocalSource(GridDescriptor, GridCoords);
+			SourcesToClean.Add({ GenerationKey, LocalSource });
 		}
 	}
 
 	for (const FGridGenerationKey& InvalidKey : InvalidKeys)
 	{
-		GeneratedComponents.Remove(InvalidKey);
+		GeneratedSources.Remove(InvalidKey);
 	}
 
-	for (const PCGComponentToClean& ComponentToClean : ComponentsToClean)
+	for (const ExecutionSourceToClean& SourceToClean : SourcesToClean)
 	{
-		CleanupComponent(ComponentToClean.Get<0>(), ComponentToClean.Get<1>());
+		CleanupSource(SourceToClean.Get<0>(), SourceToClean.Get<1>());
 	}
 
-	for (UPCGComponent* OriginalComponent : OriginalComponents)
+	for (IPCGGraphExecutionSource* OriginalSource : OriginalSources)
 	{
-		CleanupRemainingComponents(OriginalComponent);
+		CleanupRemainingSources(OriginalSource);
 	}
 }
 
@@ -1645,106 +1749,117 @@ void FPCGRuntimeGenScheduler::OnLevelStreamingStateChanged(UWorld* InWorld, cons
 	}
 }
 
-void FPCGRuntimeGenScheduler::CleanupLocalComponent(APCGPartitionActor* PartitionActor, IPCGGraphExecutionSource* LocalComponent)
+void FPCGRuntimeGenScheduler::CleanupLocalSource(APCGPartitionActor* InPartitionActor, IPCGGraphExecutionSource* InLocalSource)
 {
-	if (!PartitionActor)
+	if (InLocalSource && InLocalSource->GetExecutionState().IsLocalSource())
 	{
-		return;
-	}
+		IPCGGraphExecutionState& ExecutionState = InLocalSource->GetExecutionState();
 
-	IPCGGraphExecutionSource* OriginalExecutionSource = LocalComponent ? LocalComponent->GetExecutionState().GetOriginalSource() : nullptr;
-	UPCGComponent* OriginalComponent = Cast<UPCGComponent>(OriginalExecutionSource);
-
-	if (OriginalComponent)
-	{
-		if (PCGRuntimeGenSchedulerHelpers::CVarRuntimeGenerationEnableDebugging.GetValueOnAnyThread())
+		if (IPCGGraphExecutionSource* OriginalExecutionSource = ExecutionState.GetOriginalSource())
 		{
-			UE_LOG(LogPCG, Warning, TEXT("[RUNTIMEGEN] CLEANUP: '%s'"), *PartitionActor->GetActorNameOrLabel());
-		}
-
-		// This performs a CleanupLocalImmediate for us, no need to clean up ourselves.
-		PartitionActor->RemoveGraphInstance(OriginalComponent);
+			if (PCGRuntimeGenSchedulerHelpers::CVarRuntimeGenerationEnableDebugging.GetValueOnAnyThread())
+			{
+				UE_LOG(LogPCG, Warning, TEXT("[RUNTIMEGEN] CLEANUP: '%s'"), *ExecutionState.GetDebugName());
+			}
 
 #if PCG_RGS_ONSCREENDEBUGMESSAGES
-		++PCGRuntimeGenSchedulerHelpers::Stats.CleanupCallCounter;
+			++PCGRuntimeGenSchedulerHelpers::Stats.CleanupCallCounter;
 #endif
 
-		// Remove component mapping.
-		{
-			PCG::TUniqueScopeLock WriteLock(TrackingManager->ComponentToPartitionActorsMapLock);
-			TSet<TObjectPtr<APCGPartitionActor>>* PartitionActorsPtr = TrackingManager->ComponentToPartitionActorsMap.Find(OriginalComponent);
-
-			if (PartitionActorsPtr)
+			// @todo_pcg: DestroyLocalSource could be moved into the graph execution state interface and generalized for all execution sources.
+			if (UPCGRuntimeGenExecutionSource* RuntimeExecutionSource = Cast<UPCGRuntimeGenExecutionSource>(OriginalExecutionSource))
 			{
-				PartitionActorsPtr->Remove(PartitionActor);
-
-				if (PartitionActorsPtr->IsEmpty())
+				ensure(!InPartitionActor);
+				RuntimeExecutionSource->DestroyLocalSource(Cast<UPCGRuntimeGenExecutionSource>(InLocalSource));
+			}
+			else if (UPCGComponent* OriginalComponent = Cast<UPCGComponent>(OriginalExecutionSource); InPartitionActor && ensure(OriginalComponent))
+			{
+				// This performs a CleanupLocalImmediate for us, no need to clean up ourselves.
+				InPartitionActor->RemoveGraphInstance(OriginalComponent);
+	
+				// Remove component mapping.
+				PCG::TUniqueScopeLock WriteLock(TrackingManager->ComponentToPartitionActorsMapLock);
+				TSet<TObjectPtr<APCGPartitionActor>>* PartitionActorsPtr = TrackingManager->ComponentToPartitionActorsMap.Find(OriginalComponent);
+	
+				if (PartitionActorsPtr)
 				{
-					TrackingManager->ComponentToPartitionActorsMap.Remove(OriginalComponent);
+					PartitionActorsPtr->Remove(InPartitionActor);
+	
+					if (PartitionActorsPtr->IsEmpty())
+					{
+						TrackingManager->ComponentToPartitionActorsMap.Remove(OriginalComponent);
+					}
 				}
 			}
 		}
 	}
 
 	// Cleanup the PA if it no longer has any components (return to pool or destroy).
-	if (!PartitionActor->HasLocalPCGComponents())
+	if (InPartitionActor && !InPartitionActor->HasLocalPCGComponents())
 	{
-		PartitionActor->UnregisterPCG();
+		InPartitionActor->UnregisterPCG();
 
 		if (PCGRuntimeGenSchedulerHelpers::CVarRuntimeGenerationEnablePooling.GetValueOnAnyThread())
 		{
 			if (PCGRuntimeGenSchedulerHelpers::CVarRuntimeGenerationEnableDebugging.GetValueOnAnyThread())
 			{
-				UE_LOG(LogPCG, Warning, TEXT("[RUNTIMEGEN] RETURNING PARTITION ACTOR TO POOL: '%s' (%d remaining out of %d)"), *PartitionActor->GetActorNameOrLabel(), PartitionActorPool.Num() + 1, PartitionActorPoolSize);
+				UE_LOG(LogPCG, Warning, TEXT("[RUNTIMEGEN] RETURNING PARTITION ACTOR TO POOL: '%s' (%d remaining out of %d)"), *InPartitionActor->GetActorNameOrLabel(), PartitionActorPool.Num() + 1, PartitionActorPoolSize);
 			}
 
 #if WITH_EDITOR
-			PartitionActor->SetActorLabel(*PCGRuntimeGenSchedulerConstants::PooledPartitionActorName);
+			InPartitionActor->SetActorLabel(*PCGRuntimeGenSchedulerConstants::PooledPartitionActorName);
 #endif
-			PartitionActorPool.Push(PartitionActor);
+			PartitionActorPool.Push(InPartitionActor);
 		}
 		else
 		{
 			if (PCGRuntimeGenSchedulerHelpers::CVarRuntimeGenerationEnableDebugging.GetValueOnAnyThread())
 			{
-				UE_LOG(LogPCG, Warning, TEXT("[RUNTIMEGEN] DESTROY PARTITION ACTOR: '%s'"), *PartitionActor->GetActorNameOrLabel());
+				UE_LOG(LogPCG, Warning, TEXT("[RUNTIMEGEN] DESTROY PARTITION ACTOR: '%s'"), *InPartitionActor->GetActorNameOrLabel());
 			}
 
-			Context.World->DestroyActor(PartitionActor);
+			Context.World->DestroyActor(InPartitionActor);
 		}
 	}
 }
 
-void FPCGRuntimeGenScheduler::CleanupComponent(const FGridGenerationKey& GenerationKey, IPCGGraphExecutionSource* GeneratedComponent)
+void FPCGRuntimeGenScheduler::CleanupSource(const FGridGenerationKey& GenerationKey, IPCGGraphExecutionSource* InGeneratedSource)
 {
-	TRACE_CPUPROFILER_EVENT_SCOPE(FPCGRuntimeGenScheduler::CleanupComponent);
+	TRACE_CPUPROFILER_EVENT_SCOPE(FPCGRuntimeGenScheduler::CleanupSource);
 
 	check(TrackingManager);
 
 	const FPCGGridDescriptor GridDescriptor = GenerationKey.GetGridDescriptor();
 	const EPCGHiGenGrid Grid = GridDescriptor.GetHiGenGrid();
-
 	const FIntVector GridCoords = GenerationKey.GetGridCoords();
 
 	APCGPartitionActor* PartitionActor = nullptr;
 
-	if (!GeneratedComponent)
+	if (!InGeneratedSource)
 	{
 		UE_LOG(LogPCG, Warning, TEXT("Runtime generated component could not be recovered on grid %d at (%d, %d, %d). It has been lost or destroyed."), GridDescriptor.GetGridSize(), GridCoords.X, GridCoords.Y, GridCoords.Z);
 
-		// If the GeneratedComponent has been lost for some reason, get the PA directly from the TrackingManager.
+		// If the GeneratedSource has been lost for some reason, get the PA directly from the TrackingManager.
 		PartitionActor = TrackingManager->GetPartitionActor(GridDescriptor, GridCoords);
 	}
-	else // If the GeneratedComponent does exist, we can clean it up.
+	else // If the GeneratedSource does exist, we can clean it up.
 	{
-		PartitionActor = Cast<APCGPartitionActor>(CastChecked<UPCGComponent>(GeneratedComponent)->GetOwner());
-		ensure(PartitionActor || Grid == EPCGHiGenGrid::Unbounded);
+		PartitionActor = InGeneratedSource->GetExecutionState().GetTypedTarget<APCGPartitionActor>();
 
 		if (Grid == EPCGHiGenGrid::Unbounded)
 		{
-			if (PCGRuntimeGenSchedulerHelpers::CVarRuntimeGenerationEnableDebugging.GetValueOnAnyThread() && PartitionActor)
+			if (PCGRuntimeGenSchedulerHelpers::CVarRuntimeGenerationEnableDebugging.GetValueOnAnyThread())
 			{
-				UE_LOG(LogPCG, Warning, TEXT("[RUNTIMEGEN] CLEANUP: '%s'"), *PartitionActor->GetActorNameOrLabel());
+				UE_LOG(LogPCG, Warning, TEXT("[RUNTIMEGEN] CLEANUP: '%s'"), *InGeneratedSource->GetExecutionState().GetDebugName());
+			}
+
+			// If an OC tries to clean up but it's been generated via RG execution source instead, make sure to clean that up properly.
+			if (UPCGComponent* OriginalComponent = Cast<UPCGComponent>(InGeneratedSource))
+			{
+				if (TObjectPtr<UPCGRuntimeGenExecutionSource> FoundOriginalSource = GetOriginalRuntimeGenExecutionSource(OriginalComponent))
+				{
+					InGeneratedSource = FoundOriginalSource;
+				}
 			}
 
 #if PCG_RGS_ONSCREENDEBUGMESSAGES
@@ -1756,117 +1871,102 @@ void FPCGRuntimeGenScheduler::CleanupComponent(const FGridGenerationKey& Generat
 			CleanupParams.bReleaseManagedResources = true;
 			CleanupParams.bImmediate = true;
 
-			GeneratedComponent->GetExecutionState().Cleanup(CleanupParams);
+			InGeneratedSource->GetExecutionState().Cleanup(CleanupParams);
 		}
 	}
 
-	CleanupLocalComponent(PartitionActor, GeneratedComponent);
+	if (PartitionActor || (InGeneratedSource && InGeneratedSource->GetExecutionState().IsLocalSource()))
+	{
+		CleanupLocalSource(PartitionActor, InGeneratedSource);
+	}
 	
-	GeneratedComponents.Remove(GenerationKey);
+	GeneratedSources.Remove(GenerationKey);
 }
 
-void FPCGRuntimeGenScheduler::CleanupDelayedRefreshComponents()
+void FPCGRuntimeGenScheduler::CleanupDelayedRefreshSources()
 {
-	TRACE_CPUPROFILER_EVENT_SCOPE(FPCGRuntimeGenScheduler::CleanupDelayedRefreshComponents);
+	TRACE_CPUPROFILER_EVENT_SCOPE(FPCGRuntimeGenScheduler::CleanupDelayedRefreshSources);
 
-	check(TrackingManager);
-
-	// Check that each refreshed local component is still intersecting its original component.
-	// If it is not, it would be leaked instead of refreshed, so we should force a full cleanup.
-	for (const FGridGenerationKey& GenerationKey : GeneratedComponentsToRemove)
+	// Remove any generation keys that have been registered for deferred removal.
+	if (!GeneratedSourcesToRemove.IsEmpty())
 	{
-		if (!GenerationKey.IsValid())
-		{
-			continue;
-		}
-
-		const FPCGGridDescriptor GridDescriptor = GenerationKey.GetGridDescriptor();
-		const EPCGHiGenGrid Grid = GridDescriptor.GetHiGenGrid();
-		UPCGComponent* OriginalComponent = GenerationKey.GetOriginalComponent();
-		const FIntVector& GridCoords = GenerationKey.GetGridCoords();
-
-		// The unbounded grid level will always lie inside the original component, so we can skip it.
-		if (Grid == EPCGHiGenGrid::Unbounded)
-		{
-			if (OriginalComponent && !OriginalComponent->bActivated)
-			{
-				CleanupComponent(GenerationKey, OriginalComponent);
-			}
+		GeneratedSources = GeneratedSources.Difference(GeneratedSourcesToRemove);
+		GeneratedSourcesToRemove.Reset();
+	}
+}
 
-			continue;
-		}
+void FPCGRuntimeGenScheduler::RefreshExecutionSource(IPCGGraphExecutionSource* InExecutionSource, bool bRemovePartitionActors)
+{
+	if (!InExecutionSource || !ensure(IsInGameThread()))
+	{
+		return;
+	}
 
-		IPCGGraphExecutionSource* LocalComponent = OriginalComponent ? TrackingManager->GetLocalComponent(GridDescriptor, GridCoords, OriginalComponent) : nullptr;
-		APCGPartitionActor* PartitionActor = LocalComponent ? Cast<APCGPartitionActor>(CastChecked<UPCGComponent>(LocalComponent)->GetOwner()) : nullptr;
+	IPCGGraphExecutionState& ExecutionState = InExecutionSource->GetExecutionState();
+	IPCGGraphExecutionSource* OriginalSource = ExecutionState.GetOriginalSource();
+	UPCGComponent* OriginalComponent = Cast<UPCGComponent>(OriginalSource);
 
-		if (LocalComponent && PartitionActor)
-		{
-			const FBox OriginalBounds = OriginalComponent->GetGridBounds();
-			const FBox LocalBounds = PartitionActor->GetFixedBounds();
+	// If we are refreshing a PCG Component that has already been generated as a RuntimeGen Execution Source, then we should make sure to use the correct original source (which is the RG exec source).
+	if (TObjectPtr<UPCGRuntimeGenExecutionSource> FoundOriginalSource = GetOriginalRuntimeGenExecutionSource(OriginalComponent))
+	{
+		OriginalSource = FoundOriginalSource;
+		OriginalComponent = Cast<UPCGComponent>(FoundOriginalSource->GetOwningSource());
 
-			if (!OriginalBounds.Intersect(LocalBounds) || !OriginalComponent->bActivated)
-			{
-				CleanupComponent(GenerationKey, LocalComponent);
-			}
-		}
-		else
-		{
-			// If the component or partition actor could not be recovered, just clean up.
-			CleanupComponent(GenerationKey, /*GenerationKey=*/nullptr);
-		}
+		// If the original component was trying to refresh, we should update the registered execution source's properties in case the original component's properties changed.
+		FoundOriginalSource->SetOwningSource(OriginalComponent);
+	}
+	else if (UPCGRuntimeGenExecutionSource* OriginalRuntimeGenSource = Cast<UPCGRuntimeGenExecutionSource>(OriginalSource))
+	{
+		OriginalComponent = Cast<UPCGComponent>(OriginalRuntimeGenSource->GetOwningSource());
 	}
 
-	// Remove any remaining generation keys that have been registered for deferred removal.
-	if (!GeneratedComponentsToRemove.IsEmpty())
+	if (!OriginalSource || !OriginalComponent)
 	{
-		GeneratedComponents = GeneratedComponents.Difference(GeneratedComponentsToRemove);
-		GeneratedComponentsToRemove.Reset();
+		return;
 	}
-}
 
-void FPCGRuntimeGenScheduler::RefreshComponent(IPCGGraphExecutionSource* InComponent, bool bRemovePartitionActors)
-{
-	if (!InComponent || !ensure(IsInGameThread()))
+	// If we are mid way through setting up an original source, early out from this refresh.
+	if (OriginalSource == OriginalSourceBeingRegistered)
 	{
 		return;
 	}
 
-	IPCGGraphExecutionState& ExecutionState = InComponent->GetExecutionState();
-	APCGPartitionActor* PartitionActor = Cast<APCGPartitionActor>(CastChecked<UPCGComponent>(InComponent)->GetOwner());
-	const bool bIsLocalComponent = PartitionActor != nullptr;
-	UPCGComponent* OriginalComponent = Cast<UPCGComponent>(ExecutionState.GetOriginalSource());
+	const bool bUseActorComponentlessGeneration = OriginalComponent->UseActorComponentlessGeneration();
+	const bool bHasExistingRuntimeGenExecutionSource = OriginalExecutionSources.Contains(OriginalComponent);
 
-	// If we are mid way through setting up an original component, early out from this refresh.
-	if (!OriginalComponent || OriginalComponent == OriginalComponentBeingGenerated)
+	if (bUseActorComponentlessGeneration != bHasExistingRuntimeGenExecutionSource)
 	{
-		return;
+		// We always need to do a deep refresh if we've detected that bUseActorComponentlessGeneration has changed. We will update the original execution source registration after performing the cleanup.
+		bRemovePartitionActors = true;
 	}
 
-	// Trigger a rescan of generation cells for this component.
-	ChangeDetector.FlushCachedState(OriginalComponent);
+	// Trigger a rescan of generation cells for this execution source.
+	ChangeDetector.FlushCachedState(OriginalSource);
 
 	const bool bLoggingEnabled = PCGRuntimeGenSchedulerHelpers::CVarRuntimeGenerationEnableDebugging.GetValueOnAnyThread();
 
-	// Useful because we can run into generation order issues if components are left to continue generating.
+	// Useful because we can run into generation order issues if execution sources are left to continue generating.
 	if (ExecutionState.IsGenerating())
 	{
 		ExecutionState.Cancel();
 	}
 
+	IPCGGraphExecutionState& OriginalExecutionState = OriginalSource->GetExecutionState();
+
 	if (!bRemovePartitionActors)
 	{
-		// Refresh path - mark component dirty and removed generated keys which will cause it to be scheduled for regeneration.
+		// Refresh path - mark execution source dirty and removed generated keys which will cause it to be scheduled for regeneration.
 
-		// Register for deferred removal from generated components set, component will be regenerated later (and in grid order
+		// Register for deferred removal from generated execution sources set, execution source will be regenerated later (and in grid order
 		// so that e.g. unbounded is generated first).
-		if (PartitionActor)
+		if (ExecutionState.IsLocalSource())
 		{
 			if (bLoggingEnabled)
 			{
-				UE_LOG(LogPCG, Warning, TEXT("[RUNTIMEGEN] SHALLOW REFRESH LOCAL COMPONENT: '%s'"), *PartitionActor->GetActorNameOrLabel());
+				UE_LOG(LogPCG, Warning, TEXT("[RUNTIMEGEN] SHALLOW REFRESH LOCAL COMPONENT: '%s'"), *ExecutionState.GetDebugName());
 			}
 
-			GeneratedComponentsToRemove.Emplace({ PartitionActor->GetPCGGridSize(), PartitionActor->GetGridCoord(), OriginalComponent, bIsLocalComponent ? InComponent : nullptr });
+			GeneratedSourcesToRemove.Emplace({ ExecutionState.GetGenerationGridSize(), ExecutionState.GetGenerationGridCoords(), OriginalSource, InExecutionSource });
 
 			IPCGGraphExecutionState::FCleanupParams CleanupParams;
 			CleanupParams.bReleaseManagedResources = false;
@@ -1876,29 +1976,21 @@ void FPCGRuntimeGenScheduler::RefreshComponent(IPCGGraphExecutionSource* InCompo
 		}
 		else
 		{
-			// Register original component for deferred removal.
-			GeneratedComponentsToRemove.Emplace({ PCGHiGenGrid::UnboundedGridSize(), FIntVector(0), OriginalComponent });
+			// Register original execution source for deferred removal.
+			GeneratedSourcesToRemove.Emplace({ PCGHiGenGrid::UnboundedGridSize(), FIntVector(0), OriginalSource });
 
-			// Register local components for deferred removal if they have not already registered themselves.
-			for (const FGridGenerationKey& Key : GeneratedComponents)
+			// Register local execution sources for deferred removal if they have not already registered themselves.
+			for (const FGridGenerationKey& Key : GeneratedSources)
 			{
-				if (Key.GetOriginalComponent() == InComponent && !GeneratedComponentsToRemove.Contains(Key))
+				if (Key.GetOriginalSource() == OriginalSource && !GeneratedSourcesToRemove.Contains(Key))
 				{
-					const FPCGGridDescriptor GridDescriptor = Key.GetGridDescriptor();
-					
-					// TODO - clean up local immediate will have a flag in the future to clean up the local components on its own, so this call to CleanupLocalImmediate will not be required
-					IPCGGraphExecutionSource* LocalComponent = TrackingManager->GetLocalComponent(
-						GridDescriptor,
-						Key.GetGridCoords(),
-						OriginalComponent);
-
-					if (LocalComponent)
+					if (IPCGGraphExecutionSource* LocalSource = OriginalExecutionState.GetLocalSource(Key.GetGridDescriptor(), Key.GetGridCoords()))
 					{
-						IPCGGraphExecutionState& LocalExecutionState = LocalComponent->GetExecutionState();
+						IPCGGraphExecutionState& LocalExecutionState = LocalSource->GetExecutionState();
 
-						if (bLoggingEnabled && CastChecked<UPCGComponent>(LocalComponent)->GetOwner())
+						if (bLoggingEnabled)
 						{
-							UE_LOG(LogPCG, Warning, TEXT("[RUNTIMEGEN] SHALLOW REFRESH LOCAL COMPONENT: '%s'"), *CastChecked<UPCGComponent>(LocalComponent)->GetOwner()->GetActorNameOrLabel());
+							UE_LOG(LogPCG, Warning, TEXT("[RUNTIMEGEN] SHALLOW REFRESH LOCAL COMPONENT: '%s'"), *LocalExecutionState.GetDebugName());
 						}
 
 						IPCGGraphExecutionState::FCleanupParams CleanupParams;
@@ -1907,108 +1999,124 @@ void FPCGRuntimeGenScheduler::RefreshComponent(IPCGGraphExecutionSource* InCompo
 
 						LocalExecutionState.Cleanup(CleanupParams);
 
-						// @todo_pcg: Instead of casting to PCGComponent, we should be able to cast to a RuntimeGenExecutionSource which provides this behavior.
 						// We need to make sure that the next time this is generated that it matches the original
-						CastChecked<UPCGComponent>(LocalComponent)->SetPropertiesFromOriginal(OriginalComponent);
+						if (UPCGComponent* PCGComponent = Cast<UPCGComponent>(LocalSource))
+						{
+							PCGComponent->SetPropertiesFromOriginal(OriginalComponent);
+						}
 					}
 
-					GeneratedComponentsToRemove.Add(Key);
+					GeneratedSourcesToRemove.Add(Key);
 				}
 			}
 
-			if (bLoggingEnabled && OriginalComponent->GetOwner())
+			if (bLoggingEnabled)
 			{
 				UE_LOG(LogPCG, Warning, TEXT("[RUNTIMEGEN] SHALLOW REFRESH COMPONENT: '%s' PARTITIONED: %d"),
-					*OriginalComponent->GetOwner()->GetActorNameOrLabel(),
-					OriginalComponent->IsPartitioned() ? 1 : 0);
+					*OriginalExecutionState.GetDebugName(),
+					OriginalExecutionState.IsPartitioned() ? 1 : 0);
 			}
 
 			IPCGGraphExecutionState::FCleanupParams CleanupParams;
 			CleanupParams.bReleaseManagedResources = false;
 			CleanupParams.bImmediate = true;
 
-			ExecutionState.Cleanup(CleanupParams);
+			OriginalExecutionState.Cleanup(CleanupParams);
 		}
 	}
 	else
 	{
-		// Full cleanout path - cleanup existing components and return actors to the pool.
+		// Full cleanout path - cleanup existing execution sources and return actors to the pool.
 
-		auto RefreshLocalComponent = [this, OriginalComponent, bLoggingEnabled](IPCGGraphExecutionSource* LocalComponent)
+		auto RefreshLocalSource = [this, OriginalSource, bLoggingEnabled](IPCGGraphExecutionSource* LocalSource)
 		{
-			check(LocalComponent);
-			APCGPartitionActor* PartitionActor = CastChecked<APCGPartitionActor>(Cast<UPCGComponent>(LocalComponent)->GetOwner());
+			check(LocalSource);
 
-			// Find the specific generation key for this component, if it exists, cleanup and generate.
-			FGridGenerationKey LocalComponentKey(PartitionActor->GetPCGGridSize(), PartitionActor->GetGridCoord(), OriginalComponent, LocalComponent);
+			IPCGGraphExecutionState& LocalExecutionState = LocalSource->GetExecutionState();
 
-			if (GeneratedComponents.Find(LocalComponentKey))
+			// Find the specific generation key for this execution source, if it exists, cleanup and generate.
+			FGridGenerationKey LocalSourceKey(LocalExecutionState.GetGenerationGridSize(), LocalExecutionState.GetGenerationGridCoords(), OriginalSource, LocalSource);
+
+			if (GeneratedSources.Find(LocalSourceKey))
 			{
 				if (bLoggingEnabled)
 				{
-					UE_LOG(LogPCG, Warning, TEXT("[RUNTIMEGEN] DEEP REFRESH LOCAL COMPONENT: '%s'"), *PartitionActor->GetActorNameOrLabel());
+					UE_LOG(LogPCG, Warning, TEXT("[RUNTIMEGEN] DEEP REFRESH LOCAL COMPONENT: '%s'"), *LocalExecutionState.GetDebugName());
 				}
 
-				CleanupComponent(LocalComponentKey, LocalComponent);
+				CleanupSource(LocalSourceKey, LocalSource);
 			}
 		};
 
-		if (bIsLocalComponent)
+		if (ExecutionState.IsLocalSource())
 		{
-			RefreshLocalComponent(InComponent);
+			RefreshLocalSource(InExecutionSource);
 		}
 		else
 		{
 			TArray<FGridGenerationKey> GenerationKeys;
 
-			for (FGridGenerationKey GenerationKey : GeneratedComponents)
+			for (FGridGenerationKey GenerationKey : GeneratedSources)
 			{
-				if (GenerationKey.GetOriginalComponent() == OriginalComponent)
+				if (GenerationKey.GetOriginalSource() == OriginalSource)
 				{
 					GenerationKeys.Add(GenerationKey);
 				}
 			}
 
-			// Gather all generated components which originated from this original component.
+			// Gather all generated execution sources which originated from this original execution source.
 			for (FGridGenerationKey GenerationKey : GenerationKeys)
 			{
 				const FPCGGridDescriptor GridDescriptor = GenerationKey.GetGridDescriptor();
 				const EPCGHiGenGrid Grid = GridDescriptor.GetHiGenGrid();
 
-				// If the Grid is unbounded, we have a non-partitioned or unbounded component.
+				// If the Grid is unbounded, we have a non-partitioned or unbounded execution source.
 				if (Grid == EPCGHiGenGrid::Unbounded)
 				{
-					if (bLoggingEnabled && OriginalComponent->GetOwner())
+					if (bLoggingEnabled)
 					{
 						UE_LOG(LogPCG, Warning, TEXT("[RUNTIMEGEN] DEEP REFRESH COMPONENT: '%s' PARTITIONED: %d"),
-							*OriginalComponent->GetOwner()->GetActorNameOrLabel(),
-							OriginalComponent->IsPartitioned() ? 1 : 0);
+							*OriginalExecutionState.GetDebugName(),
+							OriginalExecutionState.IsPartitioned() ? 1 : 0);
 					}
 
-					CleanupComponent(GenerationKey, OriginalComponent);
+					CleanupSource(GenerationKey, OriginalSource);
 				}
-				// Otherwise we have a local component.
+				// Otherwise we have a local execution source.
 				else
 				{
 					const FIntVector GridCoords = GenerationKey.GetGridCoords();
 
-					if (IPCGGraphExecutionSource* LocalComponent = TrackingManager->GetLocalComponent(GridDescriptor, GridCoords, OriginalComponent))
+					if (IPCGGraphExecutionSource* LocalSource = OriginalExecutionState.GetLocalSource(GridDescriptor, GridCoords))
 					{
-						RefreshLocalComponent(LocalComponent);
+						RefreshLocalSource(LocalSource);
 					}
 					else
 					{
-						// If the local component could not be recovered, cleanup its entry to avoid leaking resources/locking the grid cell.
-						CleanupComponent(GenerationKey, nullptr);
+						// If the local execution source could not be recovered, cleanup its entry to avoid leaking resources/locking the grid cell.
+						CleanupSource(GenerationKey, nullptr);
 					}
 				}
 			}
 		}
 	}
 
-	if (!bIsLocalComponent)
+	// Update the original execution source registration after the cleanup has been performed. That way we cleanup the existing source before moving to the new source.
+	if (!ExecutionState.IsLocalSource() && bUseActorComponentlessGeneration != bHasExistingRuntimeGenExecutionSource)
+	{
+		if (bUseActorComponentlessGeneration)
+		{
+			RegisterOriginalRuntimeGenExecutionSource(OriginalComponent);
+		}
+		else
+		{
+			UnregisterOriginalRuntimeGenExecutionSource(OriginalComponent);
+		}
+	}
+
+	if (!ExecutionState.IsLocalSource())
 	{
-		// When an original/non-partitioned component is refreshed, we need to dirty the state.
+		// When an original/non-partitioned execution source is refreshed, we need to dirty the state.
 		bContextDirty = true;
 		bCleanupScanRequired = true;
 	}
@@ -2124,6 +2232,12 @@ void FPCGRuntimeGenScheduler::AddReferencedObjects(FReferenceCollector& Collecto
 
 	// The level should be keeping the pooled PAs visible to GC. This is just a tentative fix for a crash in GetPartitionActorFromPool(), to understand if the crash is happening because of unreferenced GCed actors.
 	Collector.AddReferencedObjects(PartitionActorPool);
+
+	// Keep runtime gen execution sources alive.
+	for (auto& [OriginalComponent, ExecutionSource] : OriginalExecutionSources)
+	{
+		Collector.AddReferencedObject(ExecutionSource);
+	}
 }
 
 void FPCGRuntimeGenScheduler::CacheVirtualTexturePrimingInfos(UPCGComponent* InOriginalComponent)
@@ -2181,12 +2295,9 @@ void FPCGRuntimeGenScheduler::RefreshContext()
 				ContextInner.bAnySourcesUse2DGrids |= bUses2DGrid;
 				ContextInner.bAnySourcesUse3DGrids |= !bUses2DGrid;
 
-				if (const UPCGComponent* Component = Cast<UPCGComponent>(ExecutionSource))
+				if (const UPCGSchedulingPolicyBase* Policy = ExecutionSource->GetExecutionState().GetRuntimeGenSchedulingPolicy())
 				{
-					if (UPCGSchedulingPolicyBase* Policy = Component->GetRuntimeGenSchedulingPolicy())
-					{
-						ContextInner.bAnySourcesUseFrustumCulling |= Policy->CullsBasedOnDirection();
-					}
+					ContextInner.bAnySourcesUseFrustumCulling |= Policy->CullsBasedOnDirection();
 				}
 			}
 		}
@@ -2218,8 +2329,57 @@ bool FPCGRuntimeGenScheduler::IsChangeDetectionEnabled() const
 	return PCGRuntimeGenSchedulerHelpers::CVarRuntimeGenerationEnableChangeDetection.GetValueOnGameThread();
 }
 
+void FPCGRuntimeGenScheduler::RegisterOriginalRuntimeGenExecutionSource(UPCGComponent* InOriginalComponent)
+{
+	if (InOriginalComponent && InOriginalComponent->UseActorComponentlessGeneration())
+	{
+		if (TObjectPtr<UPCGRuntimeGenExecutionSource> FoundOriginalSource = GetOriginalRuntimeGenExecutionSource(InOriginalComponent))
+		{
+			// If we're already registered, just update the runtime gen execution source to stay up to date.
+			FoundOriginalSource->SetOwningSource(InOriginalComponent);
+		}
+		else
+		{
+			// No existing source found, register a new runtime gen execution source.
+			UPCGRuntimeGenExecutionSource* NewSource = NewObject<UPCGRuntimeGenExecutionSource>();
+			NewSource->Register(InOriginalComponent);
+
+			OriginalExecutionSources.Add(InOriginalComponent, NewSource);
+		}
+	}
+}
+
+void FPCGRuntimeGenScheduler::UnregisterOriginalRuntimeGenExecutionSource(UPCGComponent* InOriginalComponent)
+{
+	if (TObjectPtr<UPCGRuntimeGenExecutionSource>* FoundRuntimeGenSource = OriginalExecutionSources.Find(InOriginalComponent))
+	{
+		if (*FoundRuntimeGenSource)
+		{
+			IPCGGraphExecutionState::FCleanupParams CleanupParams;
+			CleanupParams.bCleanupLocalSources = true;
+			CleanupParams.bReleaseManagedResources = true;
+			CleanupParams.bImmediate = true;
+
+			(*FoundRuntimeGenSource)->Cleanup(CleanupParams);
+			(*FoundRuntimeGenSource)->Unregister();
+		}
+
+		OriginalExecutionSources.Remove(InOriginalComponent);
+	}
+}
+
+UPCGRuntimeGenExecutionSource* FPCGRuntimeGenScheduler::GetOriginalRuntimeGenExecutionSource(UPCGComponent* InOriginalComponent)
+{
+	if (TObjectPtr<UPCGRuntimeGenExecutionSource>* FoundOriginalSource = OriginalExecutionSources.Find(InOriginalComponent); FoundOriginalSource)
+	{
+		return *FoundOriginalSource;
+	}
+
+	return nullptr;
+}
+
 #if WITH_EDITOR
-void FPCGRuntimeGenScheduler::RegisterGraphParameterChangedEvent(UPCGComponent* InOriginalComponent)
+void FPCGRuntimeGenScheduler::RegisterGraphParameterChangedEvent(const UPCGComponent* InOriginalComponent)
 {
 	if (UPCGGraph* Graph = InOriginalComponent ? InOriginalComponent->GetGraph() : nullptr)
 	{
@@ -2232,7 +2392,7 @@ void FPCGRuntimeGenScheduler::RegisterGraphParameterChangedEvent(UPCGComponent*
 	}
 }
 
-void FPCGRuntimeGenScheduler::UnregisterGraphParameterChangedEvent(UPCGComponent* InOriginalComponent)
+void FPCGRuntimeGenScheduler::UnregisterGraphParameterChangedEvent(const UPCGComponent* InOriginalComponent)
 {
 	FDelegateHandle DelegateHandle;
 
diff --git a/Engine/Plugins/PCG/Source/PCG/Private/Subsystems/PCGSubsystem.cpp b/Engine/Plugins/PCG/Source/PCG/Private/Subsystems/PCGSubsystem.cpp
--- a/Engine/Plugins/PCG/Source/PCG/Private/Subsystems/PCGSubsystem.cpp
+++ b/Engine/Plugins/PCG/Source/PCG/Private/Subsystems/PCGSubsystem.cpp
@@ -79,7 +79,7 @@ namespace PCGSubsystemConsole
 		{
 			if (UPCGSubsystem* PCGSubsystem = UPCGSubsystem::GetSubsystemForCurrentWorld())
 			{
-				PCGSubsystem->RefreshAllRuntimeGenComponents(EPCGChangeType::GenerationGrid);
+				PCGSubsystem->RefreshAllRuntimeGenExecutionSources(EPCGChangeType::GenerationGrid);
 			}
 		}));
 
@@ -472,20 +472,15 @@ void UPCGSubsystem::OnPCGQualityLevelChanged()
 				Graph->OnPCGQualityLevelChanged();
 			}
 #else
-			// @todo_pcg: support execution source
-			if (UPCGComponent* PCGComponent = Cast<UPCGComponent>(InExecutionSource))
-			{
-				RefreshRuntimeGenComponent(PCGComponent, EPCGChangeType::GenerationGrid);
-			}
+			RefreshRuntimeGenExecutionSource(InExecutionSource, EPCGChangeType::GenerationGrid);
 #endif
-
 		}
 	});
 }
 
 void UPCGSubsystem::OnScalabilitySettingsChanged(const Scalability::FQualityLevels& QualityLevels)
 {
-	RefreshAllRuntimeGenComponents(EPCGChangeType::GenerationGrid);
+	RefreshAllRuntimeGenExecutionSources(EPCGChangeType::GenerationGrid);
 }
 
 #if WITH_EDITOR
@@ -679,6 +674,21 @@ void UPCGSubsystem::OnOriginalExecutionSourceUnregistered(IPCGGraphExecutionSour
 #endif
 }
 
+void UPCGSubsystem::OnOriginalExecutionSourceReplaced(const IPCGGraphExecutionSource* InOldExecutionSource, IPCGGraphExecutionSource* InNewExecutionSource)
+{
+	// @todo_pcg: support execution source
+	if (RuntimeGenScheduler)
+	{
+		const UPCGComponent* OldPCGComponent = Cast<UPCGComponent>(InOldExecutionSource);
+		UPCGComponent* NewPCGComponent = Cast<UPCGComponent>(InNewExecutionSource);
+
+		if (OldPCGComponent && NewPCGComponent)
+		{
+			RuntimeGenScheduler->OnOriginalComponentReplaced(OldPCGComponent, NewPCGComponent);
+		}
+	}
+}
+
 UPCGLandscapeCache* UPCGSubsystem::GetLandscapeCache()
 {
 	APCGWorldActor* LandscapeCacheOwner = GetPCGWorldActor();
@@ -1021,27 +1031,35 @@ void UPCGSubsystem::CancelGenerationInternal(IPCGGraphExecutionSource* Source, b
 	}
 }
 
+// DEPRECATED 5.8
 void UPCGSubsystem::RefreshRuntimeGenComponent(UPCGComponent* RuntimeComponent, EPCGChangeType ChangeType)
 {
-	if (RuntimeComponent && RuntimeComponent->IsManagedByRuntimeGenSystem() && RuntimeGenScheduler)
+	RefreshRuntimeGenExecutionSource(RuntimeComponent, ChangeType);
+}
+
+void UPCGSubsystem::RefreshRuntimeGenExecutionSource(IPCGGraphExecutionSource* ExecutionSource, EPCGChangeType ChangeType)
+{
+	if (ExecutionSource && ExecutionSource->GetExecutionState().IsManagedByRuntimeGenSystem() && RuntimeGenScheduler)
 	{
 		// Only need to remove PAs if the grid sizes have changed.
 		const bool bRemovePartitionActors = !!(ChangeType & EPCGChangeType::GenerationGrid);
-		RuntimeGenScheduler->RefreshComponent(RuntimeComponent, bRemovePartitionActors);
+		RuntimeGenScheduler->RefreshExecutionSource(ExecutionSource, bRemovePartitionActors);
 	}
 }
 
+// DEPRECATED 5.8
 void UPCGSubsystem::RefreshAllRuntimeGenComponents(EPCGChangeType ChangeType)
 {
-	// @todo_pcg: support execution source
+	RefreshAllRuntimeGenExecutionSources(ChangeType);
+}
+
+void UPCGSubsystem::RefreshAllRuntimeGenExecutionSources(EPCGChangeType ChangeType)
+{
 	for (IPCGGraphExecutionSource* ExecutionSource : GetAllRegisteredExecutionSources())
 	{
-		if (UPCGComponent* Component = Cast<UPCGComponent>(ExecutionSource))
+		if (ExecutionSource->GetExecutionState().IsManagedByRuntimeGenSystem())
 		{
-			if (Component->GenerationTrigger == EPCGComponentGenerationTrigger::GenerateAtRuntime)
-			{
-				RefreshRuntimeGenComponent(Component, ChangeType);
-			}
+			RefreshRuntimeGenExecutionSource(ExecutionSource, ChangeType);
 		}
 	}
 }
@@ -1909,14 +1927,14 @@ void UPCGSubsystem::OnPreviewPlatformChanged()
 {
 	FlushCache();
 
-	RefreshAllRuntimeGenComponents(EPCGChangeType::GenerationGrid);
+	RefreshAllRuntimeGenExecutionSources(EPCGChangeType::GenerationGrid);
 }
 
 void UPCGSubsystem::OnPreviewFeatureLevelChanged(ERHIFeatureLevel::Type)
 {
 	FlushCache();
 
-	RefreshAllRuntimeGenComponents(EPCGChangeType::GenerationGrid);
+	RefreshAllRuntimeGenExecutionSources(EPCGChangeType::GenerationGrid);
 }
 
 void UPCGSubsystem::CreatePartitionActorsWithinBounds(UPCGComponent* InComponent, const FBox& InBounds, const PCGHiGenGrid::FSizeArray& InGridSizes)
@@ -2187,7 +2205,7 @@ void UPCGSubsystem::ExecuteBeginTickActions()
 
 void UPCGSubsystem::OnGlobalComponentReregisterContextDestroyed()
 {
-	RefreshAllRuntimeGenComponents();
+	RefreshAllRuntimeGenExecutionSources();
 }
 
 void UPCGSubsystem::RegisterPartitionActor(APCGPartitionActor* InActor)
diff --git a/Engine/Plugins/PCG/Source/PCG/Public/Data/Registry/PCGGetDataFunctionRegistry.h b/Engine/Plugins/PCG/Source/PCG/Public/Data/Registry/PCGGetDataFunctionRegistry.h
--- a/Engine/Plugins/PCG/Source/PCG/Public/Data/Registry/PCGGetDataFunctionRegistry.h
+++ b/Engine/Plugins/PCG/Source/PCG/Public/Data/Registry/PCGGetDataFunctionRegistry.h
@@ -9,19 +9,21 @@
 
 class AActor;
 class UActorComponent;
-class UPCGComponent;
 struct FPCGContext;
 struct FPCGComponentSelectorSettings;
 
 /** Struct to hold selection criteria */
 struct FPCGGetDataFunctionRegistryParams
 {
-	const UPCGComponent* SourceComponent = nullptr;
+	const IPCGGraphExecutionSource* Source = nullptr;
 	const FPCGComponentSelectorSettings* ComponentSelector = nullptr;
 	EPCGDataType DataTypeFilter = EPCGDataType::Any;
 	bool bParseActor = true;
 	bool bIgnorePCGGeneratedComponents = true;
 	bool bAddActorTags = true;
+
+	UE_DEPRECATED(5.8, "Use Source instead")
+	const UPCGComponent* SourceComponent = nullptr;
 };
 
 struct FPCGGetDataFunctionRegistryOutput
diff --git a/Engine/Plugins/PCG/Source/PCG/Public/Helpers/PCGHelpers.h b/Engine/Plugins/PCG/Source/PCG/Public/Helpers/PCGHelpers.h
--- a/Engine/Plugins/PCG/Source/PCG/Public/Helpers/PCGHelpers.h
+++ b/Engine/Plugins/PCG/Source/PCG/Public/Helpers/PCGHelpers.h
@@ -49,8 +49,8 @@ namespace PCGHelpers
 
 	PCG_API FBox OverlapBounds(const FBox& InBoxA, const FBox& InBoxB);
 
-	/** Returns the bounds of InActor, intersected with the component if InActor is a partition actor */
-	PCG_API FBox GetGridBounds(const AActor* InActor, const UPCGComponent* InComponent);
+	/** Returns the bounds of InActor, intersected with the execution source if InActor is a partition actor */
+	PCG_API FBox GetGridBounds(const AActor* InActor, const IPCGGraphExecutionSource* InSource);
 
 	PCG_API FBox GetActorBounds(const AActor* InActor, bool bIgnorePCGCreatedComponents = true);
 	PCG_API FBox GetActorLocalBounds(const AActor* InActor, bool bIgnorePCGCreatedComponents = true);
diff --git a/Engine/Plugins/PCG/Source/PCG/Public/PCGComponent.h b/Engine/Plugins/PCG/Source/PCG/Public/PCGComponent.h
--- a/Engine/Plugins/PCG/Source/PCG/Public/PCGComponent.h
+++ b/Engine/Plugins/PCG/Source/PCG/Public/PCGComponent.h
@@ -297,6 +297,9 @@ public:
 	/** Compute the runtime cleanup radius for the given grid size. */
 	UE_API double GetCleanupRadiusFromGrid(EPCGHiGenGrid Grid) const;
 
+	/** Allows executing runtime PCG graphs without partition actors and local PCG components. */
+	UE_API bool UseActorComponentlessGeneration() const;
+
 	/** Called during execution if one or more procedural ISM components are in use. */
 	void NotifyProceduralInstancesInUse() { bProceduralInstancesInUse = true; }
 
@@ -330,6 +333,7 @@ public:
 	UPROPERTY(BlueprintReadWrite, EditAnywhere, Category = RuntimeGeneration, meta = (EditCondition = "GenerationTrigger == EPCGComponentGenerationTrigger::GenerateAtRuntime && bOverrideGenerationRadii", EditConditionHides))
 	FPCGRuntimeGenerationRadii GenerationRadii;
 
+public:
 	/** A Scheduling Policy dictates the order in which instances of this component will be scheduled. */
 	UPROPERTY(BlueprintReadOnly, EditAnywhere, NoClear, Category = RuntimeGeneration, meta = (EditCondition = "GenerationTrigger == EPCGComponentGenerationTrigger::GenerateAtRuntime", EditConditionHides))
 	TSubclassOf<UPCGSchedulingPolicyBase> SchedulingPolicyClass;
@@ -814,6 +818,14 @@ private:
 
 	mutable FTransactionallySafeCriticalSection GeneratedResourcesLock;
 
+	/** Manual override for the enabling actor componentless generation. */
+	UPROPERTY(EditAnywhere, Category = "RuntimeGeneration|Advanced", meta = (EditCondition = "GenerationTrigger == EPCGComponentGenerationTrigger::GenerateAtRuntime", EditConditionHides))
+	bool bOverrideUseActorComponentlessGeneration = false;
+
+	/** [EXPERIMENTAL] Avoids creating partition actors and local PCG components. Any functionality which depends on a target actor/component are not supported in this mode. */
+	UPROPERTY(EditAnywhere, Category = "RuntimeGeneration|Advanced", meta = (EditCondition = "GenerationTrigger == EPCGComponentGenerationTrigger::GenerateAtRuntime && bOverrideUseActorComponentlessGeneration", EditConditionHides))
+	bool bUseActorComponentlessGeneration = false;
+
 	// Graph instance
 private:
 	/** Will set the given graph interface into our owned graph instance. Must not be used on local components.*/
diff --git a/Engine/Plugins/PCG/Source/PCG/Public/PCGComponentExecutionState.h b/Engine/Plugins/PCG/Source/PCG/Public/PCGComponentExecutionState.h
--- a/Engine/Plugins/PCG/Source/PCG/Public/PCGComponentExecutionState.h
+++ b/Engine/Plugins/PCG/Source/PCG/Public/PCGComponentExecutionState.h
@@ -76,6 +76,7 @@ public:
 	virtual void AddLocalSourceCurrentTrackingKeys(const FPCGSelectionKeyToSettingsMap& LocalSourceCurrentExecutionTrackingKeys, const TSet<const UPCGSettings*>& LocalSourceCurrentExecutionTrackingSettings) override;
 
 	virtual bool IsRefreshInProgress() const override;
+	virtual bool WasGeneratedThisSession() const override;
 	virtual FPCGDynamicTrackingPriority GetDynamicTrackingPriority() const override;
 #endif // WITH_EDITOR
 
diff --git a/Engine/Plugins/PCG/Source/PCG/Public/PCGGraph.h b/Engine/Plugins/PCG/Source/PCG/Public/PCGGraph.h
--- a/Engine/Plugins/PCG/Source/PCG/Public/PCGGraph.h
+++ b/Engine/Plugins/PCG/Source/PCG/Public/PCGGraph.h
@@ -444,6 +444,9 @@ public:
 	/** Gets cleanup radius from grid, considering grid exponential. */
 	UE_API double GetGridCleanupRadiusFromGrid(EPCGHiGenGrid Grid) const;
 
+	/** Allows executing runtime PCG graphs without partition actors and local PCG components. */
+	bool UseActorComponentlessGeneration() const;
+
 #if WITH_EDITOR
 	/** Can be overriden by child class to disable debug globally on all settings. */
 	virtual bool ShouldDisplayDebuggingProperties() const { return true; };
@@ -654,6 +657,11 @@ public:
 	UPROPERTY(BlueprintReadWrite, EditAnywhere, Category = "Runtime Generation", meta=(EditCondition = "!IsStandaloneGraph()", EditConditionHides))
 	FPCGRuntimeGenerationRadii GenerationRadii;
 
+	// @todo_pcg: This could be on by default/hidden and graphs which need actors/actor-components could create a target actor on demand when necessary.
+	/** [EXPERIMENTAL] Avoids creating partition actors and local PCG components. Any functionality which depends on a target actor/component are not supported in this mode. */
+	UPROPERTY(BlueprintReadWrite, EditAnywhere, Category = "Runtime Generation|Advanced", meta=(EditCondition = "!IsStandaloneGraph()", EditConditionHides))
+	bool bUseActorComponentlessGeneration = false;
+
 	UE_API virtual void OnGraphParametersChanged(EPCGGraphParameterEvent InChangeType, FName InChangedPropertyName) override;
 
 	// Will call the callback function with a mutable property bag and will trigger the updates when it's done.
diff --git a/Engine/Plugins/PCG/Source/PCG/Public/PCGGraphExecutionStateInterface.h b/Engine/Plugins/PCG/Source/PCG/Public/PCGGraphExecutionStateInterface.h
--- a/Engine/Plugins/PCG/Source/PCG/Public/PCGGraphExecutionStateInterface.h
+++ b/Engine/Plugins/PCG/Source/PCG/Public/PCGGraphExecutionStateInterface.h
@@ -203,6 +203,9 @@ public:
 
 	/** Indicates if the current source is waiting for a refresh */
 	virtual bool IsRefreshInProgress() const = 0;
+
+	/** Whether this execution source has been generated at any point during this editor session. */
+	virtual bool WasGeneratedThisSession() const { return false; }
 #endif
 
 	/** If this is a local execution source returns the original execution source, otherwise returns self. */
diff --git a/Engine/Plugins/PCG/Source/PCG/Public/RuntimeGen/PCGRuntimeGenScheduler.h b/Engine/Plugins/PCG/Source/PCG/Public/RuntimeGen/PCGRuntimeGenScheduler.h
--- a/Engine/Plugins/PCG/Source/PCG/Public/RuntimeGen/PCGRuntimeGenScheduler.h
+++ b/Engine/Plugins/PCG/Source/PCG/Public/RuntimeGen/PCGRuntimeGenScheduler.h
@@ -3,6 +3,7 @@
 #pragma once
 
 #include "PCGCommon.h"
+#include "PCGGraphExecutionStateInterface.h"
 #include "Grid/PCGGridDescriptor.h"
 #include "RuntimeGen/PCGRuntimeGenChangeDetection.h"
 #include "RuntimeGen/PCGRuntimeGenContext.h"
@@ -19,6 +20,8 @@ class IPCGGenSourceBase;
 class ULevel;
 class ULevelStreaming;
 class UPCGComponent;
+class UPCGGraphInterface;
+class UPCGRuntimeGenExecutionSource;
 class UPCGSubsystem;
 class UWorld;
 class UWorldPartitionSubsystem;
@@ -53,7 +56,7 @@ struct FPCGVirtualTexturePrimingInfo
 
 /**
  * The Runtime Generation Scheduler system handles the scheduling of PCG Components marked as GenerateAtRuntime.
- * It searches the level for Partitioned and Non-Partitioned components in range of the currently active
+ * It searches the level for Partitioned and Non-Partitioned execution sources in range of the currently active
  * UPCGGenSources in the level, and schedules them efficiently based on their UPCGSchedulingPolicy, creating 
  * APCGPartitionActors as necessary to support hierarchical generation.
  *
@@ -80,6 +83,7 @@ public:
 
 	void OnOriginalComponentRegistered(UPCGComponent* InOriginalComponent);
 	void OnOriginalComponentUnregistered(UPCGComponent* InOriginalComponent);
+	void OnOriginalComponentReplaced(const UPCGComponent* InOldComponent, UPCGComponent* InNewComponent);
 
 	/** Destroy all runtime gen partition actors (both generated and pooled). Executed in next tick. */
 	void FlushAllGeneratedActors() { bActorFlushRequested = true; }
@@ -89,42 +93,42 @@ public:
 #endif
 
 protected:
-	// Grid size, grid coords, original component, local source
+	// Grid size, grid coords, original source, local source
 	struct FGridGenerationKey
 	{
-		FGridGenerationKey(uint32 InGridSize, const FIntVector& InGridCoords, UPCGComponent* InOriginalComponent)
-			: FGridGenerationKey(InGridSize, InGridCoords, InOriginalComponent, nullptr)
+		FGridGenerationKey(uint32 InGridSize, const FIntVector& InGridCoords, IPCGGraphExecutionSource* InOriginalSource)
+			: FGridGenerationKey(InGridSize, InGridCoords, InOriginalSource, /*InLocalSource=*/nullptr)
 		{
 		}
 
-		FGridGenerationKey(uint32 InGridSize, const FIntVector& InGridCoords, UPCGComponent* InOriginalComponent, IPCGGraphExecutionSource* InLocalComponent);
+		FGridGenerationKey(uint32 InGridSize, const FIntVector& InGridCoords, IPCGGraphExecutionSource* InOriginalSource, IPCGGraphExecutionSource* InLocalSource);
 
 		bool operator==(const FGridGenerationKey& Other) const
 		{
 			return GridSize == Other.GridSize
 				&& GridCoords == Other.GridCoords
 				&& bUse2DGrid == Other.bUse2DGrid
-				&& OriginalComponent == Other.OriginalComponent;
+				&& OriginalSource == Other.OriginalSource;
 		}
 
-		bool IsValid() const { return !!OriginalComponent.ResolveObjectPtr(); }
+		bool IsValid() const { return !!OriginalSource.ResolveObjectPtr(); }
 
 		bool Use2DGrid() const { return bUse2DGrid; }
 		uint32 GetGridSize() const { return GridSize; }
 		FIntVector GetGridCoords() const { return GridCoords; }
 
-		UPCGComponent* GetOriginalComponent() const { return OriginalComponent.ResolveObjectPtr(); }
-		UPCGComponent* GetOriginalComponentEvenIfGarbage() const { return OriginalComponent.ResolveObjectPtrEvenIfGarbage(); }
-		void SetOriginalComponent(UPCGComponent* InComponent) { OriginalComponent = InComponent; }
+		IPCGGraphExecutionSource* GetOriginalSource() const { return Cast<IPCGGraphExecutionSource>(OriginalSource.ResolveObjectPtr()); }
+		IPCGGraphExecutionSource* GetOriginalSourceEvenIfGarbage() const { return Cast<IPCGGraphExecutionSource>(OriginalSource.ResolveObjectPtrEvenIfGarbage()); }
+		void SetOriginalSource(UObject* InOriginalSource) { OriginalSource = InOriginalSource; }
 
-		IPCGGraphExecutionSource* GetCachedLocalComponent() const { return CachedLocalComponent.Get(); }
-		void SetCachedLocalComponent(IPCGGraphExecutionSource* InLocalComponent) const
+		IPCGGraphExecutionSource* GetCachedLocalSource() const { return CachedLocalSource.Get(); }
+		void SetCachedLocalSource(IPCGGraphExecutionSource* InLocalSource) const
 		{
-			CachedLocalComponent = InLocalComponent;
+			CachedLocalSource = InLocalSource;
 
-			if (InLocalComponent)
+			if (InLocalSource)
 			{
-				CachedBounds = InLocalComponent->GetExecutionState().GetBounds();
+				CachedBounds = InLocalSource->GetExecutionState().GetBounds();
 			}
 		}
 
@@ -136,40 +140,53 @@ protected:
 		bool bUse2DGrid = false;
 		uint32 GridSize = 0;
 		FIntVector GridCoords = FIntVector(0);
-		TObjectKey<UPCGComponent> OriginalComponent;
+		TObjectKey<UObject> OriginalSource;
 
 		// Optional/opportunistic cached local execution source if one has been created.
-		mutable TWeakInterfacePtr<IPCGGraphExecutionSource> CachedLocalComponent;
+		mutable TWeakInterfacePtr<IPCGGraphExecutionSource> CachedLocalSource;
 
 		mutable TOptional<FBox> CachedBounds;
 
 		friend uint32 GetTypeHash(const FGridGenerationKey& In)
 		{
-			return HashCombine(GetTypeHash(In.GridCoords), GetTypeHash(In.GridSize), GetTypeHash(In.bUse2DGrid), GetTypeHash(In.OriginalComponent));
+			return HashCombine(GetTypeHash(In.GridCoords), GetTypeHash(In.GridSize), GetTypeHash(In.bUse2DGrid), GetTypeHash(In.OriginalSource));
 		}
+
+		// Deprecated section
+	protected:
+		UE_DEPRECATED(5.8, "Use GetOriginalSource() instead.")
+		PCG_API UPCGComponent* GetOriginalComponent() const;
+		UE_DEPRECATED(5.8, "Use GetOriginalSourceEvenIfGarbage() instead.")
+		PCG_API UPCGComponent* GetOriginalComponentEvenIfGarbage() const;
+		UE_DEPRECATED(5.8, "Use SetOriginalSource() instead.")
+		PCG_API void SetOriginalComponent(UPCGComponent* InComponent);
+		UE_DEPRECATED(5.8, "Use GetCachedLocalSource() instead.")
+		PCG_API IPCGGraphExecutionSource* GetCachedLocalComponent() const;
+		UE_DEPRECATED(5.8, "Use SetCachedLocalSource() instead.")
+		PCG_API void SetCachedLocalComponent(IPCGGraphExecutionSource* InLocalComponent) const;
 	};
 
 	/** Returns true if the scheduler should tick this frame. */
 	bool ShouldTick();
 
 	// Add constructor, get moves for the sets
-	struct FTickQueueComponentsForGenerationInputs
+	struct FTickQueueSourcesForGenerationInputs
 	{
 		const TSet<IPCGGenSourceBase*>* GenSources = nullptr;
 		const APCGWorldActor* PCGWorldActor = nullptr;
 		TSet<IPCGGraphExecutionSource*> AllPartitionedExecutionSources;
 		TSet<IPCGGraphExecutionSource*> AllNonPartitionedExecutionSources;
-		TSet<FGridGenerationKey>* GeneratedComponents = nullptr;
+		TSet<FGridGenerationKey>* GeneratedSources = nullptr;
 	};
 
-	/** Queue nearby components for generation. */
-	void TickQueueComponentsForGeneration(const FTickQueueComponentsForGenerationInputs& Inputs, TMap<FGridGenerationKey, double>& OutComponentsToGenerate);
+	/** Queue nearby execution sources for generation. */
+	void TickQueueSourcesForGeneration(const FTickQueueSourcesForGenerationInputs& Inputs, TMap<FGridGenerationKey, double>& OutSourcesToGenerate);
 
-	/** Perform immediate cleanup on components that become out of range. */
+	/** Perform immediate cleanup on execution sources that become out of range. */
 	void TickCleanup(const TSet<IPCGGenSourceBase*>& GenSources, const APCGWorldActor* InPCGWorldActor, double InEndTime);
 
-	/** Schedule generation on components in priority order. */
-	void TickScheduleGeneration(TMap<FGridGenerationKey, double>& InOutComponentsToGenerate);
+	/** Schedule generation on execution sources in priority order. */
+	void TickScheduleGeneration(TMap<FGridGenerationKey, double>& InOutSourcesToGenerate);
 
 	/** Request any required virtual textures to be primed within the necessary generation radius. */
 	void TickRequestVirtualTexturePriming(const TSet<IPCGGenSourceBase*>& InGenSources);
@@ -177,17 +194,17 @@ protected:
 	/** Detects changes in RuntimeGen CVars to keep the PA pool in a valid state. */
 	void TickCVars(const APCGWorldActor* InPCGWorldActor);
 
-	/** Cleanup all local components in the GeneratedComponents set. */
-	void CleanupLocalComponents(const APCGWorldActor* InPCGWorldActor);
+	/** Cleanup all local execution sources in the GeneratedSources set. */
+	void CleanupLocalSources(const APCGWorldActor* InPCGWorldActor);
 
-	/** Cleanup a component and remove it from the GeneratedComponents set. */
-	void CleanupComponent(const FGridGenerationKey& GenerationKey, IPCGGraphExecutionSource* GeneratedComponent);
+	/** Cleanup a execution source and remove it from the GeneratedSources set. */
+	void CleanupSource(const FGridGenerationKey& GenerationKey, IPCGGraphExecutionSource* InGeneratedSource);
 
-	/** Remove components from the GeneratedComponents set that have been marked for delayed refresh. Fully cleanup any that would be leaked otherwise. */
-	void CleanupDelayedRefreshComponents();
+	/** Remove execution sources from the GeneratedSources set that have been marked for delayed refresh. Fully cleanup any that would be leaked otherwise. */
+	void CleanupDelayedRefreshSources();
 
-	/** Refresh a generated component. bRemovePartitionActors will also perform a full cleanup of PAs and local components. */
-	void RefreshComponent(IPCGGraphExecutionSource* InComponent, bool bRemovePartitionActors = false);
+	/** Refresh a generated execution source. bRemovePartitionActors will also perform a full cleanup of PAs and local sources. */
+	void RefreshExecutionSource(IPCGGraphExecutionSource* InExecutionSource, bool bRemovePartitionActors = false);
 	
 	/** Grabs an empty RuntimeGen PA from the PartitionActorPool and initializes it at the given GridSize and GridCoords. If no PAs are available in the pool,
 	* the pool capacity will double and new PAs will be created.
@@ -207,11 +224,11 @@ private:
 	/** Called on world streaming events. */
 	void OnLevelStreamingStateChanged(UWorld* InWorld, const ULevelStreaming* InLevelStreaming, ULevel* InLevelIfLoaded, ELevelStreamingState InPreviousState, ELevelStreamingState InNewState);
 
-	/** Cleanup a local component from the Partition actor. */
-	void CleanupLocalComponent(APCGPartitionActor* PartitionActor, IPCGGraphExecutionSource* LocalComponent);
+	/** Cleanup a local execution source along with its partition actor if it exists. */
+	void CleanupLocalSource(APCGPartitionActor* InPartitionActor, IPCGGraphExecutionSource* InLocalSource);
 
-	/** Cleanup the remaining components that are not part of the GeneratedComponents */
-	void CleanupRemainingComponents(UPCGComponent* InOriginalComponent);
+	/** Cleanup the remaining execution sources that are not part of the GeneratedSources */
+	void CleanupRemainingSources(IPCGGraphExecutionSource* InOriginalSource);
 
 	/** Cache VT priming infos found in graph params of the given component. */
 	void CacheVirtualTexturePrimingInfos(UPCGComponent* InOriginalComponent);
@@ -222,25 +239,29 @@ private:
 	/** Whether the change detection is enabled which avoids scanning grids when not necessary. */
 	bool IsChangeDetectionEnabled() const;
 
+	void RegisterOriginalRuntimeGenExecutionSource(UPCGComponent* InOriginalComponent);
+	void UnregisterOriginalRuntimeGenExecutionSource(UPCGComponent* InOriginalComponent);
+	UPCGRuntimeGenExecutionSource* GetOriginalRuntimeGenExecutionSource(UPCGComponent* InOriginalComponent);
+
 #if WITH_EDITOR
-	void RegisterGraphParameterChangedEvent(UPCGComponent* InOriginalComponent);
-	void UnregisterGraphParameterChangedEvent(UPCGComponent* InOriginalComponent);
+	void RegisterGraphParameterChangedEvent(const UPCGComponent* InOriginalComponent);
+	void UnregisterGraphParameterChangedEvent(const UPCGComponent* InOriginalComponent);
 	void OnGraphParametersChanged(UPCGGraphInterface* InGraph, EPCGGraphParameterEvent InChangeType, FName InChangedPropertyName);
 #endif
 
-	/** Tracks the generated components managed by the RuntimeGenScheduler. For local components, this generation key will hold the original component.
-	* For non-partitioned components, the generation key should have unbounded grid size and (0, 0, 0) grid coordinates.
+	/** Tracks the generated execution sources managed by the RuntimeGenScheduler. For local execution sources, this generation key will hold the original execution source.
+	* For non-partitioned execution sources, the generation key should have unbounded grid size and (0, 0, 0) grid coordinates.
 	*/
-	TSet<FGridGenerationKey> GeneratedComponents;
+	TSet<FGridGenerationKey> GeneratedSources;
 
-	/** Tracks the components which should be removed from the GeneratedComponents set on the next tick. This helps us defer removal in case we get multiple
-	* refreshes in a single tick. For example, a shallow refresh followed by a deep refresh would require the generated components to persist, otherwise we
+	/** Tracks the execution sources which should be removed from the GeneratedSources set on the next tick. This helps us defer removal in case we get multiple
+	* refreshes in a single tick. For example, a shallow refresh followed by a deep refresh would require the generated execution sources to persist, otherwise we
 	* will leak Partition Actors.
 	*/
-	TSet<FGridGenerationKey> GeneratedComponentsToRemove;
+	TSet<FGridGenerationKey> GeneratedSourcesToRemove;
 
-	/** Mapping of component + coordinates to priorities - needed to compute max priority over all gen sources. */
-	TMap<FGridGenerationKey, double> ComponentsToGenerate;
+	/** Mapping of execution source + coordinates to priorities - needed to compute max priority over all gen sources. */
+	TMap<FGridGenerationKey, double> SourcesToGenerate;
 
 	// Local to member functions but hoisted for efficiency.
 	TSet<IPCGGenSourceBase*> GenSources;
@@ -268,13 +289,13 @@ private:
 
 	/** Setting up a PA calls APCGPartitionActor::AddGraphInstance which later calls RefreshComponent, which can create
 	* an infinite refresh loop. To break this loop we write the OC pointer to this variable, and if Refresh gets called for
-	* this OC we early out. Basically we don't respond to refresh calls for a component we are midway through setting up.
+	* this OC we early out. Basically we don't respond to refresh calls for a execution source we are midway through setting up.
 	*/
-	const UPCGComponent* OriginalComponentBeingGenerated = nullptr;
+	const IPCGGraphExecutionSource* OriginalSourceBeingRegistered = nullptr;
 
 	int32 FramesUntilGeneration = 0;
 
-	// Local to TickQueueComponentsForGeneration, cached here for efficiency.
+	// Local to TickQueueSourcesForGeneration, cached here for efficiency.
 	using FStreamingCompleteQueryKey = TPair<FVector, float>;
 	TMap<FStreamingCompleteQueryKey, bool> CachedStreamingQueryResults;
 
@@ -300,5 +321,19 @@ private:
 	/** Maps original component to a cache of VT priming infos. */
 	TMap<TObjectKey<UPCGComponent>, TArray<FPCGVirtualTexturePrimingInfo>> CachedPrimingInfos;
 
+	/** Maps an original PCG component to its runtime gen execution source. This relationship only exists for components whose graphs opt-in to actor-component-less generation. */
+	TMap<TObjectKey<UPCGComponent>, TObjectPtr<UPCGRuntimeGenExecutionSource>> OriginalExecutionSources;
+
 	friend struct PCGRuntimeGenSchedulerHelpers::FStatsOverlay;
+
+	// Deprecated section
+protected:
+	struct UE_DEPRECATED(5.8, "Use FTickQueueSourcesForGenerationInputs instead") FTickQueueComponentsForGenerationInputs
+	{
+		const TSet<IPCGGenSourceBase*>* GenSources = nullptr;
+		const APCGWorldActor* PCGWorldActor = nullptr;
+		TSet<IPCGGraphExecutionSource*> AllPartitionedExecutionSources;
+		TSet<IPCGGraphExecutionSource*> AllNonPartitionedExecutionSources;
+		TSet<FGridGenerationKey>* GeneratedComponents = nullptr;
+	};
 };
diff --git a/Engine/Plugins/PCG/Source/PCG/Public/Subsystems/PCGSubsystem.h b/Engine/Plugins/PCG/Source/PCG/Public/Subsystems/PCGSubsystem.h
--- a/Engine/Plugins/PCG/Source/PCG/Public/Subsystems/PCGSubsystem.h
+++ b/Engine/Plugins/PCG/Source/PCG/Public/Subsystems/PCGSubsystem.h
@@ -148,6 +148,9 @@ public:
 	PCG_API void OnOriginalExecutionSourceRegistered(IPCGGraphExecutionSource* InExecutionSource);
 	PCG_API void OnOriginalExecutionSourceUnregistered(IPCGGraphExecutionSource* InExecutionSource);
 
+	/** In case of blueprints, we need to remap the old execution source destroyed by the construction script to the new one. */
+	PCG_API void OnOriginalExecutionSourceReplaced(const IPCGGraphExecutionSource* InOldExecutionSource, IPCGGraphExecutionSource* InNewExecutionSource);
+
 	PCG_API UPCGLandscapeCache* GetLandscapeCache();
 
 	/** Schedule graph(owner->graph) */
@@ -204,11 +207,11 @@ public:
 	*/
 	PCG_API FPCGTaskId ScheduleGenericWithContext(TFunction<bool(FPCGContext*)> InOperation, TFunction<void(FPCGContext*)> InAbortOperation, IPCGGraphExecutionSource* ExecutionSource, const TArray<FPCGTaskId>& TaskExecutionDependencies, const TArray<FPCGTaskId>& TaskDataDependencies) const;
 
-	/** Asks the runtime generation scheduler to refresh a given GenerateAtRuntime component. ChangeType should be 'GenerationGrid' to perform a full cleanup of PAs and local components. */
-	PCG_API void RefreshRuntimeGenComponent(UPCGComponent* RuntimeComponent, EPCGChangeType ChangeType = EPCGChangeType::None);
+	/** Asks the runtime generation scheduler to refresh a given GenerateAtRuntime execution source. ChangeType should be 'GenerationGrid' to perform a full cleanup of PAs and local components. */
+	PCG_API void RefreshRuntimeGenExecutionSource(IPCGGraphExecutionSource* RuntimeExecutionSource, EPCGChangeType ChangeType = EPCGChangeType::None);
 
-	/** Asks the runtime generation scheduler to refresh all GenerateAtRuntime components. ChangeType should be 'GenerationGrid' to perform a full cleanup of PAs and local components. */
-	PCG_API void RefreshAllRuntimeGenComponents(EPCGChangeType ChangeType = EPCGChangeType::None);
+	/** Asks the runtime generation scheduler to refresh all GenerateAtRuntime execution sources. ChangeType should be 'GenerationGrid' to perform a full cleanup of PAs and local components. */
+	PCG_API void RefreshAllRuntimeGenExecutionSources(EPCGChangeType ChangeType = EPCGChangeType::None);
 
 #if WITH_EDITOR
 	/** Refresh all components selected by the filter (runtime generated or otherwise). */
@@ -277,6 +280,12 @@ public:
 	UE_DEPRECATED(5.8, "Use ForAllIntersectingPartitionedExecutionSources instead")
 	PCG_API void ForAllIntersectingPartitionedComponents(const FBoxCenterAndExtent& InBounds, TFunctionRef<void(UPCGComponent*)> InFunc) const;
 
+	UE_DEPRECATED(5.8, "Use RefreshRuntimeGenExecutionSource instead")
+	PCG_API void RefreshRuntimeGenComponent(UPCGComponent* RuntimeComponent, EPCGChangeType ChangeType = EPCGChangeType::None);
+
+	UE_DEPRECATED(5.8, "Use RefreshAllRuntimeGenExecutionSources instead")
+	PCG_API void RefreshAllRuntimeGenComponents(EPCGChangeType ChangeType = EPCGChangeType::None);
+
 	/**
 	 * Call InFunc to all partition grid cells matching 'InGridSizes' and overlapping with 'InBounds'. 'InFunc' can schedule work or execute immediately.
 	 * 'InGridSizes' should be sorted in descending order. If 'bCanCreateActor' is true, it will create the partition actor at that cell if necessary.
diff --git a/Engine/Plugins/PCG/Source/PCGEditor/Internal/ChangeTracking/PCGRuntimeGenChangeHandler.h b/Engine/Plugins/PCG/Source/PCGEditor/Internal/ChangeTracking/PCGRuntimeGenChangeHandler.h
new file mode 100644
--- /dev/null
+++ b/Engine/Plugins/PCG/Source/PCGEditor/Internal/ChangeTracking/PCGRuntimeGenChangeHandler.h
@@ -0,0 +1,42 @@
+// Copyright Epic Games, Inc. All Rights Reserved.
+
+#pragma once
+
+#if WITH_EDITOR
+
+#include "ChangeTracking/PCGChangeTrackingRegistry.h"
+
+class FPCGTrackingManager;
+class IPCGGraphExecutionSource;
+class UPCGRuntimeGenExecutionSource;
+
+class FPCGRuntimeGenChangeHandler : public IPCGChangeHandler
+{
+public:
+	static TUniquePtr<IPCGChangeHandler> MakeInstance(FPCGTrackingManager* InOwner);
+	static FName GetName();
+
+	//~ Begin IPCGChangeHandler Interface
+	virtual void BeginChangeHandling(const TSharedRef<FPCGChangeHandlerChange>& InChange) override;
+	virtual void HandleChange(IPCGGraphExecutionSource* InExecutionSource) override;
+	virtual void HandleBoundedChange(IPCGGraphExecutionSource* InExecutionSource, const FBox& InExecutionSourceBounds, const FBox& InChangeBounds) override;
+	virtual void EndChangeHandling(bool bSkipRefresh) override;
+	//~ End IPCGChangeHandler Interface
+
+private:
+	FPCGRuntimeGenChangeHandler(FPCGTrackingManager* InOwner);
+
+	void HandleChangeInternal(IPCGGraphExecutionSource* InExecutionSource, const FBox* InExecutionSourceBounds = nullptr, const FBox* InChangeBounds = nullptr);
+	bool ShouldDiscardExecutionSource(UPCGRuntimeGenExecutionSource* InExecutionSource) const;
+
+	struct FCurrentChange
+	{
+		TWeakPtr<FPCGChangeHandlerChange> Change;
+		TArray<const UObject*> ChangedObjects;
+		TSet<UPCGRuntimeGenExecutionSource*> DirtyExecutionSources;
+	};
+
+	TUniquePtr<FCurrentChange> CurrentChange;
+};
+
+#endif // WITH_EDITOR
diff --git a/Engine/Plugins/PCG/Source/PCGEditor/Private/ChangeTracking/PCGRuntimeGenChangeHandler.cpp b/Engine/Plugins/PCG/Source/PCGEditor/Private/ChangeTracking/PCGRuntimeGenChangeHandler.cpp
new file mode 100644
--- /dev/null
+++ b/Engine/Plugins/PCG/Source/PCGEditor/Private/ChangeTracking/PCGRuntimeGenChangeHandler.cpp
@@ -0,0 +1,164 @@
+// Copyright Epic Games, Inc. All Rights Reserved.
+
+#if WITH_EDITOR
+
+#include "ChangeTracking/PCGRuntimeGenChangeHandler.h"
+
+#include "PCGSubsystem.h"
+#include "PCGContext.h"
+#include "PCGGraphExecutionStateInterface.h"
+#include "PCGTrackingManager.h"
+#include "RuntimeGen/PCGRuntimeGenExecutionSource.h"
+
+namespace PCGRuntimeGenChangeHandler
+{
+	constexpr FLazyName Name("PCGRuntimeGenChangeHandler");
+}
+
+FPCGRuntimeGenChangeHandler::FPCGRuntimeGenChangeHandler(FPCGTrackingManager* InOwner)
+	: IPCGChangeHandler(InOwner)
+{
+}
+
+TUniquePtr<IPCGChangeHandler> FPCGRuntimeGenChangeHandler::MakeInstance(FPCGTrackingManager* InOwner)
+{
+	return MakeUnique<IPCGChangeHandler>(InOwner);
+}
+
+FName FPCGRuntimeGenChangeHandler::GetName()
+{
+	return PCGRuntimeGenChangeHandler::Name;
+}
+
+void FPCGRuntimeGenChangeHandler::BeginChangeHandling(const TSharedRef<FPCGChangeHandlerChange>& InChange)
+{
+	check(InChange->ChangedObject);
+
+	CurrentChange = MakeUnique<FCurrentChange>();
+	CurrentChange->Change = InChange.ToWeakPtr();
+	CurrentChange->ChangedObjects.Add(InChange->ChangedObject);
+
+	if (InChange->OriginatingChangeObject != nullptr && InChange->ChangedObject != InChange->OriginatingChangeObject)
+	{
+		CurrentChange->ChangedObjects.Add(InChange->OriginatingChangeObject);
+	}
+}
+
+void FPCGRuntimeGenChangeHandler::HandleChange(IPCGGraphExecutionSource* InExecutionSource)
+{
+	HandleChangeInternal(InExecutionSource);
+}
+
+void FPCGRuntimeGenChangeHandler::HandleBoundedChange(IPCGGraphExecutionSource* InExecutionSource, const FBox& InExecutionSourceBounds, const FBox& InChangeBounds)
+{
+	HandleChangeInternal(InExecutionSource, &InExecutionSourceBounds, &InChangeBounds);
+}
+
+void FPCGRuntimeGenChangeHandler::EndChangeHandling(bool bSkipRefresh)
+{
+	check(CurrentChange.IsValid());
+	
+	// Release current change
+	ON_SCOPE_EXIT
+	{
+		CurrentChange.Reset();
+	};
+
+	if (bSkipRefresh)
+	{
+		return;
+	}
+
+	TSharedPtr<FPCGChangeHandlerChange> Change = CurrentChange->Change.Pin();
+
+	if (!Change)
+	{
+		return;
+	}
+
+	// And refresh all dirtied execution sources.
+	for (UPCGRuntimeGenExecutionSource* DirtyExecutionSource : CurrentChange->DirtyExecutionSources)
+	{
+		if (!ensure(DirtyExecutionSource))
+		{
+			continue;
+		}
+
+		const bool bOwnerHasChanged = DirtyExecutionSource->GetExecutionState().GetTarget() == Change->ChangedObject;
+
+		if (!Change->bNoRefreshOwner || !bOwnerHasChanged)
+		{
+			if (UPCGSubsystem* Subsystem = Cast<UPCGSubsystem>(DirtyExecutionSource->GetExecutionState().GetSubsystem()))
+			{
+				Subsystem->RefreshRuntimeGenExecutionSource(DirtyExecutionSource);
+			}
+		}
+	}
+}
+
+void FPCGRuntimeGenChangeHandler::HandleChangeInternal(IPCGGraphExecutionSource* InExecutionSource, const FBox* InExecutionSourceBounds, const FBox* InChangeBounds)
+{
+	check(CurrentChange.IsValid());
+	
+	TSharedPtr<FPCGChangeHandlerChange> Change = CurrentChange->Change.Pin();
+
+	if (!Change)
+	{
+		return;
+	}
+
+	UPCGRuntimeGenExecutionSource* ExecutionSource = Cast<UPCGRuntimeGenExecutionSource>(InExecutionSource);
+
+	if (!ExecutionSource || CurrentChange->DirtyExecutionSources.Contains(ExecutionSource))
+	{
+		return;
+	}
+
+	if (ShouldDiscardExecutionSource(ExecutionSource))
+	{
+		return;
+	}
+
+	if (InExecutionSourceBounds && InChangeBounds && !InChangeBounds->Intersect(*InExecutionSourceBounds))
+	{
+		return;
+	}
+
+	if (Owner->ClearCacheForKeys(Change->MatchedKeys, ExecutionSource, /*bIntersect=*/true, Change->OriginatingChangeObject))
+	{
+		CurrentChange->DirtyExecutionSources.Add(ExecutionSource);
+	}
+}
+
+bool FPCGRuntimeGenChangeHandler::ShouldDiscardExecutionSource(UPCGRuntimeGenExecutionSource* InExecutionSource) const
+{
+	check(CurrentChange.IsValid());
+
+	TSharedPtr<FPCGChangeHandlerChange> Change = CurrentChange->Change.Pin();
+
+	if (!Change)
+	{
+		return true;
+	}
+
+	// Discard null execution sources.
+	if (!InExecutionSource)
+	{
+		return true;
+	}
+
+	// Avoid loops if the change originates from our own execution.
+	if (FPCGContext* CurrentContext = FPCGContext::GetContextForThread(); CurrentContext && CurrentContext->ExecutionSource.Get() == InExecutionSource)
+	{
+		return true;
+	}
+
+	if (Change->bNoRefreshOwner && CurrentChange->ChangedObjects.Contains(InExecutionSource->GetExecutionState().GetTarget()))
+	{
+		return true;
+	}
+
+	return false;
+}
+
+#endif // WITH_EDITOR
diff --git a/Engine/Plugins/PCG/Source/PCGEditor/Private/Nodes/PCGEditorGraphNodeBase.cpp b/Engine/Plugins/PCG/Source/PCGEditor/Private/Nodes/PCGEditorGraphNodeBase.cpp
--- a/Engine/Plugins/PCG/Source/PCGEditor/Private/Nodes/PCGEditorGraphNodeBase.cpp
+++ b/Engine/Plugins/PCG/Source/PCGEditor/Private/Nodes/PCGEditorGraphNodeBase.cpp
@@ -2,10 +2,11 @@
 
 #include "Nodes/PCGEditorGraphNodeBase.h"
 
-#include "PCGComponent.h"
 #include "PCGEdge.h"
 #include "PCGEngineSettings.h"
 #include "PCGGraph.h"
+#include "PCGGraphExecutionInspection.h"
+#include "PCGGraphExecutionStateInterface.h"
 #include "PCGPin.h"
 #include "PCGSettingsWithDynamicInputs.h"
 #include "PCGSubsystem.h"
@@ -524,15 +525,21 @@ void UPCGEditorGraphNodeBase::SetInspected(bool bInIsInspecting)
 			const FPCGEditor* Editor = EditorGraph ? EditorGraph->GetEditor().Pin().Get() : nullptr;
 
 			// Force refresh/regenerate.
-			if (UPCGComponent* Component = Editor ? Cast<UPCGComponent>(Editor->GetPCGSourceBeingInspected()) : nullptr)
+			if (IPCGGraphExecutionSource* Source = Editor ? Editor->GetPCGSourceBeingInspected() : nullptr)
 			{
-				if (Component->IsManagedByRuntimeGenSystem())
+				if (Source->GetExecutionState().IsManagedByRuntimeGenSystem())
 				{
-					Component->Refresh(EPCGChangeType::Node);
+					if (UPCGSubsystem* Subsystem = Cast<UPCGSubsystem>(Source->GetExecutionState().GetSubsystem()))
+					{
+						Subsystem->RefreshRuntimeGenExecutionSource(Source, EPCGChangeType::Node);
+					}
 				}
 				else
 				{
-					Component->GenerateLocal(/*bForce=*/true);
+					IPCGGraphExecutionState::FGenerateParams GenerateParams;
+					GenerateParams.bEvenIfAlreadyGenerated = true;
+
+					Source->GetExecutionState().Generate(GenerateParams);
 				}
 			}
 		}
@@ -717,14 +724,13 @@ EPCGChangeType UPCGEditorGraphNodeBase::UpdateStructuralVisualization(const IPCG
 	const UPCGGraph* TopGraph = bInspecting ? InStackBeingInspected->GetRootGraph() : Graph;
 	const bool HiGenEnabled = TopGraph && TopGraph->IsHierarchicalGenerationEnabled();
 
-	const UPCGComponent* ComponentBeingDebugged = Cast<const UPCGComponent>(InSourceBeingDebugged);
 	// Set the inspected grid size - this is used for grid size visualization.
 	uint32 InspectingGridSize = PCGHiGenGrid::UninitializedGridSize();
 	EPCGHiGenGrid InspectingGrid = EPCGHiGenGrid::Uninitialized;
-	if (TopGraph && TopGraph->IsHierarchicalGenerationEnabled() && ComponentBeingDebugged && (ComponentBeingDebugged->IsPartitioned() || ComponentBeingDebugged->IsLocalComponent()))
+	if (TopGraph && TopGraph->IsHierarchicalGenerationEnabled() && InSourceBeingDebugged && (InSourceBeingDebugged->GetExecutionState().IsPartitioned() || InSourceBeingDebugged->GetExecutionState().IsLocalSource()))
 	{
-		InspectingGridSize = ComponentBeingDebugged->GetGenerationGridSize();
-		InspectingGrid = ComponentBeingDebugged->GetGenerationGrid();
+		InspectingGridSize = InSourceBeingDebugged->GetExecutionState().GetGenerationGridSize();
+		InspectingGrid = InSourceBeingDebugged->GetExecutionState().GetGenerationGrid();
 	}
 
 	if (InspectedGenerationGrid != InspectingGrid)
@@ -743,9 +749,9 @@ EPCGChangeType UPCGEditorGraphNodeBase::UpdateStructuralVisualization(const IPCG
 		bShouldDisplayAsDisabled = (Graph != TopGraph) || !TopGraph->IsHierarchicalGenerationEnabled();
 
 		// If we're inspecting a component, it must either be a partitioned OC or an LC (because higen requires partitioning).
-		if (!bShouldDisplayAsDisabled && ComponentBeingDebugged)
+		if (!bShouldDisplayAsDisabled && InSourceBeingDebugged)
 		{
-			bShouldDisplayAsDisabled = !ComponentBeingDebugged->IsPartitioned() && !ComponentBeingDebugged->IsLocalComponent();
+			bShouldDisplayAsDisabled = !InSourceBeingDebugged->GetExecutionState().IsPartitioned() && !InSourceBeingDebugged->GetExecutionState().IsLocalSource();
 		}
 	}
 
diff --git a/Engine/Plugins/PCG/Source/PCGEditor/Private/PCGEditor.cpp b/Engine/Plugins/PCG/Source/PCGEditor/Private/PCGEditor.cpp
--- a/Engine/Plugins/PCG/Source/PCGEditor/Private/PCGEditor.cpp
+++ b/Engine/Plugins/PCG/Source/PCGEditor/Private/PCGEditor.cpp
@@ -282,11 +282,9 @@ void FPCGEditor::UpdateAfterInspectedStackChanged(const FPCGStack& FullStack)
 		
 		if (Source)
 		{
-			UPCGComponent* Component = Cast<UPCGComponent>(Source);
-		
 			// Implementation note: if we're inspecting and have not pre-run the graph, then it probably makes sense to enable inspection by default. 
 			// TODO This could be selected with a cvar though.
-			const bool bHasBeenGeneratedThisSession = Component && Component->bGenerated && Component->WasGeneratedThisSession();
+			const bool bHasBeenGeneratedThisSession = Source->GetExecutionState().IsGenerated() && Source->GetExecutionState().WasGeneratedThisSession();
 			const bool bWasAborted = LastExecutionStatus.IsSet() && LastExecutionStatus.GetValue().Key == Source && LastExecutionStatus.GetValue().Value == EPCGGenerationStatus::Aborted;
 			const bool bWasInspecting = Source->GetExecutionState().GetInspection().IsInspecting();
 			const bool bNeedsInspection = Algo::AnyOf(AttributesWidgets, [](const TSharedPtr<SPCGEditorGraphAttributeListView>& ALV) { return ALV.IsValid() && ALV->GetNodeBeingInspected() != nullptr; });
@@ -298,7 +296,7 @@ void FPCGEditor::UpdateAfterInspectedStackChanged(const FPCGStack& FullStack)
 				// Making sure to not re-trigger a new generation if the previous one was aborted, to avoid infinite loops
 				if (!bWasAborted)
 				{
-					UpdateDebugAfterComponentSelection(Component, Component, true);
+					UpdateDebugAfterSourceSelection(Source, Source, true);
 				}
 			}
 		}
@@ -344,54 +342,57 @@ IPCGGraphExecutionSource* FPCGEditor::GetPCGSourceBeingInspected() const
 	return InspectionDataManager.GetPCGSourceBeingInspected();
 }
 
-void FPCGEditor::UpdateDebugAfterComponentSelection(UPCGComponent* InOldComponent, UPCGComponent* InNewComponent, bool bInNewComponentStartedInspecting)
+void FPCGEditor::UpdateDebugAfterSourceSelection(IPCGGraphExecutionSource* InOldSource, IPCGGraphExecutionSource* InNewSource, bool bInNewSourceStartedInspecting)
 {
 	if (!ensure(PCGGraphBeingEdited))
 	{
 		return;
 	}
 
-	auto RefreshComponent = [](UPCGComponent* Component)
+	auto RefreshSource = [](IPCGGraphExecutionSource* Source)
 	{
-		if (!ensure(Component))
+		if (!ensure(Source))
 		{
 			return;
 		}
 
-		// GenerateAtRuntime components should be refreshed through the runtime gen scheduler.
-		if (Component->IsManagedByRuntimeGenSystem())
+		// GenerateAtRuntime sources should be refreshed through the runtime gen scheduler.
+		if (Source->GetExecutionState().IsManagedByRuntimeGenSystem())
 		{
 			if (UPCGSubsystem* WorldSubsystem = GetWorldSubsystem())
 			{
-				// We don't want to do a full cleanup if we're setting the debug object, since full cleanup destroys the component, which is the debug object itself!
-				WorldSubsystem->RefreshRuntimeGenComponent(Component);
+				// We don't want to do a full cleanup if we're setting the debug object, since full cleanup destroys the source, which is the debug object itself!
+				WorldSubsystem->RefreshRuntimeGenExecutionSource(Source);
 			}
 		}
 		else
 		{
-			Component->GenerateLocal(/*bForce=*/true);
+			IPCGGraphExecutionState::FGenerateParams GenerateParams;
+			GenerateParams.bEvenIfAlreadyGenerated = true;
+
+			Source->GetExecutionState().Generate(GenerateParams);
 		}
 	};
 
-	// If individual component debugging is disabled, just generate the new component if required.
+	// If individual source debugging is disabled, just generate the new source if required.
 	if (!PCGGraphBeingEdited->DebugFlagAppliesToIndividualComponents())
 	{
-		if (InNewComponent && bInNewComponentStartedInspecting)
+		if (InNewSource && bInNewSourceStartedInspecting)
 		{
-			RefreshComponent(InNewComponent);
+			RefreshSource(InNewSource);
 		}
 
 		return;
 	}
 
-	// Trigger necessary generation(s) for per-component debugging.
-	if (!InOldComponent)
+	// Trigger necessary generation(s) for per-source debugging.
+	if (!InOldSource)
 	{
-		if (InNewComponent && bInNewComponentStartedInspecting)
+		if (InNewSource && bInNewSourceStartedInspecting)
 		{
 			// Transition from 'null' to 'any component not already inspecting' - generate to create debug/inspection info.
-			// If we have null selected, all components are displaying debug. Go to Original component so that all refresh.
-			RefreshComponent(InNewComponent->GetOriginalComponent());
+			// If we have null selected, all sources are displaying debug. Go to Original source so that all refresh.
+			RefreshSource(InNewSource->GetExecutionState().GetOriginalSource());
 		}
 	}
 	else
@@ -401,18 +402,18 @@ void FPCGEditor::UpdateDebugAfterComponentSelection(UPCGComponent* InOldComponen
 			return InNode && InNode->GetSettings() && InNode->GetSettings()->bDebug;
 		});
 
-		// Regenerate to clear debug info if switching components, or if changing from a component to null.
-		if (InNewComponent != InOldComponent && (InNewComponent || bDebugFlagSetOnAnyNode))
+		// Regenerate to clear debug info if switching sources, or if changing from a source to null.
+		if (InNewSource != InOldSource && (InNewSource || bDebugFlagSetOnAnyNode))
 		{
-			// Use original component - debug can be displayed both by the local component and parent local components.
-			RefreshComponent(InOldComponent->GetOriginalComponent());
+			// Use original source - debug can be displayed both by the local source and parent local sources.
+			RefreshSource(InOldSource->GetExecutionState().GetOriginalSource());
 		}
 
-		// Debug new component if it wasn't already
-		if (InNewComponent && bInNewComponentStartedInspecting)
+		// Debug new source if it wasn't already
+		if (InNewSource && bInNewSourceStartedInspecting)
 		{
-			// Use original component - debug can be displayed both by the local component and parent local components.
-			RefreshComponent(InNewComponent->GetOriginalComponent());
+			// Use original source - debug can be displayed both by the local source and parent local sources.
+			RefreshSource(InNewSource->GetExecutionState().GetOriginalSource());
 		}
 	}
 }
diff --git a/Engine/Plugins/PCG/Source/PCGEditor/Private/PCGEditorModule.cpp b/Engine/Plugins/PCG/Source/PCGEditor/Private/PCGEditorModule.cpp
--- a/Engine/Plugins/PCG/Source/PCGEditor/Private/PCGEditorModule.cpp
+++ b/Engine/Plugins/PCG/Source/PCGEditor/Private/PCGEditorModule.cpp
@@ -40,6 +40,7 @@
 #include "ChangeTracking/PCGComponentChangeHandler.h"
 #include "ChangeTracking/PCGLandscapeTracker.h"
 #include "ChangeTracking/PCGMiscTracker.h"
+#include "ChangeTracking/PCGRuntimeGenChangeHandler.h"
 #include "DataVisualizations/PCGBaseTextureDataVisualization.h"
 #include "DataVisualizations/PCGLandscapeDataVisualization.h"
 #include "DataVisualizations/PCGParamDataVisualization.h"
@@ -857,7 +858,7 @@ void FPCGEditorModule::PopulateMenuActions(FMenuBuilder& MenuBuilder)
 			{
 				if (UPCGSubsystem* PCGSubsystem = UPCGSubsystem::GetSubsystemForCurrentWorld())
 				{
-					PCGSubsystem->RefreshAllRuntimeGenComponents(EPCGChangeType::GenerationGrid);
+					PCGSubsystem->RefreshAllRuntimeGenExecutionSources(EPCGChangeType::GenerationGrid);
 				}
 			})),
 		NAME_None);
@@ -989,6 +990,7 @@ void FPCGEditorModule::RegisterTrackerFactories()
 	TrackingRegistry.RegisterTrackerFactory(FPCGMiscTracker::GetName(), FPCGCreateTrackerInstance::CreateStatic(&FPCGMiscTracker::MakeInstance));
 
 	TrackingRegistry.RegisterHandlerFactory(FPCGComponentChangeHandler::GetName(), FPCGCreateHandlerInstance::CreateStatic(&FPCGComponentChangeHandler::MakeInstance));
+	TrackingRegistry.RegisterHandlerFactory(FPCGRuntimeGenChangeHandler::GetName(), FPCGCreateHandlerInstance::CreateStatic(&FPCGRuntimeGenChangeHandler::MakeInstance));
 }
 
 void FPCGEditorModule::UnregisterTrackerFactories()
@@ -999,6 +1001,7 @@ void FPCGEditorModule::UnregisterTrackerFactories()
 	TrackingRegistry.UnregisterTrackerFactory(FPCGMiscTracker::GetName());
 
 	TrackingRegistry.UnregisterHandlerFactory(FPCGComponentChangeHandler::GetName());
+	TrackingRegistry.UnregisterHandlerFactory(FPCGRuntimeGenChangeHandler::GetName());
 }
 
 void FPCGEditorModule::RegisterPinColorAndIcons()
diff --git a/Engine/Plugins/PCG/Source/PCGEditor/Private/Widgets/SPCGEditorGraphDebugObjectTree.cpp b/Engine/Plugins/PCG/Source/PCGEditor/Private/Widgets/SPCGEditorGraphDebugObjectTree.cpp
--- a/Engine/Plugins/PCG/Source/PCGEditor/Private/Widgets/SPCGEditorGraphDebugObjectTree.cpp
+++ b/Engine/Plugins/PCG/Source/PCGEditor/Private/Widgets/SPCGEditorGraphDebugObjectTree.cpp
@@ -266,11 +266,10 @@ FString FPCGEditorGraphDebugObjectItem_PCGSource::GetLabel(bool bForSorting) con
 		return LabelOverride.GetValue();
 	}
 
-	//@todo_pcg: Should be part of the interface?
-	const UObject* Source = PCGSource.GetObject();
+	const IPCGGraphExecutionSource* Source = Cast<IPCGGraphExecutionSource>(PCGSource.GetObject());
 	if (Source && PCGGraph.IsValid())
 	{
-		return PCGEditorGraphDebugObjectTree::GetStringFromName(Source->GetFName(), bForSorting) + FString(TEXT(" - ")) + PCGEditorGraphDebugObjectTree::GetStringFromName(PCGGraph->GetFName(), bForSorting);
+		return Source->GetExecutionState().GetDebugName() + FString(TEXT(" - ")) + PCGEditorGraphDebugObjectTree::GetStringFromName(PCGGraph->GetFName(), bForSorting);
 	}
 
 	return FString();
@@ -975,6 +974,13 @@ void SPCGEditorGraphDebugObjectTree::AddStackToTree(const FPCGStack& Stack,
 
 		// @todo_pcg: GetOwner might need to be part of the execution source interface.
 		UPCGComponent* PCGComponent = Cast<UPCGComponent>(PCGSource);
+
+		// Don't show actor-componentless components in the tree view, as they will never be generated and won't have anything to inspect.
+		if (PCGComponent && PCGComponent->UseActorComponentlessGeneration())
+		{
+			return;
+		}
+
 		FPCGEditorGraphDebugObjectItemPtr OwnerItem;
 		if (AActor* Actor = PCGComponent ? PCGComponent->GetOwner() : nullptr)
 		{
diff --git a/Engine/Plugins/PCG/Source/PCGEditor/Public/PCGEditor.h b/Engine/Plugins/PCG/Source/PCGEditor/Public/PCGEditor.h
--- a/Engine/Plugins/PCG/Source/PCGEditor/Public/PCGEditor.h
+++ b/Engine/Plugins/PCG/Source/PCGEditor/Public/PCGEditor.h
@@ -538,7 +538,7 @@ private:
 	void OnSourceGenerationDone(IPCGBaseSubsystem* Subsystem, IPCGGraphExecutionSource* Source, EPCGGenerationStatus Status);
 
 	/** Trigger any generation required to ensure debug display is up to date. */
-	void UpdateDebugAfterComponentSelection(UPCGComponent* InOldComponent, UPCGComponent* InNewComponent, bool bNewComponentStartedInspecting);
+	void UpdateDebugAfterSourceSelection(IPCGGraphExecutionSource* InOldSource, IPCGGraphExecutionSource* InNewSource, bool bNewSourceStartedInspecting);
 
 	void RegisterDelegatesForWorld(UWorld* World);
 	void UnregisterDelegatesForWorld(UWorld* World);
