[PCG] Data binding optimizations to smooth out hot spots in profile traces.

#rnx

#Synced-CL 50144017

#Synced-CL 50144017

diff --git a/Engine/Plugins/PCG/Source/PCG/Private/Compute/PCGDataBinding.cpp b/Engine/Plugins/PCG/Source/PCG/Private/Compute/PCGDataBinding.cpp
--- a/Engine/Plugins/PCG/Source/PCG/Private/Compute/PCGDataBinding.cpp
+++ b/Engine/Plugins/PCG/Source/PCG/Private/Compute/PCGDataBinding.cpp
@@ -403,6 +403,15 @@ const FPCGDataCollection& UPCGDataBinding::GetOutputDataCollection() const
 
 void UPCGDataBinding::AddInputDataAttributesToTable()
 {
+	if (InputDataCollection.TaggedData.IsEmpty())
+	{
+		return;
+	}
+
+	AttributeTable.Reserve(16); // Heuristic guess at a reasonable upper bound for many cases.
+
+	FPCGKernelAttributeTable::FAttributeKeyToIndexMap AttributeToGraphAttributeIndex = AttributeTable.GetAttributeKeyToIndex();
+
 	for (const FPCGTaggedData& Data : InputDataCollection.TaggedData)
 	{
 		if (const UPCGProxyForGPUData* Proxy = Cast<UPCGProxyForGPUData>(Data.Data))
@@ -422,7 +431,22 @@ void UPCGDataBinding::AddInputDataAttributesToTable()
 						continue;
 					}
 
-					const int32 GraphAttributeIndex = AttributeTable.AddAttribute(SourceAttrDesc.GetAttributeKey().GetIdentifier(), SourceAttrDesc.GetAttributeKey().GetType());
+					int32 GraphAttributeIndex = INDEX_NONE;
+					if (int32* FoundGraphAttributeIndex = AttributeToGraphAttributeIndex.Find(SourceAttrDesc.GetAttributeKey()))
+					{
+						GraphAttributeIndex = *FoundGraphAttributeIndex;
+					}
+					else
+					{
+						GraphAttributeIndex = AttributeTable.AddAttributeAssumeUnique(SourceAttrDesc.GetAttributeKey());
+						AttributeToGraphAttributeIndex.Add(SourceAttrDesc.GetAttributeKey(), GraphAttributeIndex);
+					}
+
+					if (GraphAttributeIndex == INDEX_NONE)
+					{
+						UE_LOG(LogPCG, Error, TEXT("Attribute table exceeded maximum size (%d), use the 'Dump Data Descriptions' setting on the GPU node(s) to list attributes that are present."), PCGDataCollectionPackingConstants::MAX_NUM_CUSTOM_ATTRS);
+						break;
+					}
 
 					if (GraphAttributeIndex != SourceAttributeTableIndex)
 					{
@@ -457,18 +481,30 @@ void UPCGDataBinding::AddInputDataAttributesToTable()
 				{
 					AttributeIdentifier.MetadataDomain = PCGMetadataDomainID::Default;
 				}
-					
+
 				const EPCGKernelAttributeType AttributeType = PCGDataDescriptionHelpers::GetAttributeTypeFromMetadataType(AttributeTypes[Index]);
+				const FPCGKernelAttributeKey Key(AttributeIdentifier, AttributeType);
 
-				if (AttributeTable.AddAttribute(AttributeIdentifier, AttributeType) == INDEX_NONE)
+				if (!AttributeToGraphAttributeIndex.Contains(Key))
 				{
-					UE_LOG(LogPCG, Error, TEXT("Attribute table exceeded maximum size (%d), use the 'Dump Data Descriptions' setting on the GPU node(s) to list attributes that are present."), PCGDataCollectionPackingConstants::MAX_NUM_CUSTOM_ATTRS);
-					break;
+					const int32 GraphAttributeIndex = AttributeTable.AddAttributeAssumeUnique(Key);
+
+					if (GraphAttributeIndex != INDEX_NONE)
+					{
+						AttributeToGraphAttributeIndex.Add(Key, GraphAttributeIndex);
+					}
+					else
+					{
+						UE_LOG(LogPCG, Error, TEXT("Attribute table exceeded maximum size (%d), use the 'Dump Data Descriptions' setting on the GPU node(s) to list attributes that are present."), PCGDataCollectionPackingConstants::MAX_NUM_CUSTOM_ATTRS);
+						break;
+					}
 				}
 			}
 		}
 	}
 
+	AttributeTable.Shrink();
+
 	ensure(AttributeTable.Num() <= PCGDataCollectionPackingConstants::MAX_NUM_CUSTOM_ATTRS);
 }
 
diff --git a/Engine/Plugins/PCG/Source/PCG/Private/Compute/PCGDataDescription.cpp b/Engine/Plugins/PCG/Source/PCG/Private/Compute/PCGDataDescription.cpp
--- a/Engine/Plugins/PCG/Source/PCG/Private/Compute/PCGDataDescription.cpp
+++ b/Engine/Plugins/PCG/Source/PCG/Private/Compute/PCGDataDescription.cpp
@@ -190,6 +190,31 @@ int32 FPCGKernelAttributeTable::AddAttribute(FPCGAttributeIdentifier InIdentifie
 	return AddAttribute(FPCGKernelAttributeKey(InIdentifier, InType));
 }
 
+int32 FPCGKernelAttributeTable::AddAttributeAssumeUnique(const FPCGKernelAttributeKey& Key)
+{
+	if (AttributeTable.Num() < PCGDataCollectionPackingConstants::MAX_NUM_CUSTOM_ATTRS)
+	{
+		return AttributeTable.Add(Key);
+	}
+	else
+	{
+		return INDEX_NONE;
+	}
+}
+
+FPCGKernelAttributeTable::FAttributeKeyToIndexMap FPCGKernelAttributeTable::GetAttributeKeyToIndex() const
+{
+	FAttributeKeyToIndexMap AttributeToIndex;
+	AttributeToIndex.Reserve(AttributeTable.Num());
+
+	for (int Index = 0; Index < AttributeTable.Num(); ++Index)
+	{
+		AttributeToIndex.Add(AttributeTable[Index], Index);
+	}
+
+	return AttributeToIndex;
+}
+
 #if PCG_KERNEL_LOGGING_ENABLED
 void FPCGKernelAttributeTable::DebugLog() const
 {
@@ -503,7 +528,7 @@ FPCGDataDesc FPCGDataDesc::CombineElementCount(const FPCGDataDesc& A, const FPCG
 bool FPCGDataDesc::IsAttributeAllocated(int32 InAttributeId) const
 {
 	// @todo_pcg: Only point properties are supported for single value ranges. All other attributes are fully allocated.
-	if (Type == EPCGDataType::Point && InAttributeId < PCGDataCollectionPackingConstants::NUM_RESERVED_ATTRS)
+	if (Type == FPCGDataTypeInfoPoint::AsId() && InAttributeId < PCGDataCollectionPackingConstants::NUM_RESERVED_ATTRS)
 	{
 		switch (InAttributeId)
 		{
@@ -543,7 +568,7 @@ void FPCGDataDesc::InitializeAttributeDescs(const UPCGData* InData, const UPCGDa
 {
 	TRACE_CPUPROFILER_EVENT_SCOPE(FPCGDataDesc::InitializeAttributeDescs);
 
-	if (Type == EPCGDataType::Point)
+	if (Type == FPCGDataTypeInfoPoint::AsId())
 	{
 		AttributeDescs.Append(PCGDataDescriptionConstants::PointPropertyDescs, PCGDataCollectionPackingConstants::NUM_POINT_PROPERTIES);
 
diff --git a/Engine/Plugins/PCG/Source/PCG/Public/Compute/PCGDataDescription.h b/Engine/Plugins/PCG/Source/PCG/Public/Compute/PCGDataDescription.h
--- a/Engine/Plugins/PCG/Source/PCG/Public/Compute/PCGDataDescription.h
+++ b/Engine/Plugins/PCG/Source/PCG/Public/Compute/PCGDataDescription.h
@@ -92,10 +92,21 @@ struct FPCGKernelAttributeTable
 	int32 AddAttribute(const FPCGKernelAttributeKey& Key);
 	int32 AddAttribute(FPCGAttributeIdentifier InIdentifier, EPCGKernelAttributeType InType);
 
+	/** Caller must not insert the same key twice. */
+	int32 AddAttributeAssumeUnique(const FPCGKernelAttributeKey& Key);
+
 	int32 Num() const { return AttributeTable.Num(); }
 
 	void Empty() { AttributeTable.Empty(); }
 
+	void Reserve(int32 InCapacity) { AttributeTable.Reserve(InCapacity); }
+	void Shrink() { AttributeTable.Shrink(); }
+
+	using FAttributeKeyToIndexMap = TMap<FPCGKernelAttributeKey, int32, TInlineSetAllocator<32>>;
+
+	/** Export a map of keys to attribute indices. */
+	FAttributeKeyToIndexMap GetAttributeKeyToIndex() const;
+
 #if PCG_KERNEL_LOGGING_ENABLED
 	void DebugLog() const;
 #endif
