[PCG] Fix bFlipAxis issue for Subdivide Segment

* Logic was wrong and it was not working (Size was a negative value)

#rnx

#Synced-CL 50157473

diff --git a/Engine/Plugins/PCG/Source/PCG/Private/Elements/Grammar/PCGSubdivideSegment.cpp b/Engine/Plugins/PCG/Source/PCG/Private/Elements/Grammar/PCGSubdivideSegment.cpp
--- a/Engine/Plugins/PCG/Source/PCG/Private/Elements/Grammar/PCGSubdivideSegment.cpp
+++ b/Engine/Plugins/PCG/Source/PCG/Private/Elements/Grammar/PCGSubdivideSegment.cpp
@@ -71,13 +71,12 @@ public:
 			const FVector PreviousMin = Segment.Min;
 			Segment.Min = Segment.Min * InOutParameters.PerpendicularSubdivisionDirection + Segment.Max * InOutParameters.SubdivisionDirection;
 			Segment.Max = Segment.Max * InOutParameters.PerpendicularSubdivisionDirection + PreviousMin * InOutParameters.SubdivisionDirection;
-			PointScaledSize *= (InOutParameters.PerpendicularSubdivisionDirection - InOutParameters.SubdivisionDirection);
 		}
 
-		const FVector Direction = TransformNoTranslation.TransformVectorNoScale(InOutParameters.SubdivisionDirection).GetSafeNormal();
+		const FVector Direction = TransformNoTranslation.TransformVectorNoScale(bFlipAxis ? -InOutParameters.SubdivisionDirection : InOutParameters.SubdivisionDirection).GetSafeNormal();
 		const FVector OtherDirection = TransformNoTranslation.TransformVectorNoScale(PointScaledSize * InOutParameters.PerpendicularSubdivisionDirection) * 0.5;
 		const FVector HalfExtents2D = PointScaledSize * InOutParameters.PerpendicularSubdivisionDirection * 0.5;
-		const double Size = PointScaledSize.Dot(InOutParameters.SubdivisionDirection);
+		const double Size = FMath::Abs(PointScaledSize.Dot(InOutParameters.SubdivisionDirection));
 
 		TArray<PCGSubdivisionBase::TModuleInstance<PCGGrammar::FTokenizedModule>> ModulesInstances;
 		double RemainingSubdivide;
