[PCG] More fixes related to refactor
- Attribute list view was no longer created on inspection (if it didn't already exist)
- Re-bind some missing events on specific widget tab closing

#rnx

#Synced-CL 50158750

diff --git a/Engine/Plugins/PCG/Source/PCGEditor/Private/PCGEditor.cpp b/Engine/Plugins/PCG/Source/PCGEditor/Private/PCGEditor.cpp
--- a/Engine/Plugins/PCG/Source/PCGEditor/Private/PCGEditor.cpp
+++ b/Engine/Plugins/PCG/Source/PCGEditor/Private/PCGEditor.cpp
@@ -913,22 +913,24 @@ void FPCGEditor::OpenDetailsView()
 	}
 }
 
-void FPCGEditor::OnDetailsViewTabClosed(TSharedRef<SDockTab> DockTab, int Index)
+void FPCGEditor::OnDetailsViewTabClosed(FName Id, int32 Index)
 {
 	check(Index >= 0 && Index < PCGEditorTabs::NumPropertyDetailTabs);
+	check(PCGEditorTabs::PropertyDetailsID[Index] == Id);
 
-	TSharedPtr<SPCGEditorGraphDetailsView> DetailsView = PropertyDetailsWidgets[Index];
-	if (DetailsView.IsValid() && DetailsView->IsLocked())
+	if (PropertyDetailsWidgets[Index].IsValid() && PropertyDetailsWidgets[Index]->IsLocked())
 	{
-		DetailsView->SetIsLocked(false);
+		PropertyDetailsWidgets[Index]->SetIsLocked(false);
 	}
 }
 
-void FPCGEditor::OnAttributeListViewTabClosed(TSharedRef<SDockTab> DockTab, int Index)
+void FPCGEditor::OnAttributeListViewTabClosed(FName Id, int32 Index)
 {
 	check(Index >= 0 && Index < PCGEditorTabs::NumAttributeTabs);
+	check(PCGEditorTabs::AttributesID[Index] == Id);
 
 	TSharedPtr<SPCGEditorGraphAttributeListView> AttributeListView = AttributesWidgets[Index];
+
 	if (AttributeListView.IsValid())
 	{
 		if (AttributeListView->IsLocked())
@@ -968,13 +970,14 @@ void FPCGEditor::OnAttributeListViewTabClosed(TSharedRef<SDockTab> DockTab, int
 	}
 }
 
-void FPCGEditor::OnViewportViewTabClosed(TSharedRef<SDockTab> DockTab, int Index)
+void FPCGEditor::OnViewportViewTabClosed(FName Id, int32 Index)
 {
-	check(Index >= 0 && Index < PCGEditorTabs::NumAttributeTabs);
-
-	if (AttributesWidgets[Index].IsValid())
+	check(Index >= 0 && Index < PCGEditorTabs::NumViewportTabs);
+	check(PCGEditorTabs::ViewportID[Index] == Id);
+	
+	if (ViewportWidgets[Index].IsValid())
 	{
-		AttributesWidgets[Index]->ResetViewport();
+		ViewportWidgets[Index]->ResetScene();
 	}
 }
 
@@ -1886,16 +1889,20 @@ void FPCGEditor::SetNodeInspected(UPCGEditorGraphNodeBase* InspectedNode, bool b
 	{
 		TArray<UPCGEditorGraphNodeBase*, TInlineAllocator<4>> InspectedNodesAfter;
 
-		for (TSharedPtr<SPCGEditorGraphAttributeListView> AttributeListView : AttributesWidgets)
+		for (int32 Index = 0; Index < PCGEditorTabs::NumAttributeTabs; ++Index)
 		{
-			if (AttributeListView.IsValid())
+			// If not yet created do that now
+			if (!AttributesWidgets[Index].IsValid())
 			{
-				if (!AttributeListView->IsLocked())
-				{
-					AttributeListView->SetNodeBeingInspected(InspectedNode);
-				}
+				GetOrCreateAttributesWidget(Index);
+				check(AttributesWidgets[Index].IsValid());
+			}
 
-				InspectedNodesAfter.Add(AttributeListView->GetNodeBeingInspected());
+			if (!AttributesWidgets[Index]->IsLocked())
+			{
+				AttributesWidgets[Index]->SetNodeBeingInspected(InspectedNode);
+			
+				InspectedNodesAfter.Add(AttributesWidgets[Index]->GetNodeBeingInspected());
 			}
 		}
 
@@ -2942,6 +2949,7 @@ TSharedPtr<FWorkflowTabFactory> FPCGEditor::CreateAttributesTabFactory(int32 Ind
 	
 	Params.Label = AttributeText[Index];
 	Params.Group = Group;
+	Params.OnTabClosed = FPCGGenericTabClosed::CreateSP(this, &FPCGEditor::OnAttributeListViewTabClosed, Index);
 	return MakeShared<FPCGGenericTabFactory>(Params);
 }
 
@@ -3062,6 +3070,7 @@ TSharedPtr<FWorkflowTabFactory> FPCGEditor::CreatePropertyDetailsTabFactory(int3
 
 	Params.Label = PropertyDetailText[Index];
 	Params.Group = Group;
+	Params.OnTabClosed = FPCGGenericTabClosed::CreateSP(this, &FPCGEditor::OnDetailsViewTabClosed, Index);
 	return MakeShared<FPCGGenericTabFactory>(Params);
 }
 
@@ -3100,6 +3109,7 @@ TSharedPtr<FWorkflowTabFactory> FPCGEditor::CreateViewportTabFactory(int32 Index
 
 	Params.Label = ViewportText[Index];
 	Params.Group = Group;
+	Params.OnTabClosed = FPCGGenericTabClosed::CreateSP(this, &FPCGEditor::OnViewportViewTabClosed, Index);
 	return MakeShared<FPCGGenericTabFactory>(Params);
 }
 
diff --git a/Engine/Plugins/PCG/Source/PCGEditor/Private/PCGEditorTabFactories.cpp b/Engine/Plugins/PCG/Source/PCGEditor/Private/PCGEditorTabFactories.cpp
--- a/Engine/Plugins/PCG/Source/PCGEditor/Private/PCGEditorTabFactories.cpp
+++ b/Engine/Plugins/PCG/Source/PCGEditor/Private/PCGEditorTabFactories.cpp
@@ -4,6 +4,8 @@
 
 #include "PCGEditor.h"
 
+#include "Widgets/Docking/SDockTab.h"
+
 #define LOCTEXT_NAMESPACE "PCGGraphEditor"
 
 FPCGGenericTabFactory::FPCGGenericTabFactory(const FPCGGenericTabFactoryParams& InParams)
@@ -14,6 +16,8 @@ FPCGGenericTabFactory::FPCGGenericTabFactory(const FPCGGenericTabFactoryParams&
 	CreateTabBodyFunc = InParams.CreateTabBodyFunc;
 	check(CreateTabBodyFunc.IsSet());
 
+	OnTabClosedDelegate = InParams.OnTabClosed;
+
 	bIsSingleton = InParams.bSingleton;
 
 	if (InParams.Label.IsSet())
@@ -34,9 +38,27 @@ FPCGGenericTabFactory::FPCGGenericTabFactory(const FPCGGenericTabFactoryParams&
 	Group = InParams.Group;
 }
 
+TSharedRef<SDockTab> FPCGGenericTabFactory::SpawnTab(const FWorkflowTabSpawnInfo& Info) const
+{
+	TSharedRef<SDockTab> Tab = FWorkflowTabFactory::SpawnTab(Info);
+
+	if (OnTabClosedDelegate.IsBound())
+	{
+		Tab->SetOnTabClosed(SDockTab::FOnTabClosedCallback::CreateSP(this, &FPCGGenericTabFactory::OnTabClosed));
+	}
+
+	return Tab;
+}
+
 TSharedRef<SWidget> FPCGGenericTabFactory::CreateTabBody(const FWorkflowTabSpawnInfo& Info) const
 {
 	return CreateTabBodyFunc();
 }
 
+void FPCGGenericTabFactory::OnTabClosed(TSharedRef<SDockTab> DockTab) const
+{
+	OnTabClosedDelegate.ExecuteIfBound(GetIdentifier());
+}
+
+
 #undef LOCTEXT_NAMESPACE
\ No newline at end of file
diff --git a/Engine/Plugins/PCG/Source/PCGEditor/Public/PCGEditor.h b/Engine/Plugins/PCG/Source/PCGEditor/Public/PCGEditor.h
--- a/Engine/Plugins/PCG/Source/PCGEditor/Public/PCGEditor.h
+++ b/Engine/Plugins/PCG/Source/PCGEditor/Public/PCGEditor.h
@@ -320,13 +320,13 @@ private:
 	void OpenDetailsView();
 
 	/** Called when a details view tab is closed */
-	void OnDetailsViewTabClosed(TSharedRef<SDockTab> DockTab, int Index);
+	void OnDetailsViewTabClosed(FName Id, int32 Index);
 
 	/** Called when an attribute list view tab is closed */
-	void OnAttributeListViewTabClosed(TSharedRef<SDockTab> DockTab, int Index);
+	void OnAttributeListViewTabClosed(FName Id, int32 Index);
 
 	/** Called when a viewport view tab is closed */
-	void OnViewportViewTabClosed(TSharedRef<SDockTab> DockTab, int Index);
+	void OnViewportViewTabClosed(FName Id, int32 Index);
 
 	/** Enable/Disable automatic PCG node generation */
 	void OnPauseAutomaticRegeneration_Clicked();
diff --git a/Engine/Plugins/PCG/Source/PCGEditor/Public/PCGEditorTabFactories.h b/Engine/Plugins/PCG/Source/PCGEditor/Public/PCGEditorTabFactories.h
--- a/Engine/Plugins/PCG/Source/PCGEditor/Public/PCGEditorTabFactories.h
+++ b/Engine/Plugins/PCG/Source/PCGEditor/Public/PCGEditorTabFactories.h
@@ -2,6 +2,7 @@
 
 #pragma once
 
+#include "Delegates/DelegateCombinations.h"
 #include "WorkflowOrientedApp/WorkflowTabFactory.h"
 
 enum ETabRole : uint8;
@@ -11,6 +12,8 @@ class FText;
 class FWorkspaceItem;
 class SWidget;
 
+DECLARE_DELEGATE_OneParam(FPCGGenericTabClosed, FName);
+
 struct FPCGGenericTabFactoryParams
 {
 	FPCGGenericTabFactoryParams(FName InId, TSharedPtr<FPCGEditor> InEditor, TFunction<TSharedRef<SWidget>()> InCreateTabBodyFunc)
@@ -22,6 +25,7 @@ struct FPCGGenericTabFactoryParams
 	FName Id;
 	TSharedPtr<FPCGEditor> Editor;
 	TFunction<TSharedRef<SWidget>()> CreateTabBodyFunc;
+	FPCGGenericTabClosed OnTabClosed;
 	
 	bool bSingleton = true;
 
@@ -37,8 +41,12 @@ struct FPCGGenericTabFactory : public FWorkflowTabFactory
 public:
 	PCGEDITOR_API FPCGGenericTabFactory(const FPCGGenericTabFactoryParams& InParams);
 
+	PCGEDITOR_API virtual TSharedRef<SDockTab> SpawnTab(const FWorkflowTabSpawnInfo& Info) const;
 	PCGEDITOR_API virtual TSharedRef<SWidget> CreateTabBody(const FWorkflowTabSpawnInfo& Info) const override;
 
 private:
+	void OnTabClosed(TSharedRef<SDockTab> DockTab) const;
+
 	TFunction<TSharedRef<SWidget>()> CreateTabBodyFunc;
+	FPCGGenericTabClosed OnTabClosedDelegate;
 };
\ No newline at end of file
