[PCG] Fix crash when toggling PCG Graph Standalone flag

#rnx

#Synced-CL 50158750

diff --git a/Engine/Plugins/PCG/Source/PCG/Private/PCGDefaultExecutionSource.cpp b/Engine/Plugins/PCG/Source/PCG/Private/PCGDefaultExecutionSource.cpp
--- a/Engine/Plugins/PCG/Source/PCG/Private/PCGDefaultExecutionSource.cpp
+++ b/Engine/Plugins/PCG/Source/PCG/Private/PCGDefaultExecutionSource.cpp
@@ -198,13 +198,20 @@ void UPCGDefaultExecutionSource::Generate()
 		if (GenerationTaskId != InvalidPCGTaskId)
 		{
 			FPCGScheduleGenericParams Params{
-				/*InOperation=*/ [this, WeakSubsystem = TWeakInterfacePtr(Subsystem)](FPCGContext*) -> bool
+				/*InOperation=*/ [WeakSubsystem = TWeakInterfacePtr(Subsystem)](FPCGContext* InContext) -> bool
 				{
-					CurrentGenerationTask = InvalidPCGTaskId;
+					// Source was deleted
+					UPCGDefaultExecutionSource* This = Cast<UPCGDefaultExecutionSource>(InContext->ExecutionSource.Get());
+					if (!This)
+					{
+						return true;
+					}
+
+					This->CurrentGenerationTask = InvalidPCGTaskId;
 					if (WeakSubsystem.IsValid())
 					{
 #if WITH_EDITOR
-						WeakSubsystem->OnPCGSourceGenerationDone(this, EPCGGenerationStatus::Completed);
+						WeakSubsystem->OnPCGSourceGenerationDone(This, EPCGGenerationStatus::Completed);
 #endif // WITH_EDITOR
 					}
 					
