[PCG] Add length to the Data domain view for splines

#rnx

#Synced-CL 50157473

diff --git a/Engine/Plugins/PCG/Source/PCGEditor/Private/DataVisualizations/PCGSplineDataVisualization.cpp b/Engine/Plugins/PCG/Source/PCGEditor/Private/DataVisualizations/PCGSplineDataVisualization.cpp
--- a/Engine/Plugins/PCG/Source/PCGEditor/Private/DataVisualizations/PCGSplineDataVisualization.cpp
+++ b/Engine/Plugins/PCG/Source/PCGEditor/Private/DataVisualizations/PCGSplineDataVisualization.cpp
@@ -38,6 +38,7 @@ FPCGTableVisualizerInfo IPCGSplineDataVisualization::GetTableVisualizerInfoWithD
 		// Columns
 		AddPropertyEnumColumnInfo<FTransform>(Info, SplineData, EPCGSplineDataProperties::SplineTransform);
 		AddPropertyEnumColumnInfo<bool>(Info, SplineData, EPCGSplineDataProperties::IsClosed);
+		AddPropertyEnumColumnInfo<double>(Info, SplineData, EPCGSplineDataProperties::Length);
 		
 		// Add Metadata Columns
 		CreateMetadataColumnInfos(SplineData, Info, PCGMetadataDomainID::Data);
