[PCG] Add GC scope guard to prevent crashes when resolving soft object path settings while GC is in progress.

#rnx

#Synced-CL 50157688

diff --git a/Engine/Plugins/PCG/Source/PCG/Private/PCGContext.cpp b/Engine/Plugins/PCG/Source/PCG/Private/PCGContext.cpp
--- a/Engine/Plugins/PCG/Source/PCG/Private/PCGContext.cpp
+++ b/Engine/Plugins/PCG/Source/PCG/Private/PCGContext.cpp
@@ -499,6 +499,9 @@ PRAGMA_ENABLE_DEPRECATION_WARNINGS
 				return false;
 			}
 
+			// Resolving/setting properties (ex: FSoftObjectPath) can end up doing StaticFindObject which needs to be protected from running at same time as GC
+			FGCScopeGuard GCScope;
+
 			// Add some validation on Object Ptr overrides, to make sure that the object/class stored there has the right class
 			if constexpr (std::is_base_of_v<FSoftObjectPath, PropertyType>)
 			{
@@ -548,8 +551,6 @@ PRAGMA_ENABLE_DEPRECATION_WARNINGS
 				}
 			}
 
-			// Setting properties (ex: FSoftObjectPath) can end up doing StaticFindObject which needs to be protected from running at same time as GC
-			FGCScopeGuard GCScope;
 			PropertyAccessor->SetRange<PropertyType>(Values, 0, PropertyObjectKey);
 
 			return true;
