[PCG] Fixed spurious warning that recommended providing explicit bounds when they were already being provided.

In this case the sampling failed because the volume was too thin.

#rnx

#Synced-CL 50157688

diff --git a/Engine/Plugins/PCG/Source/PCG/Private/Elements/PCGSurfaceSampler.cpp b/Engine/Plugins/PCG/Source/PCG/Private/Elements/PCGSurfaceSampler.cpp
--- a/Engine/Plugins/PCG/Source/PCG/Private/Elements/PCGSurfaceSampler.cpp
+++ b/Engine/Plugins/PCG/Source/PCG/Private/Elements/PCGSurfaceSampler.cpp
@@ -626,7 +626,7 @@ PRAGMA_ENABLE_DEPRECATION_WARNINGS
 
 			if (!OutState.SamplerData.Initialize(Settings, Context, EffectiveGridBounds, GeneratingShape->GetTransform()))
 			{
-				if (!GeneratingShape->IsBounded())
+				if (!GeneratingShape->IsBounded() && !ExecState.BoundingShapeBounds.IsValid)
 				{
 					// Some inputs are unable to provide bounds, like the WorldRayHit, in which case the user must provide bounds.
 					PCGE_LOG_C(Warning, GraphAndLog, Context, LOCTEXT("CouldNotObtainInputBounds", "Input data is not bounded, so bounds must be provided for sampling. Consider providing a Bounding Shape input."));
