[PCG] Remove unnecessary soft object path

#rnx

#Synced-CL 50157688

diff --git a/Engine/Plugins/PCG/Source/PCG/Private/Compute/DataInterfaces/PCGComputeDataInterface.cpp b/Engine/Plugins/PCG/Source/PCG/Private/Compute/DataInterfaces/PCGComputeDataInterface.cpp
--- a/Engine/Plugins/PCG/Source/PCG/Private/Compute/DataInterfaces/PCGComputeDataInterface.cpp
+++ b/Engine/Plugins/PCG/Source/PCG/Private/Compute/DataInterfaces/PCGComputeDataInterface.cpp
@@ -68,8 +68,7 @@ void UPCGComputeDataProvider::Initialize(const UComputeDataInterface* InDataInte
 	++GenerationCounter;
 
 	const UPCGComputeDataInterface* DataInterface = CastChecked<UPCGComputeDataInterface>(InDataInterface);
-	ResolvedProducerSettings = DataInterface->GetProducerSettings();
-	ProducerSettings = ResolvedProducerSettings;
+	ProducerSettings = DataInterface->GetProducerSettings();
 	bProducedByCPU = DataInterface->IsProducedByCPU();
 	GraphBindingIndex = DataInterface->GetGraphBindingIndex();
 	ProducerKernel = DataInterface->GetProducerKernel();
@@ -95,7 +94,6 @@ void UPCGComputeDataProvider::Reset()
 
 	ProducerKernel = nullptr;
 	ProducerSettings = nullptr;
-	ResolvedProducerSettings = nullptr;
 	GraphBindingIndex = INDEX_NONE;
 	OutputPinLabel = NAME_None;
 	OutputPinLabelAlias = NAME_None;
@@ -105,27 +103,12 @@ void UPCGComputeDataProvider::Reset()
 
 const UPCGSettings* UPCGComputeDataProvider::GetProducerSettings() const
 {
-	if (!ResolvedProducerSettings)
-	{
-		FGCScopeGuard Guard;
-		ResolvedProducerSettings = ProducerSettings.Get();
-	}
-
-	return ResolvedProducerSettings;
+	return ProducerSettings;
 }
 
 void UPCGComputeDataProvider::SetProducerSettings(const UPCGSettings* InSettings)
 {
 	ProducerSettings = InSettings;
-	ResolvedProducerSettings = InSettings;
-}
-
-void UPCGComputeDataProvider::AddReferencedObjects(UObject* InThis, FReferenceCollector& Collector)
-{
-	Super::AddReferencedObjects(InThis, Collector);
-
-	UPCGComputeDataProvider* This = CastChecked<UPCGComputeDataProvider>(InThis);
-	Collector.AddReferencedObject(This->ResolvedProducerSettings);
 }
 
 #if WITH_EDITOR
diff --git a/Engine/Plugins/PCG/Source/PCG/Public/Compute/DataInterfaces/PCGComputeDataInterface.h b/Engine/Plugins/PCG/Source/PCG/Public/Compute/DataInterfaces/PCGComputeDataInterface.h
--- a/Engine/Plugins/PCG/Source/PCG/Public/Compute/DataInterfaces/PCGComputeDataInterface.h
+++ b/Engine/Plugins/PCG/Source/PCG/Public/Compute/DataInterfaces/PCGComputeDataInterface.h
@@ -107,8 +107,6 @@ public:
 	const TArray<FName>& GetDownstreamInputPinLabelAliases() const { return DownstreamInputPinLabelAliases; }
 	int64 GetGenerationCounter() const { return GenerationCounter; }
 
-	PCG_API static void AddReferencedObjects(UObject* InThis, FReferenceCollector& Collector);
-
 protected:
 #if WITH_EDITOR
 	PCG_API void NotifyProducerUploadedData(UPCGDataBinding* InBinding);
@@ -121,14 +119,12 @@ private:
 	TArray<FName> DownstreamInputPinLabelAliases;
 	bool bProducedByCPU = false;
 
-	UPROPERTY()
+	UPROPERTY(Transient)
 	TObjectPtr<const UPCGComputeKernel> ProducerKernel;
 
 	/** Settings of node that produces this data, normally the upstream node. */
-	UPROPERTY()
-	TSoftObjectPtr<const UPCGSettings> ProducerSettings;
-
-	mutable TObjectPtr<const UPCGSettings> ResolvedProducerSettings = nullptr;
+	UPROPERTY(Transient)
+	TObjectPtr<const UPCGSettings> ProducerSettings;
 
 	/** Bumped each time the data provider is initialized or reset, so that async callbacks can detect if they originated from a previous usage of the data provider and no-op. */
 	std::atomic<uint64> GenerationCounter = 0;
