[PCG] Demote Spline Sampler on Interior intersection check from Error to Warning

* Since it is not an issue per-se as the sampler is still able to produce results, it seems overkill to have an error here.

#rnx

#Synced-CL 50157465

diff --git a/Engine/Plugins/PCG/Source/PCG/Private/Elements/PCGSplineSampler.cpp b/Engine/Plugins/PCG/Source/PCG/Private/Elements/PCGSplineSampler.cpp
--- a/Engine/Plugins/PCG/Source/PCG/Private/Elements/PCGSplineSampler.cpp
+++ b/Engine/Plugins/PCG/Source/PCG/Private/Elements/PCGSplineSampler.cpp
@@ -857,7 +857,7 @@ namespace PCGSplineSamplerHelpers
 		const FPCGSplineStruct* Spline = nullptr;
 
 		// Add more information in order to detect which spline it might be
-		auto LogError = [Context, LineData](const FText& ErrorMessage)
+		auto EmitLog = [Context, LineData](const FText& ErrorMessage, bool bError = true)
 		{
 			const UPCGMetadata* Metadata = LineData ? LineData->ConstMetadata() : nullptr;
 			const FPCGMetadataDomain* MetadataDomain = Metadata ? Metadata->GetConstMetadataDomain(PCGMetadataDomainID::Data) : nullptr;
@@ -865,8 +865,16 @@ namespace PCGSplineSamplerHelpers
 			const FText ActorReferenceHint = ActorRefefrenceAttribute ? FText::Format(LOCTEXT("ActorReferenceHint", "\nPossible actor source: {0}"), FText::FromString(ActorRefefrenceAttribute->GetValue(PCGFirstEntryKey).ToString())) : FText{};
 
 			const FText SplineTransformHint = LineData ? FText::Format(LOCTEXT("TransfromHint", "\nSpline world position: {0}"), FText::FromString(LineData->GetTransform().GetLocation().ToString())) : FText{};
+			const FText FinalMessage = FText::Format(INVTEXT("{0}{1}{2}"), ErrorMessage, ActorReferenceHint, SplineTransformHint);
 
-			PCGLog::LogErrorOnGraph(FText::Format(INVTEXT("{0}{1}{2}"), ErrorMessage, ActorReferenceHint, SplineTransformHint), Context);
+			if (bError)
+			{
+				PCGLog::LogErrorOnGraph(FinalMessage, Context);
+			}
+			else
+			{
+				PCGLog::LogWarningOnGraph(FinalMessage, Context);
+			}
 		};
 
 		if (const UPCGSplineData* SplineData = Cast<UPCGSplineData>(LineData))
@@ -875,12 +883,12 @@ namespace PCGSplineSamplerHelpers
 		}
 		else if (const UPCGLandscapeSplineData* LandscapeSplineData = Cast<UPCGLandscapeSplineData>(LineData))
 		{
-			LogError(LOCTEXT("LandscapeSplinesNotSupported", "Input data of type Landscape Spline are not supported for interior sampling"));
+			EmitLog(LOCTEXT("LandscapeSplinesNotSupported", "Input data of type Landscape Spline are not supported for interior sampling"));
 			return;
 		}
 		else
 		{
-			LogError(LOCTEXT("CouldNotCreateSplineData", "Could not create UPCGSplineData from LineData"));
+			EmitLog(LOCTEXT("CouldNotCreateSplineData", "Could not create UPCGSplineData from LineData"));
 			return;
 		}
 
@@ -888,7 +896,7 @@ namespace PCGSplineSamplerHelpers
 
 		if (!Spline->IsClosedLoop())
 		{
-			LogError(LOCTEXT("ShapeNotClosed", "Interior sampling only generates for closed shapes, enable the 'Closed Loop' setting on the spline"));
+			EmitLog(LOCTEXT("ShapeNotClosed", "Interior sampling only generates for closed shapes, enable the 'Closed Loop' setting on the spline"));
 			return;
 		}
 
@@ -1242,7 +1250,7 @@ namespace PCGSplineSamplerHelpers
 
 		if (bAnyRowFailedToSample)
 		{
-			LogError(LOCTEXT("IntersectionTestFailed", "One or more rows of the spline interior failed to sample (intersection test failed). Ensure the spline points are not overlapping."));
+			EmitLog(LOCTEXT("IntersectionTestFailed", "One or more rows of the spline interior failed to sample (intersection test failed). Ensure the spline points are not overlapping."), /*bError=*/false);
 		}
 
 		// Finally, gather the data and push to the points
