[PCG] Fixed bug preventing inline constant changes from regenerating the graph automatically

#rnx

#Synced-CL 50160642

diff --git a/Engine/Plugins/PCG/Source/PCG/Private/Elements/Metadata/PCGMetadataOpElementBase.cpp b/Engine/Plugins/PCG/Source/PCG/Private/Elements/Metadata/PCGMetadataOpElementBase.cpp
--- a/Engine/Plugins/PCG/Source/PCG/Private/Elements/Metadata/PCGMetadataOpElementBase.cpp
+++ b/Engine/Plugins/PCG/Source/PCG/Private/Elements/Metadata/PCGMetadataOpElementBase.cpp
@@ -294,7 +294,7 @@ void UPCGMetadataSettingsBase::SetPinDefaultValue(const FName PinLabel, const FS
 
 		if (DefaultValues.SetPropertyValueFromString(PropertyName, DefaultValue))
 		{
-			OnSettingsChangedDelegate.Broadcast(this, EPCGChangeType::Node | EPCGChangeType::Edge);
+			OnSettingsChangedDelegate.Broadcast(this, EPCGChangeType::Node | EPCGChangeType::Settings);
 		}
 	}
 }
@@ -308,7 +308,7 @@ void UPCGMetadataSettingsBase::ConvertPinDefaultValueMetadataType(const FName Pi
 		{
 			Modify();
 			DefaultValues.ConvertPropertyType(PropertyName, DataType);
-			OnSettingsChangedDelegate.Broadcast(this, EPCGChangeType::Node | EPCGChangeType::Edge);
+			OnSettingsChangedDelegate.Broadcast(this, EPCGChangeType::Node | EPCGChangeType::Edge | EPCGChangeType::Settings);
 		}
 	}
 }
@@ -326,7 +326,7 @@ void UPCGMetadataSettingsBase::SetPinDefaultValueIsActivated(const FName PinLabe
 		const bool bPropertyChanged = DefaultValues.SetPropertyActivated(PropertyName, bIsActivated);
 		if (bPropertyChanged && bDirtySettings)
 		{
-			OnSettingsChangedDelegate.Broadcast(this, EPCGChangeType::Node | EPCGChangeType::Edge);
+			OnSettingsChangedDelegate.Broadcast(this, EPCGChangeType::Node | EPCGChangeType::Edge | EPCGChangeType::Settings);
 		}
 	}
 }
