[PCG] Support PerQualityLevelFloat for HiGen generation radii

#Synced-CL 50112464

diff --git a/Engine/Plugins/PCG/Source/PCG/Private/PCGCommon.cpp b/Engine/Plugins/PCG/Source/PCG/Private/PCGCommon.cpp
--- a/Engine/Plugins/PCG/Source/PCG/Private/PCGCommon.cpp
+++ b/Engine/Plugins/PCG/Source/PCG/Private/PCGCommon.cpp
@@ -207,30 +207,32 @@ namespace PCGHiGenGrid
 
 double FPCGRuntimeGenerationRadii::GetGenerationRadiusFromGrid(EPCGHiGenGrid Grid) const
 {
+	const int32 QualityLevel = UPCGSubsystem::GetPCGQualityLevel();
+
 	switch (Grid)
 	{
-		case EPCGHiGenGrid::Grid4: return GenerationRadius400;
-		case EPCGHiGenGrid::Grid8: return GenerationRadius800;
-		case EPCGHiGenGrid::Grid16: return GenerationRadius1600;
-		case EPCGHiGenGrid::Grid32: return GenerationRadius3200;
-		case EPCGHiGenGrid::Grid64: return GenerationRadius6400;
-		case EPCGHiGenGrid::Grid128: return GenerationRadius12800;
-		case EPCGHiGenGrid::Grid256: return GenerationRadius25600;
-		case EPCGHiGenGrid::Grid512: return GenerationRadius51200;
-		case EPCGHiGenGrid::Grid1024: return GenerationRadius102400;
-		case EPCGHiGenGrid::Grid2048: return GenerationRadius204800;
-		case EPCGHiGenGrid::Grid4096: return GenerationRadius204800 * (1 << 1);
-		case EPCGHiGenGrid::Grid8192: return GenerationRadius204800 * (1 << 2);
-		case EPCGHiGenGrid::Grid16384: return GenerationRadius204800 * (1 << 3);
-		case EPCGHiGenGrid::Grid32768: return GenerationRadius204800 * (1 << 4);
-		case EPCGHiGenGrid::Grid65536: return GenerationRadius204800 * (1 << 5);
-		case EPCGHiGenGrid::Grid131072: return GenerationRadius204800 * (1 << 6);
-		case EPCGHiGenGrid::Grid262144: return GenerationRadius204800 * (1 << 7);
-		case EPCGHiGenGrid::Grid524288: return GenerationRadius204800 * (1 << 8);
-		case EPCGHiGenGrid::Grid1048576: return GenerationRadius204800 * (1 << 9);
-		case EPCGHiGenGrid::Grid2097152: return GenerationRadius204800 * (1 << 10);
-		case EPCGHiGenGrid::Grid4194304: return GenerationRadius204800 * (1 << 11);
-		case EPCGHiGenGrid::Unbounded: return GenerationRadius;
+		case EPCGHiGenGrid::Grid4: return Radius400.GetValue(QualityLevel);
+		case EPCGHiGenGrid::Grid8: return Radius800.GetValue(QualityLevel);
+		case EPCGHiGenGrid::Grid16: return Radius1600.GetValue(QualityLevel);
+		case EPCGHiGenGrid::Grid32: return Radius3200.GetValue(QualityLevel);
+		case EPCGHiGenGrid::Grid64: return Radius6400.GetValue(QualityLevel);
+		case EPCGHiGenGrid::Grid128: return Radius12800.GetValue(QualityLevel);
+		case EPCGHiGenGrid::Grid256: return Radius25600.GetValue(QualityLevel);
+		case EPCGHiGenGrid::Grid512: return Radius51200.GetValue(QualityLevel);
+		case EPCGHiGenGrid::Grid1024: return Radius102400.GetValue(QualityLevel);
+		case EPCGHiGenGrid::Grid2048: return Radius204800.GetValue(QualityLevel);
+		case EPCGHiGenGrid::Grid4096: return Radius204800.GetValue(QualityLevel) * (1 << 1);
+		case EPCGHiGenGrid::Grid8192: return Radius204800.GetValue(QualityLevel) * (1 << 2);
+		case EPCGHiGenGrid::Grid16384: return Radius204800.GetValue(QualityLevel) * (1 << 3);
+		case EPCGHiGenGrid::Grid32768: return Radius204800.GetValue(QualityLevel) * (1 << 4);
+		case EPCGHiGenGrid::Grid65536: return Radius204800.GetValue(QualityLevel) * (1 << 5);
+		case EPCGHiGenGrid::Grid131072: return Radius204800.GetValue(QualityLevel) * (1 << 6);
+		case EPCGHiGenGrid::Grid262144: return Radius204800.GetValue(QualityLevel) * (1 << 7);
+		case EPCGHiGenGrid::Grid524288: return Radius204800.GetValue(QualityLevel) * (1 << 8);
+		case EPCGHiGenGrid::Grid1048576: return Radius204800.GetValue(QualityLevel) * (1 << 9);
+		case EPCGHiGenGrid::Grid2097152: return Radius204800.GetValue(QualityLevel) * (1 << 10);
+		case EPCGHiGenGrid::Grid4194304: return Radius204800.GetValue(QualityLevel) * (1 << 11);
+		case EPCGHiGenGrid::Unbounded: return RadiusUnbounded.GetValue(QualityLevel);
 	}
 
 	ensure(false);
@@ -239,41 +241,160 @@ double FPCGRuntimeGenerationRadii::GetGenerationRadiusFromGrid(EPCGHiGenGrid Gri
 
 double FPCGRuntimeGenerationRadii::GetCleanupRadiusFromGrid(EPCGHiGenGrid Grid) const
 {
-	return GetGenerationRadiusFromGrid(Grid) * CleanupRadiusMultiplier;
+	return GetGenerationRadiusFromGrid(Grid) * CleanupRadiusScalar.GetValue(UPCGSubsystem::GetPCGQualityLevel());
 }
 
 bool FPCGRuntimeGenerationRadii::operator==(const FPCGRuntimeGenerationRadii& Other) const
 {
-	return GenerationRadius == Other.GenerationRadius
-		&& GenerationRadius400 == Other.GenerationRadius400
-		&& GenerationRadius800 == Other.GenerationRadius800
-		&& GenerationRadius1600 == Other.GenerationRadius1600
-		&& GenerationRadius3200 == Other.GenerationRadius3200
-		&& GenerationRadius6400 == Other.GenerationRadius6400
-		&& GenerationRadius12800 == Other.GenerationRadius12800
-		&& GenerationRadius25600 == Other.GenerationRadius25600
-		&& GenerationRadius51200 == Other.GenerationRadius51200
-		&& GenerationRadius102400 == Other.GenerationRadius102400
-		&& GenerationRadius204800 == Other.GenerationRadius204800
-		&& CleanupRadiusMultiplier == Other.CleanupRadiusMultiplier;
+	auto ArePerQualityFloatsEqual = [](const FPerQualityLevelFloat& A, const FPerQualityLevelFloat& B)
+	{
+		if (A.GetDefault() != B.GetDefault())
+		{
+			return false;
+		}
+
+		const TMap<int32, float>& QualityValuesA = A.PerQuality;
+		const TMap<int32, float>& QualityValuesB = B.PerQuality;
+
+		if (QualityValuesA.Num() != QualityValuesB.Num())
+		{
+			return false;
+		}
+		
+		for (const auto& [QualityLevel, ValueA] : QualityValuesA)
+		{
+			const float* ValueB = QualityValuesB.Find(QualityLevel);
+
+			if (!ValueB || !FMath::IsNearlyEqual(ValueA, *ValueB))
+			{
+				return false;
+			}
+		}
+
+		return true;
+	};
+
+	return ArePerQualityFloatsEqual(RadiusUnbounded, Other.RadiusUnbounded)
+		&& ArePerQualityFloatsEqual(Radius400, Other.Radius400)
+		&& ArePerQualityFloatsEqual(Radius800, Other.Radius800)
+		&& ArePerQualityFloatsEqual(Radius1600, Other.Radius1600)
+		&& ArePerQualityFloatsEqual(Radius3200, Other.Radius3200)
+		&& ArePerQualityFloatsEqual(Radius6400, Other.Radius6400)
+		&& ArePerQualityFloatsEqual(Radius12800, Other.Radius12800)
+		&& ArePerQualityFloatsEqual(Radius25600, Other.Radius25600)
+		&& ArePerQualityFloatsEqual(Radius51200, Other.Radius51200)
+		&& ArePerQualityFloatsEqual(Radius102400, Other.Radius102400)
+		&& ArePerQualityFloatsEqual(Radius204800, Other.Radius204800)
+		&& ArePerQualityFloatsEqual(CleanupRadiusScalar, Other.CleanupRadiusScalar);
+}
+
+#if WITH_EDITOR
+void FPCGRuntimeGenerationRadii::ApplyDeprecation()
+{
+	if (GenerationRadius != 0.0)
+	{
+		RadiusUnbounded = GenerationRadius;
+		GenerationRadius = 0.0;
+	}
+
+	if (GenerationRadius400 != 0.0)
+	{
+		Radius400 = GenerationRadius400;
+		GenerationRadius400 = 0.0;
+	}
+
+	if (GenerationRadius800 != 0.0)
+	{
+		Radius800 = GenerationRadius800;
+		GenerationRadius800 = 0.0;
+	}
+
+	if (GenerationRadius1600 != 0.0)
+	{
+		Radius1600 = GenerationRadius1600;
+		GenerationRadius1600 = 0.0;
+	}
+
+	if (GenerationRadius3200 != 0.0)
+	{
+		Radius3200 = GenerationRadius3200;
+		GenerationRadius3200 = 0.0;
+	}
+
+	if (GenerationRadius6400 != 0.0)
+	{
+		Radius6400 = GenerationRadius6400;
+		GenerationRadius6400 = 0.0;
+	}
+
+	if (GenerationRadius12800 != 0.0)
+	{
+		Radius12800 = GenerationRadius12800;
+		GenerationRadius12800 = 0.0;
+	}
+
+	if (GenerationRadius25600 != 0.0)
+	{
+		Radius25600 = GenerationRadius25600;
+		GenerationRadius25600 = 0.0;
+	}
+
+	if (GenerationRadius51200 != 0.0)
+	{
+		Radius51200 = GenerationRadius51200;
+		GenerationRadius51200 = 0.0;
+	}
+
+	if (GenerationRadius102400 != 0.0)
+	{
+		Radius102400 = GenerationRadius102400;
+		GenerationRadius102400 = 0.0;
+	}
+
+	if (GenerationRadius204800 != 0.0)
+	{
+		Radius204800 = GenerationRadius204800;
+		GenerationRadius204800 = 0.0;
+	}
+
+	if (CleanupRadiusMultiplier != 0.0)
+	{
+		CleanupRadiusScalar = CleanupRadiusMultiplier;
+		CleanupRadiusMultiplier = 0.0;
+	}
 }
+#endif
 
 uint32 GetTypeHash(const FPCGRuntimeGenerationRadii& InGenerationRadii)
 {
 	uint32 Hash = 0;
 
-	Hash = HashCombine(Hash, GetTypeHash(InGenerationRadii.GenerationRadius));
-	Hash = HashCombine(Hash, GetTypeHash(InGenerationRadii.GenerationRadius400));
-	Hash = HashCombine(Hash, GetTypeHash(InGenerationRadii.GenerationRadius800));
-	Hash = HashCombine(Hash, GetTypeHash(InGenerationRadii.GenerationRadius1600));
-	Hash = HashCombine(Hash, GetTypeHash(InGenerationRadii.GenerationRadius3200));
-	Hash = HashCombine(Hash, GetTypeHash(InGenerationRadii.GenerationRadius6400));
-	Hash = HashCombine(Hash, GetTypeHash(InGenerationRadii.GenerationRadius12800));
-	Hash = HashCombine(Hash, GetTypeHash(InGenerationRadii.GenerationRadius25600));
-	Hash = HashCombine(Hash, GetTypeHash(InGenerationRadii.GenerationRadius51200));
-	Hash = HashCombine(Hash, GetTypeHash(InGenerationRadii.GenerationRadius102400));
-	Hash = HashCombine(Hash, GetTypeHash(InGenerationRadii.GenerationRadius204800));
-	Hash = HashCombine(Hash, GetTypeHash(InGenerationRadii.CleanupRadiusMultiplier));
+	auto HashQualityFloat = [](const FPerQualityLevelFloat& InQualityFloat)
+	{
+		uint32 Hash = 0;
+
+		for (const auto& [QualityLevel, Value] : InQualityFloat.PerQuality)
+		{
+			Hash = HashCombine(Hash, HashCombine(QualityLevel, Value));
+		}
+
+		Hash = HashCombine(Hash, InQualityFloat.GetDefault());
+
+		return Hash;
+	};
+
+	Hash = HashCombine(Hash, HashQualityFloat(InGenerationRadii.RadiusUnbounded));
+	Hash = HashCombine(Hash, HashQualityFloat(InGenerationRadii.Radius400));
+	Hash = HashCombine(Hash, HashQualityFloat(InGenerationRadii.Radius800));
+	Hash = HashCombine(Hash, HashQualityFloat(InGenerationRadii.Radius1600));
+	Hash = HashCombine(Hash, HashQualityFloat(InGenerationRadii.Radius3200));
+	Hash = HashCombine(Hash, HashQualityFloat(InGenerationRadii.Radius6400));
+	Hash = HashCombine(Hash, HashQualityFloat(InGenerationRadii.Radius12800));
+	Hash = HashCombine(Hash, HashQualityFloat(InGenerationRadii.Radius25600));
+	Hash = HashCombine(Hash, HashQualityFloat(InGenerationRadii.Radius51200));
+	Hash = HashCombine(Hash, HashQualityFloat(InGenerationRadii.Radius102400));
+	Hash = HashCombine(Hash, HashQualityFloat(InGenerationRadii.Radius204800));
+	Hash = HashCombine(Hash, HashQualityFloat(InGenerationRadii.CleanupRadiusScalar));
 
 	return Hash;
 }
diff --git a/Engine/Plugins/PCG/Source/PCG/Private/PCGComponent.cpp b/Engine/Plugins/PCG/Source/PCG/Private/PCGComponent.cpp
--- a/Engine/Plugins/PCG/Source/PCG/Private/PCGComponent.cpp
+++ b/Engine/Plugins/PCG/Source/PCG/Private/PCGComponent.cpp
@@ -1955,6 +1955,8 @@ PRAGMA_ENABLE_DEPRECATION_WARNINGS
 		GraphInstance->SetGraph(Graph_DEPRECATED);
 		Graph_DEPRECATED = nullptr;
 	}
+	
+	GenerationRadii.ApplyDeprecation();
 
 	SetupCallbacksOnCreation();
 
diff --git a/Engine/Plugins/PCG/Source/PCG/Private/PCGGraph.cpp b/Engine/Plugins/PCG/Source/PCG/Private/PCGGraph.cpp
--- a/Engine/Plugins/PCG/Source/PCG/Private/PCGGraph.cpp
+++ b/Engine/Plugins/PCG/Source/PCG/Private/PCGGraph.cpp
@@ -624,6 +624,8 @@ void UPCGGraph::PostLoad()
 		FixInvalidEdges();
 	}
 
+	GenerationRadii.ApplyDeprecation();
+
 	GraphCustomization.Graph = this;
 #endif // WITH_EDITOR
 }
@@ -1774,8 +1776,9 @@ bool UPCGGraph::Recompile()
 
 void UPCGGraph::OnPCGQualityLevelChanged()
 {
-	bool bContainsQualityControlFlowNode = false;
+	// No need to broadcast a refresh here since relevant execution sources will be refreshed from UPCGSubsystem::OnPCGQualityLevelChanged()
 
+	// If we have a quality branch/select node we need to broadcast a cosmetic change to update the pin visualization.
 	for (UPCGNode* Node : Nodes)
 	{
 		if (!Node)
@@ -1788,14 +1791,8 @@ void UPCGGraph::OnPCGQualityLevelChanged()
 		if (Settings && (Settings->IsA<UPCGQualityBranchSettings>() || Settings->IsA<UPCGQualitySelectSettings>()))
 		{
 			Node->OnNodeChangedDelegate.Broadcast(Node, EPCGChangeType::Cosmetic);
-			bContainsQualityControlFlowNode = true;
 		}
 	}
-
-	if (bContainsQualityControlFlowNode)
-	{
-		OnGraphChangedDelegate.Broadcast(this, EPCGChangeType::Structural | EPCGChangeType::GenerationGrid);
-	}
 }
 
 FPCGSelectionKeyToSettingsMap UPCGGraph::GetTrackedSelectionKeysToSettings() const
diff --git a/Engine/Plugins/PCG/Source/PCG/Private/Subsystems/PCGSubsystem.cpp b/Engine/Plugins/PCG/Source/PCG/Private/Subsystems/PCGSubsystem.cpp
--- a/Engine/Plugins/PCG/Source/PCG/Private/Subsystems/PCGSubsystem.cpp
+++ b/Engine/Plugins/PCG/Source/PCG/Private/Subsystems/PCGSubsystem.cpp
@@ -471,9 +471,9 @@ void UPCGSubsystem::OnPCGQualityLevelChanged()
 			{
 				Graph->OnPCGQualityLevelChanged();
 			}
-#else
-			RefreshRuntimeGenExecutionSource(InExecutionSource, EPCGChangeType::GenerationGrid);
 #endif
+
+			RefreshRuntimeGenExecutionSource(InExecutionSource, EPCGChangeType::GenerationGrid);
 		}
 	});
 }
diff --git a/Engine/Plugins/PCG/Source/PCG/Public/PCGCommon.h b/Engine/Plugins/PCG/Source/PCG/Public/PCGCommon.h
--- a/Engine/Plugins/PCG/Source/PCG/Public/PCGCommon.h
+++ b/Engine/Plugins/PCG/Source/PCG/Public/PCGCommon.h
@@ -4,6 +4,7 @@
 
 #include "PCGPointPropertiesTraits.h"
 
+#include "PerQualityLevelProperties.h"
 #include "Async/Mutex.h"
 #include "Async/RecursiveMutex.h"
 #include "Async/SharedMutex.h"
@@ -574,49 +575,92 @@ public:
 
 	UE_API bool operator==(const FPCGRuntimeGenerationRadii& Other) const;
 
+#if WITH_EDITOR
+	UE_API void ApplyDeprecation();
+#endif
+
+	friend UE_API uint32 GetTypeHash(const FPCGRuntimeGenerationRadii& In);
+
 	static constexpr double DefaultGenerationRadiusMultiplier = 2.0;
 	static constexpr double DefaultCleanupRadiusMultiplier = 1.1;
 
+protected:
 	/** The distance (in centimeters) at which the component will be considered for generation by the RuntimeGenerationScheduler. For partitioned components, this also acts as the unbounded generation radius. */
 	UPROPERTY(BlueprintReadWrite, EditAnywhere, Category = "Runtime Generation")
-	double GenerationRadius = PCGHiGenGrid::UnboundedGridSize() * DefaultGenerationRadiusMultiplier;
+	FPerQualityLevelFloat RadiusUnbounded = PCGHiGenGrid::UnboundedGridSize() * DefaultGenerationRadiusMultiplier;
 
 	UPROPERTY(BlueprintReadWrite, EditAnywhere, Category = "Runtime Generation")
-	double GenerationRadius400 = PCGHiGenGrid::GridToGridSize(EPCGHiGenGrid::Grid4) * DefaultGenerationRadiusMultiplier;
+	FPerQualityLevelFloat Radius400 = PCGHiGenGrid::GridToGridSize(EPCGHiGenGrid::Grid4) * DefaultGenerationRadiusMultiplier;
 
 	UPROPERTY(BlueprintReadWrite, EditAnywhere, Category = "Runtime Generation")
-	double GenerationRadius800 = PCGHiGenGrid::GridToGridSize(EPCGHiGenGrid::Grid8) * DefaultGenerationRadiusMultiplier;
+	FPerQualityLevelFloat Radius800 = PCGHiGenGrid::GridToGridSize(EPCGHiGenGrid::Grid8) * DefaultGenerationRadiusMultiplier;
 
 	UPROPERTY(BlueprintReadWrite, EditAnywhere, Category = "Runtime Generation")
-	double GenerationRadius1600 = PCGHiGenGrid::GridToGridSize(EPCGHiGenGrid::Grid16) * DefaultGenerationRadiusMultiplier;
+	FPerQualityLevelFloat Radius1600 = PCGHiGenGrid::GridToGridSize(EPCGHiGenGrid::Grid16) * DefaultGenerationRadiusMultiplier;
 
 	UPROPERTY(BlueprintReadWrite, EditAnywhere, Category = "Runtime Generation")
-	double GenerationRadius3200 = PCGHiGenGrid::GridToGridSize(EPCGHiGenGrid::Grid32) * DefaultGenerationRadiusMultiplier;
+	FPerQualityLevelFloat Radius3200 = PCGHiGenGrid::GridToGridSize(EPCGHiGenGrid::Grid32) * DefaultGenerationRadiusMultiplier;
 
 	UPROPERTY(BlueprintReadWrite, EditAnywhere, Category = "Runtime Generation")
-	double GenerationRadius6400 = PCGHiGenGrid::GridToGridSize(EPCGHiGenGrid::Grid64) * DefaultGenerationRadiusMultiplier;
+	FPerQualityLevelFloat Radius6400 = PCGHiGenGrid::GridToGridSize(EPCGHiGenGrid::Grid64) * DefaultGenerationRadiusMultiplier;
 
 	UPROPERTY(BlueprintReadWrite, EditAnywhere, Category = "Runtime Generation")
-	double GenerationRadius12800 = PCGHiGenGrid::GridToGridSize(EPCGHiGenGrid::Grid128) * DefaultGenerationRadiusMultiplier;
+	FPerQualityLevelFloat Radius12800 = PCGHiGenGrid::GridToGridSize(EPCGHiGenGrid::Grid128) * DefaultGenerationRadiusMultiplier;
 
 	UPROPERTY(BlueprintReadWrite, EditAnywhere, Category = "Runtime Generation")
-	double GenerationRadius25600 = PCGHiGenGrid::GridToGridSize(EPCGHiGenGrid::Grid256) * DefaultGenerationRadiusMultiplier;
+	FPerQualityLevelFloat Radius25600 = PCGHiGenGrid::GridToGridSize(EPCGHiGenGrid::Grid256) * DefaultGenerationRadiusMultiplier;
 
 	UPROPERTY(BlueprintReadWrite, EditAnywhere, Category = "Runtime Generation")
-	double GenerationRadius51200 = PCGHiGenGrid::GridToGridSize(EPCGHiGenGrid::Grid512) * DefaultGenerationRadiusMultiplier;
+	FPerQualityLevelFloat Radius51200 = PCGHiGenGrid::GridToGridSize(EPCGHiGenGrid::Grid512) * DefaultGenerationRadiusMultiplier;
 
 	UPROPERTY(BlueprintReadWrite, EditAnywhere, Category = "Runtime Generation")
-	double GenerationRadius102400 = PCGHiGenGrid::GridToGridSize(EPCGHiGenGrid::Grid1024) * DefaultGenerationRadiusMultiplier;
+	FPerQualityLevelFloat Radius102400 = PCGHiGenGrid::GridToGridSize(EPCGHiGenGrid::Grid1024) * DefaultGenerationRadiusMultiplier;
 
 	UPROPERTY(BlueprintReadWrite, EditAnywhere, Category = "Runtime Generation")
-	double GenerationRadius204800 = PCGHiGenGrid::GridToGridSize(EPCGHiGenGrid::Grid2048) * DefaultGenerationRadiusMultiplier;
+	FPerQualityLevelFloat Radius204800 = PCGHiGenGrid::GridToGridSize(EPCGHiGenGrid::Grid2048) * DefaultGenerationRadiusMultiplier;
 
 	/** Multiplier on the GenerationRadius to control the distance at which runtime generated components will be cleaned up. Applied per grid size. */
 	UPROPERTY(BlueprintReadWrite, EditAnywhere, Category = "Runtime Generation", meta = (UIMin = "1.0", ClampMin = "1.0"))
-	double CleanupRadiusMultiplier = DefaultCleanupRadiusMultiplier;
-};
+	FPerQualityLevelFloat CleanupRadiusScalar = DefaultCleanupRadiusMultiplier;
+
+	// Deprecated section
+public:
+	UPROPERTY(meta=(DeprecatedProperty, DeprecationMessage="Use RadiusUnbounded instead"))
+	double GenerationRadius = 0.0;
+
+	UPROPERTY(meta=(DeprecatedProperty, DeprecationMessage="Use Radius400 instead"))
+	double GenerationRadius400 = 0.0;
+
+	UPROPERTY(meta=(DeprecatedProperty, DeprecationMessage="Use Radius800 instead"))
+	double GenerationRadius800 = 0.0;
 
-UE_API uint32 GetTypeHash(const FPCGRuntimeGenerationRadii& InGenerationRadii);
+	UPROPERTY(meta=(DeprecatedProperty, DeprecationMessage="Use Radius1600 instead"))
+	double GenerationRadius1600 = 0.0;
+
+	UPROPERTY(meta=(DeprecatedProperty, DeprecationMessage="Use Radius3200 instead"))
+	double GenerationRadius3200 = 0.0;
+
+	UPROPERTY(meta=(DeprecatedProperty, DeprecationMessage="Use Radius6400 instead"))
+	double GenerationRadius6400 = 0.0;
+
+	UPROPERTY(meta=(DeprecatedProperty, DeprecationMessage="Use Radius12800 instead"))
+	double GenerationRadius12800 = 0.0;
+
+	UPROPERTY(meta=(DeprecatedProperty, DeprecationMessage="Use Radius25600 instead"))
+	double GenerationRadius25600 = 0.0;
+
+	UPROPERTY(meta=(DeprecatedProperty, DeprecationMessage="Use Radius51200 instead"))
+	double GenerationRadius51200 = 0.0;
+
+	UPROPERTY(meta=(DeprecatedProperty, DeprecationMessage="Use Radius102400 instead"))
+	double GenerationRadius102400 = 0.0;
+
+	UPROPERTY(meta=(DeprecatedProperty, DeprecationMessage="Use Radius204800 instead"))
+	double GenerationRadius204800 = 0.0;
+
+	UPROPERTY(meta=(DeprecatedProperty, DeprecationMessage="Use CleanupRadiusScalar instead"))
+	double CleanupRadiusMultiplier = 0.0;
+};
 
 struct FInstancedPropertyBag;
 
