[PCG] Fix Crc issues
* Volume Data now produces a crc even when initialized with an FBox instead of a volume
* Get Actor Data and Get Actor Property now crc's correctly even when the execution source is not a PCG Component

#rnx

#Synced-CL 50112464

diff --git a/Engine/Plugins/PCG/Source/PCG/Private/Data/PCGVolumeData.cpp b/Engine/Plugins/PCG/Source/PCG/Private/Data/PCGVolumeData.cpp
--- a/Engine/Plugins/PCG/Source/PCG/Private/Data/PCGVolumeData.cpp
+++ b/Engine/Plugins/PCG/Source/PCG/Private/Data/PCGVolumeData.cpp
@@ -95,40 +95,35 @@ void UPCGVolumeData::AddToCrc(FArchiveCrc32& Ar, bool bFullDataCrc) const
 {
 	Super::AddToCrc(Ar, bFullDataCrc);
 
-	if (VolumeBodyInstance)
-	{
-		// Implementation note: no metadata in this data at this point.
+	// Implementation note: no metadata in this data at this point.
 
-		FString ClassName = StaticClass()->GetPathName();
-		Ar << ClassName;
+	FString ClassName = StaticClass()->GetPathName();
+	Ar << ClassName;
 
+	if (VolumeBodyInstance)
+	{
 		// Weird implementation to fix ambiguous call between the FBodyInstance friend function (which isn't exposed) and the one in the FArchiveCrc32 code.
 		FBodyInstance::StaticStruct()->SerializeItem(Ar, const_cast<FBodyInstance*>(VolumeBodyInstance), nullptr);
+	}
 
-		// Implementation note: we will not consider the volume pointer in this instance
-		Ar.Serialize((void*)&VoxelSize, sizeof(FVector));
+	// Implementation note: we will not consider the volume pointer in this instance
+	Ar.Serialize((void*)&VoxelSize, sizeof(FVector));
 
-		// TODO: move this to helper function
-		auto SerializeBounds = [&Ar](FBox* Box)
-		{
-			check(Box);
-			Ar << Box->IsValid;
+	// TODO: move this to helper function
+	auto SerializeBounds = [&Ar](FBox* Box)
+	{
+		check(Box);
+		Ar << Box->IsValid;
 
-			if (Box->IsValid)
-			{
-				Ar << Box->Min;
-				Ar << Box->Max;
-			}
-		};
+		if (Box->IsValid)
+		{
+			Ar << Box->Min;
+			Ar << Box->Max;
+		}
+	};
 
-		SerializeBounds(const_cast<FBox*>(&Bounds));
-		SerializeBounds(const_cast<FBox*>(&StrictBounds));
-	}
-	else
-	{
-		// This data does not have a bespoke CRC implementation so just use a global unique data CRC.
-		AddUIDToCrc(Ar);
-	}
+	SerializeBounds(const_cast<FBox*>(&Bounds));
+	SerializeBounds(const_cast<FBox*>(&StrictBounds));
 }
 
 FBox UPCGVolumeData::GetBounds() const
diff --git a/Engine/Plugins/PCG/Source/PCG/Private/Elements/PCGDataFromActor.cpp b/Engine/Plugins/PCG/Source/PCG/Private/Elements/PCGDataFromActor.cpp
--- a/Engine/Plugins/PCG/Source/PCG/Private/Elements/PCGDataFromActor.cpp
+++ b/Engine/Plugins/PCG/Source/PCG/Private/Elements/PCGDataFromActor.cpp
@@ -1136,21 +1136,20 @@ void FPCGDataFromActorElement::GetDependenciesCrc(const FPCGGetDependenciesCrcPa
 		const bool bDependsOnSelfOrHierarchy = (Settings->ActorSelector.ActorFilter == EPCGActorFilter::Self || Settings->ActorSelector.ActorFilter == EPCGActorFilter::Original);
 		const bool bDependsOnSelfBounds = Settings->ActorSelector.bMustOverlapSelf;
 
-		UPCGComponent* PCGComponent = Cast<UPCGComponent>(InParams.ExecutionSource);
-		if (PCGComponent && (bDependsOnSelfOrHierarchy || bDependsOnSelfBounds))
+		if (InParams.ExecutionSource && (bDependsOnSelfOrHierarchy || bDependsOnSelfBounds))
 		{
-			UPCGComponent* ComponentToCheck = (Settings->ActorSelector.ActorFilter == EPCGActorFilter::Original) ? PCGComponent->GetOriginalComponent() : PCGComponent;
-			const UPCGData* ActorData = ComponentToCheck ? ComponentToCheck->GetActorPCGData() : nullptr;
+			IPCGGraphExecutionSource* ExecutionSourceToCheck = (Settings->ActorSelector.ActorFilter == EPCGActorFilter::Original) ? InParams.ExecutionSource->GetExecutionState().GetOriginalSource() : InParams.ExecutionSource;
+			const UPCGData* SelfData = ExecutionSourceToCheck ? ExecutionSourceToCheck->GetExecutionState().GetSelfData() : nullptr;
 
-			if (ActorData)
+			if (SelfData)
 			{
-				Crc.Combine(ActorData->GetOrComputeCrc(/*bFullDataCrc=*/false));
+				Crc.Combine(SelfData->GetOrComputeCrc(/*bFullDataCrc=*/false));
 
 				// Also if it is depending on Self, we need to CRC the tags (as they will be passed to the output data)
 				// Same with anything that would pull the actor reference.
-				if (bDependsOnSelfOrHierarchy && ComponentToCheck)
+				if (bDependsOnSelfOrHierarchy && ExecutionSourceToCheck)
 				{
-					if (const AActor* Owner = ComponentToCheck->GetOwner())
+					if (const AActor* Owner = ExecutionSourceToCheck->GetExecutionState().GetTypedTarget<AActor>())
 					{
 						const bool bOutputActorReference = Settings->bAlsoOutputSinglePointData
 							|| Settings->Mode == EPCGGetDataFromActorMode::GetActorReference
@@ -1164,9 +1163,9 @@ void FPCGDataFromActorElement::GetDependenciesCrc(const FPCGGetDependenciesCrcPa
 
 			// Include isolated actors when relevant.
 			// @todo_pcg : is this a good idea if this isn't on the GT?
-			if (Settings->ActorSelector.IncludeIsolatedActors != EPCGIsolatedActorIncludeFlag::Ignore && ComponentToCheck)
+			if (Settings->ActorSelector.IncludeIsolatedActors != EPCGIsolatedActorIncludeFlag::Ignore && ExecutionSourceToCheck)
 			{
-				TArray<APCGIsolatedActor*> IsolatedActors = UPCGActorHelpers::GetIsolatedActorsOriginatingFrom(ComponentToCheck->GetOwner());
+				TArray<APCGIsolatedActor*> IsolatedActors = UPCGActorHelpers::GetIsolatedActorsOriginatingFrom(ExecutionSourceToCheck->GetExecutionState().GetTypedTarget<AActor>());
 				for (APCGIsolatedActor* IsolatedActor : IsolatedActors)
 				{
 					Crc.Combine(GetTypeHash(IsolatedActor->GetPathName()));
@@ -1177,9 +1176,9 @@ void FPCGDataFromActorElement::GetDependenciesCrc(const FPCGGetDependenciesCrcPa
 		const bool bDependsOnComponentData = Settings->Mode == EPCGGetDataFromActorMode::GetDataFromPCGComponent || Settings->Mode == EPCGGetDataFromActorMode::GetDataFromPCGComponentOrParseComponents;
 		const bool bDependsOnLocalComponentBounds = Settings->bComponentsMustOverlapSelf || !Settings->bGetDataOnAllGrids;
 
-		if (PCGComponent && bDependsOnComponentData && bDependsOnLocalComponentBounds)
+		if (InParams.ExecutionSource && bDependsOnComponentData && bDependsOnLocalComponentBounds)
 		{
-			if (const UPCGData* LocalActorData = PCGComponent->GetActorPCGData())
+			if (const UPCGData* LocalActorData = InParams.ExecutionSource->GetExecutionState().GetSelfData())
 			{
 				Crc.Combine(LocalActorData->GetOrComputeCrc(/*bFullDataCrc=*/false));
 			}
diff --git a/Engine/Plugins/PCG/Source/PCG/Private/Elements/PCGGetActorProperty.cpp b/Engine/Plugins/PCG/Source/PCG/Private/Elements/PCGGetActorProperty.cpp
--- a/Engine/Plugins/PCG/Source/PCG/Private/Elements/PCGGetActorProperty.cpp
+++ b/Engine/Plugins/PCG/Source/PCG/Private/Elements/PCGGetActorProperty.cpp
@@ -188,17 +188,16 @@ void FPCGGetActorPropertyElement::GetDependenciesCrc(const FPCGGetDependenciesCr
 			return;
 		}
 
-		UPCGComponent* PCGComponent = Cast<UPCGComponent>(InParams.ExecutionSource);
-		if (PCGComponent)
+		if (InParams.ExecutionSource)
 		{
 			const bool bDependsOnSelfOrHierarchy = (Settings->ActorSelector.ActorFilter == EPCGActorFilter::Self || Settings->ActorSelector.ActorFilter == EPCGActorFilter::Original);
-			UPCGComponent* ComponentToCheck = (Settings->ActorSelector.ActorFilter == EPCGActorFilter::Original) ? PCGComponent->GetOriginalComponent() : PCGComponent;
+			IPCGGraphExecutionSource* ExecutionSourceToCheck = (Settings->ActorSelector.ActorFilter == EPCGActorFilter::Original) ? InParams.ExecutionSource->GetExecutionState().GetOriginalSource() : InParams.ExecutionSource;
 
 			// If it is depending on Self, we need to CRC the tags (as they will be passed to the output data)
 			// and also the owner path.
-			if (bDependsOnSelfOrHierarchy && ComponentToCheck)
+			if (bDependsOnSelfOrHierarchy && ExecutionSourceToCheck)
 			{
-				if (const AActor* Owner = ComponentToCheck->GetOwner())
+				if (const AActor* Owner = ExecutionSourceToCheck->GetExecutionState().GetTypedTarget<AActor>())
 				{
 					Crc.Combine(UPCGActorHelpers::ComputeHashFromActorTagsAndReference(Owner, /*bIncludeTags=*/true, /*bIncludeActorReference=*/true));
 				}
