Enabling subgraph for PVE
- Exposing subgraph functions for external use.
- Removed the dependency of UPCGSubgraphSettings where possible. Added a new class FPCGBaseSubgraphElement, which contains PCGBaseSubgraphSettings usage. UPCGSubgraphSettings is only used in FPCGSubgraphElement::ExecuteInternal for SubgraphOverride.

#Synced-CL 50245349

diff --git a/Engine/Plugins/PCG/Source/PCG/Private/PCGSubgraph.cpp b/Engine/Plugins/PCG/Source/PCG/Private/PCGSubgraph.cpp
--- a/Engine/Plugins/PCG/Source/PCG/Private/PCGSubgraph.cpp
+++ b/Engine/Plugins/PCG/Source/PCG/Private/PCGSubgraph.cpp
@@ -504,13 +504,13 @@ TObjectPtr<UPCGGraph> UPCGBaseSubgraphNode::GetSubgraph() const
 
 TObjectPtr<UPCGGraphInterface> UPCGSubgraphNode::GetSubgraphInterface() const
 {
-	TObjectPtr<UPCGSubgraphSettings> Settings = Cast<UPCGSubgraphSettings>(GetSettings());
+	TObjectPtr<UPCGBaseSubgraphSettings> Settings = Cast<UPCGBaseSubgraphSettings>(GetSettings());
 	return Settings ? Settings->GetSubgraphInterface() : nullptr;
 }
 
 FStructView FPCGSubgraphContext::GetExternalStructContainerForOverridableParam(const FPCGSettingsOverridableParam& InParam)
 {
-	const UPCGSubgraphSettings* Settings = GetInputSettings<UPCGSubgraphSettings>();
+	const UPCGBaseSubgraphSettings* Settings = GetInputSettings<UPCGBaseSubgraphSettings>();
 	check(Settings);
 	const UPCGGraphInterface* Graph = Settings->GetSubgraphInterface();
 	const FInstancedPropertyBag* UserParameters = Graph ? Graph->GetUserParametersStruct() : nullptr;
@@ -640,7 +640,7 @@ FPCGContext* FPCGSubgraphElement::Initialize(const FPCGInitializeElementParams&
 	return Context;
 }
 
-void FPCGSubgraphElement::PrepareSubgraphData(const UPCGSubgraphSettings* Settings, FPCGSubgraphContext* Context, const FPCGDataCollection& InputData, FPCGDataCollection& OutputData) const
+void FPCGSubgraphElement::PrepareSubgraphData(const UPCGBaseSubgraphSettings* Settings, FPCGSubgraphContext* Context, const FPCGDataCollection& InputData, FPCGDataCollection& OutputData) const
 {
 	// Don't forward overrides
 	if (Settings->HasOverridableParams())
@@ -674,7 +674,7 @@ void FPCGSubgraphElement::PrepareSubgraphData(const UPCGSubgraphSettings* Settin
 	// you should probably do some filtering, like the UPCGUserParametersData.
 }
 
-void FPCGSubgraphElement::PrepareSubgraphUserParameters(const UPCGSubgraphSettings* Settings, FPCGSubgraphContext* Context, FPCGDataCollection& OutputData) const
+void FPCGSubgraphElement::PrepareSubgraphUserParameters(const UPCGBaseSubgraphSettings* Settings, FPCGSubgraphContext* Context, FPCGDataCollection& OutputData) const
 {
 	// Also create a new data containing information about the original subgraph and the parameters override
 	// It is used mainly by the UserParameterGetElement to access the correct value.
@@ -717,7 +717,7 @@ void FPCGSubgraphElement::PrepareSubgraphUserParameters(const UPCGSubgraphSettin
 
 bool FPCGSubgraphElement::IsPassthrough(const UPCGSettings* InSettings) const
 {
-	const UPCGSubgraphSettings* Settings = Cast<UPCGSubgraphSettings>(InSettings);
+	const UPCGBaseSubgraphSettings* Settings = Cast<UPCGBaseSubgraphSettings>(InSettings);
 	return (!Settings || (Settings->bEnabled && Settings->GetSubgraph()));
 }
 
@@ -725,7 +725,7 @@ bool FPCGSubgraphElement::ExecuteInternal(FPCGContext* InContext) const
 {
 	FPCGSubgraphContext* Context = static_cast<FPCGSubgraphContext*>(InContext);
 
-	const UPCGSubgraphSettings* Settings = Context->GetInputSettings<UPCGSubgraphSettings>();
+	const UPCGBaseSubgraphSettings* Settings = Context->GetInputSettings<UPCGBaseSubgraphSettings>();
 	check(Settings);
 
 	bool bIsDynamic = Settings->IsDynamicGraph();
@@ -734,16 +734,16 @@ bool FPCGSubgraphElement::ExecuteInternal(FPCGContext* InContext) const
 	bIsDynamic |= PCGSystemSwitches::CVarForceDynamicGraphDispatch.GetValueOnAnyThread();
 #endif
 
-	if (bIsDynamic && !Context->bScheduledSubgraph && Settings->SubgraphOverride)
+	if (bIsDynamic && !Context->bScheduledSubgraph && Settings->GetSubgraphOverride())
 	{
 		TRACE_CPUPROFILER_EVENT_SCOPE(FPCGSubgraphElement::ExecuteInternal::DynamicTrackingAndOverrides);
 
 #if WITH_EDITOR
-		FPCGDynamicTrackingHelper::AddSingleDynamicTrackingKey(Context, FPCGSelectionKey::CreateFromPath(Settings->SubgraphOverride), /*bIsCulled=*/false);
+		FPCGDynamicTrackingHelper::AddSingleDynamicTrackingKey(Context, FPCGSelectionKey::CreateFromPath(Settings->GetSubgraphOverride()), /*bIsCulled=*/false);
 		// Also need to add statically tracked keys
 		if (InContext && InContext->ExecutionSource.IsValid())
 		{
-			if (const UPCGGraph* Graph = Settings->SubgraphOverride->GetGraph())
+			if (const UPCGGraph* Graph = Settings->GetSubgraphOverride()->GetGraph())
 			{
 				InContext->ExecutionSource->GetExecutionState().RegisterDynamicTracking(Graph->GetTrackedSelectionKeysToSettings());
 			}
diff --git a/Engine/Plugins/PCG/Source/PCG/Public/PCGGraph.h b/Engine/Plugins/PCG/Source/PCG/Public/PCGGraph.h
--- a/Engine/Plugins/PCG/Source/PCG/Public/PCGGraph.h
+++ b/Engine/Plugins/PCG/Source/PCG/Public/PCGGraph.h
@@ -728,7 +728,7 @@ public:
 	UE_API void SetupCallbacks();
 	UE_API void TeardownCallbacks();
 
-	virtual TOptional<FPCGGraphToolData> GetGraphToolData() const override;
+	UE_API virtual TOptional<FPCGGraphToolData> GetGraphToolData() const override;
 #endif
 
 	UE_API virtual void OnGraphParametersChanged(EPCGGraphParameterEvent InChangeType, FName InChangedPropertyName) override;
diff --git a/Engine/Plugins/PCG/Source/PCG/Public/PCGSubgraph.h b/Engine/Plugins/PCG/Source/PCG/Public/PCGSubgraph.h
--- a/Engine/Plugins/PCG/Source/PCG/Public/PCGSubgraph.h
+++ b/Engine/Plugins/PCG/Source/PCG/Public/PCGSubgraph.h
@@ -38,6 +38,8 @@ public:
 	UFUNCTION(BlueprintCallable, Category = Graph)
 	UE_API virtual void SetSubgraph(UPCGGraphInterface* InGraph);
 
+	virtual UPCGGraphInterface* GetSubgraphOverride() const { return nullptr; }
+
 protected:
 	//~Begin UObject interface implementation
 	UE_API virtual void PostLoad() override;
@@ -133,6 +135,7 @@ public:
 	UE_API virtual UPCGGraphInterface* GetSubgraphInterface() const override;
 	UE_API virtual bool IsDynamicGraph() const override;
 	virtual UPCGGraphInterface* GetOriginalSubgraphInterface() const { return SubgraphInstance; }
+	virtual UPCGGraphInterface* GetSubgraphOverride() const override { return SubgraphOverride; }
 
 	/** Used to filter the subgraph list based on the graph configuration. */
 	UFUNCTION()
@@ -201,14 +204,14 @@ protected:
 class FPCGSubgraphElement : public IPCGElement
 {
 public:
-	virtual FPCGContext* Initialize(const FPCGInitializeElementParams& InParams) override;
+	UE_API virtual FPCGContext* Initialize(const FPCGInitializeElementParams& InParams) override;
 	virtual bool IsCacheable(const UPCGSettings* InSettings) const override { return false; }
 
 protected:
-	virtual bool ExecuteInternal(FPCGContext* Context) const override;
-	virtual bool IsPassthrough(const UPCGSettings* InSettings) const override;
-	void PrepareSubgraphData(const UPCGSubgraphSettings* Settings, FPCGSubgraphContext* Context, const FPCGDataCollection& InputData, FPCGDataCollection& OutputData) const;
-	void PrepareSubgraphUserParameters(const UPCGSubgraphSettings* Settings, FPCGSubgraphContext* Context, FPCGDataCollection& OutputData) const;
+	UE_API virtual bool ExecuteInternal(FPCGContext* Context) const override;
+	UE_API virtual bool IsPassthrough(const UPCGSettings* InSettings) const override;
+	void PrepareSubgraphData(const UPCGBaseSubgraphSettings* Settings, FPCGSubgraphContext* Context, const FPCGDataCollection& InputData, FPCGDataCollection& OutputData) const;
+	void PrepareSubgraphUserParameters(const UPCGBaseSubgraphSettings* Settings, FPCGSubgraphContext* Context, FPCGDataCollection& OutputData) const;
 	virtual bool SupportsGPUResidentData(FPCGContext* InContext) const override { return true; }
 	virtual bool SupportsBasePointDataInputs(FPCGContext* InContext) const { return true; }
 };
