[PCG] PCG.Cache.ResetStats/PCG.Cache.LogStats cvar commands
- Prints Stores/Hits/Gets/Grow count

#Synced-CL 50269217

diff --git a/Engine/Plugins/PCG/Source/PCG/Private/Graph/PCGGraphCache.cpp b/Engine/Plugins/PCG/Source/PCG/Private/Graph/PCGGraphCache.cpp
--- a/Engine/Plugins/PCG/Source/PCG/Private/Graph/PCGGraphCache.cpp
+++ b/Engine/Plugins/PCG/Source/PCG/Private/Graph/PCGGraphCache.cpp
@@ -4,7 +4,9 @@
 
 #include "PCGComponent.h"
 #include "PCGModule.h"
+#include "Subsystems/PCGSubsystem.h"
 
+#include "Engine/World.h"
 #include "GameFramework/Actor.h"
 #include "HAL/IConsoleManager.h"
 #include "Misc/ScopeRWLock.h"
@@ -54,6 +56,28 @@ static TAutoConsoleVariable<bool> CVarEnableTLSGatherCacheSize(
 	false,
 	TEXT("Whether we give a hint to UPCGData subclass that we are calling GetResourceSizeEx for cache purposes."));
 
+static FAutoConsoleCommandWithWorld CResetCacheStatsCommand(
+	TEXT("pcg.Cache.ResetStats"),
+	TEXT("Resets the current cache stats"),
+	FConsoleCommandWithWorldDelegate::CreateLambda([](UWorld* InWorld)
+	{
+		if(UPCGSubsystem* PCGSubsystem = UWorld::GetSubsystem<UPCGSubsystem>(InWorld))
+		{
+			PCGSubsystem->ResetCacheStats();
+		}
+	}));
+
+static FAutoConsoleCommandWithWorld CLogCacheStatsCommand(
+	TEXT("pcg.Cache.LogStats"),
+	TEXT("Log the current cache stats"),
+	FConsoleCommandWithWorldDelegate::CreateLambda([](UWorld* InWorld)
+	{
+		if (UPCGSubsystem* PCGSubsystem = UWorld::GetSubsystem<UPCGSubsystem>(InWorld))
+		{
+			PCGSubsystem->LogCacheStats();
+		}
+	}));
+
 // Aliases for old deprecated settings.
 static TAutoConsoleVariable<bool> CVarCacheEnabled_DEPRECATED(
 	TEXT("pcg.Cache.Enabled"),
@@ -109,6 +133,17 @@ FPCGGraphCache::~FPCGGraphCache()
 	ClearCache();
 }
 
+void FPCGGraphCache::ResetStats()
+{
+	Stats = FPCGCacheStats();
+}
+
+void FPCGGraphCache::LogStats()
+{
+	const float HitRate = Stats.GetCount > 0 ? ((float)Stats.HitCount / Stats.GetCount) * 100.f : 0.f;
+	UE_LOG(LogPCG, Log, TEXT("Cache Stats : %lld Stores, %lld Hits / %lld Gets (%f), %lld Grow count"), Stats.StoreCount, Stats.HitCount, Stats.GetCount, HitRate, Stats.GrowCount);
+}
+
 bool FPCGGraphCache::IsEnabled() const
 {
 	return bIsGameWorld ? CVarCacheEnabledRuntime.GetValueOnAnyThread() : CVarCacheEnabledEditor.GetValueOnAnyThread();
@@ -142,6 +177,7 @@ bool FPCGGraphCache::GetFromCache(const FPCGGetFromCacheParams& Params, FPCGData
 	{
 		TRACE_CPUPROFILER_EVENT_SCOPE(FPCGGraphCache::GetFromCache);
 		PCG::TScopeLock ScopedLock(CacheLock);
+		++Stats.GetCount;
 
 		FPCGCacheEntryKey CacheKey(InElement, InDependenciesCrc);
 		if (const FPCGDataCollection* Value = const_cast<FPCGGraphCache*>(this)->CacheData.FindAndTouch(CacheKey))
@@ -153,6 +189,7 @@ bool FPCGGraphCache::GetFromCache(const FPCGGetFromCacheParams& Params, FPCGData
 			}
 
 			OutOutput = *Value;
+			++Stats.HitCount;
 
 			return true;
 		}
@@ -193,8 +230,10 @@ void FPCGGraphCache::StoreInCache(const FPCGStoreInCacheParams& Params, const FP
 		TRACE_CPUPROFILER_EVENT_SCOPE(FPCGGraphCache::StoreInCache);
 		PCG::TScopeLock ScopedLock(CacheLock);
 
+		++Stats.StoreCount;
 		if (CacheData.Num() == CacheData.Max())
 		{
+			++Stats.GrowCount;
 			GrowCache_Unsafe();
 		}
 
diff --git a/Engine/Plugins/PCG/Source/PCG/Private/Graph/PCGGraphCache.h b/Engine/Plugins/PCG/Source/PCG/Private/Graph/PCGGraphCache.h
--- a/Engine/Plugins/PCG/Source/PCG/Private/Graph/PCGGraphCache.h
+++ b/Engine/Plugins/PCG/Source/PCG/Private/Graph/PCGGraphCache.h
@@ -89,6 +89,12 @@ public:
 	/** True if debugging features enabled. Exposes a CVar so it can be called from editor module. */
 	bool IsDebuggingEnabled() const;
 
+	/** Reset cache stats. */
+	void ResetStats();
+
+	/** Log cache stats. */
+	void LogStats();
+
 #if WITH_EDITOR
 	/** Clears any cache entry for the given data. InSettings is optional and for logging only. */
 	void CleanFromCache(const IPCGElement* InElement, const UPCGSettings* InSettings = nullptr);
@@ -131,4 +137,15 @@ private:
 	bool bIsGameWorld = true;
 
 	mutable PCG::FLock CacheLock;
+
+	// Store Get/Store stats
+	struct FPCGCacheStats
+	{
+		int64 GetCount = 0;
+		int64 StoreCount = 0;
+		int64 HitCount = 0;
+		int64 GrowCount = 0;
+	};
+	
+	mutable FPCGCacheStats Stats;
 };
diff --git a/Engine/Plugins/PCG/Source/PCG/Private/Subsystems/PCGSubsystem.cpp b/Engine/Plugins/PCG/Source/PCG/Private/Subsystems/PCGSubsystem.cpp
--- a/Engine/Plugins/PCG/Source/PCG/Private/Subsystems/PCGSubsystem.cpp
+++ b/Engine/Plugins/PCG/Source/PCG/Private/Subsystems/PCGSubsystem.cpp
@@ -1426,6 +1426,22 @@ FPCGGenSourceManager* UPCGSubsystem::GetGenSourceManager() const
 	return RuntimeGenScheduler ? RuntimeGenScheduler->GenSourceManager : nullptr;
 }
 
+void UPCGSubsystem::ResetCacheStats()
+{
+	if (GraphExecutor)
+	{
+		GraphExecutor->GetCache().ResetStats();
+	}
+}
+
+void UPCGSubsystem::LogCacheStats()
+{
+	if (GraphExecutor)
+	{
+		GraphExecutor->GetCache().LogStats();
+	}
+}
+
 #if WITH_EDITOR
 
 FPCGTaskId UPCGSubsystem::ScheduleRefresh(UPCGComponent* Component, bool bForceRegen)
diff --git a/Engine/Plugins/PCG/Source/PCG/Public/Subsystems/PCGSubsystem.h b/Engine/Plugins/PCG/Source/PCG/Public/Subsystems/PCGSubsystem.h
--- a/Engine/Plugins/PCG/Source/PCG/Public/Subsystems/PCGSubsystem.h
+++ b/Engine/Plugins/PCG/Source/PCG/Public/Subsystems/PCGSubsystem.h
@@ -316,6 +316,9 @@ public:
 
 	PCG_API FPCGGenSourceManager* GetGenSourceManager() const;
 
+	void ResetCacheStats();
+	void LogCacheStats();
+
 #if WITH_EDITOR
 public:
 	/** Schedule refresh on the current or next frame */
