[PCG] Give access to execution cache to other execution sources + rename Actor methods to Self

#Synced-CL 50299325

diff --git a/Engine/Plugins/PCG/Source/PCG/Private/Graph/PCGGraphExecutor.h b/Engine/Plugins/PCG/Source/PCG/Private/Graph/PCGGraphExecutor.h
--- a/Engine/Plugins/PCG/Source/PCG/Private/Graph/PCGGraphExecutor.h
+++ b/Engine/Plugins/PCG/Source/PCG/Private/Graph/PCGGraphExecutor.h
@@ -166,10 +166,10 @@ private:
 
 	UPCGData* GetPCGData(FPCGTaskId InGraphExecutionTaskId, bool& bOutFound) { return PerExecutionCache.GetExecutionCachePCGData(InGraphExecutionTaskId, EPCGPerExecutionCacheDataType::PCGData, bOutFound); }
 	UPCGData* GetInputPCGData(FPCGTaskId InGraphExecutionTaskId, bool& bOutFound) { return PerExecutionCache.GetExecutionCachePCGData(InGraphExecutionTaskId, EPCGPerExecutionCacheDataType::InputData, bOutFound); }
-	UPCGData* GetActorPCGData(FPCGTaskId InGraphExecutionTaskId, bool& bOutFound) { return PerExecutionCache.GetExecutionCachePCGData(InGraphExecutionTaskId, EPCGPerExecutionCacheDataType::ActorData, bOutFound); }
+	UPCGData* GetSelfPCGData(FPCGTaskId InGraphExecutionTaskId, bool& bOutFound) { return PerExecutionCache.GetExecutionCachePCGData(InGraphExecutionTaskId, EPCGPerExecutionCacheDataType::SelfData, bOutFound); }
 	UPCGData* GetLandscapePCGData(FPCGTaskId InGraphExecutionTaskId, bool& bOutFound) { return PerExecutionCache.GetExecutionCachePCGData(InGraphExecutionTaskId, EPCGPerExecutionCacheDataType::LandscapeData, bOutFound); }
 	UPCGData* GetLandscapeHeightPCGData(FPCGTaskId InGraphExecutionTaskId, bool& bOutFound) { return PerExecutionCache.GetExecutionCachePCGData(InGraphExecutionTaskId, EPCGPerExecutionCacheDataType::LandscapeHeightData, bOutFound); }
-	UPCGData* GetOriginalActorPCGData(FPCGTaskId InGraphExecutionTaskId, bool& bOutFound) { return PerExecutionCache.GetExecutionCachePCGData(InGraphExecutionTaskId, EPCGPerExecutionCacheDataType::OriginalActorData, bOutFound); }
+	UPCGData* GetOriginalSelfPCGData(FPCGTaskId InGraphExecutionTaskId, bool& bOutFound) { return PerExecutionCache.GetExecutionCachePCGData(InGraphExecutionTaskId, EPCGPerExecutionCacheDataType::OriginalSelfData, bOutFound); }
 	TOptional<FBox> GetBounds(FPCGTaskId InGraphExecutionTaskId) { return PerExecutionCache.GetExecutionCacheBounds(InGraphExecutionTaskId, EPCGPerExecutionCacheDataType::Bounds); }
 	TOptional<FBox> GetOriginalBounds(FPCGTaskId InGraphExecutionTaskId) { return PerExecutionCache.GetExecutionCacheBounds(InGraphExecutionTaskId, EPCGPerExecutionCacheDataType::OriginalBounds); }
 	TOptional<FBox> GetLocalSpaceBounds(FPCGTaskId InGraphExecutionTaskId) { return PerExecutionCache.GetExecutionCacheBounds(InGraphExecutionTaskId, EPCGPerExecutionCacheDataType::LocalSpaceBounds); }
@@ -178,10 +178,10 @@ private:
 		
 	void SetPCGData(FPCGTaskId InGraphExecutionTaskId, UPCGData* InData) { PerExecutionCache.SetExecutionCachePCGData(InGraphExecutionTaskId, EPCGPerExecutionCacheDataType::PCGData, InData); }
 	void SetInputPCGData(FPCGTaskId InGraphExecutionTaskId, UPCGData* InData) { PerExecutionCache.SetExecutionCachePCGData(InGraphExecutionTaskId, EPCGPerExecutionCacheDataType::InputData, InData); }
-	void SetActorPCGData(FPCGTaskId InGraphExecutionTaskId, UPCGData* InData) { PerExecutionCache.SetExecutionCachePCGData(InGraphExecutionTaskId, EPCGPerExecutionCacheDataType::ActorData, InData); }
+	void SetSelfPCGData(FPCGTaskId InGraphExecutionTaskId, UPCGData* InData) { PerExecutionCache.SetExecutionCachePCGData(InGraphExecutionTaskId, EPCGPerExecutionCacheDataType::SelfData, InData); }
 	void SetLandscapePCGData(FPCGTaskId InGraphExecutionTaskId, UPCGData* InData) { PerExecutionCache.SetExecutionCachePCGData(InGraphExecutionTaskId, EPCGPerExecutionCacheDataType::LandscapeData, InData); }
 	void SetLandscapeHeightPCGData(FPCGTaskId InGraphExecutionTaskId, UPCGData* InData) { PerExecutionCache.SetExecutionCachePCGData(InGraphExecutionTaskId, EPCGPerExecutionCacheDataType::LandscapeHeightData, InData); }
-	void SetOriginalActorPCGData(FPCGTaskId InGraphExecutionTaskId, UPCGData* InData) { PerExecutionCache.SetExecutionCachePCGData(InGraphExecutionTaskId, EPCGPerExecutionCacheDataType::OriginalActorData, InData); }
+	void SetOriginalSelfPCGData(FPCGTaskId InGraphExecutionTaskId, UPCGData* InData) { PerExecutionCache.SetExecutionCachePCGData(InGraphExecutionTaskId, EPCGPerExecutionCacheDataType::OriginalSelfData, InData); }
 	void SetBounds(FPCGTaskId InGraphExecutionTaskId, const FBox& InBounds) { PerExecutionCache.SetExecutionCacheBounds(InGraphExecutionTaskId, EPCGPerExecutionCacheDataType::Bounds, InBounds); }
 	void SetOriginalBounds(FPCGTaskId InGraphExecutionTaskId, const FBox& InBounds) { PerExecutionCache.SetExecutionCacheBounds(InGraphExecutionTaskId, EPCGPerExecutionCacheDataType::OriginalBounds, InBounds); }
 	void SetLocalSpaceBounds(FPCGTaskId InGraphExecutionTaskId, const FBox& InBounds) { PerExecutionCache.SetExecutionCacheBounds(InGraphExecutionTaskId, EPCGPerExecutionCacheDataType::LocalSpaceBounds, InBounds); }
diff --git a/Engine/Plugins/PCG/Source/PCG/Private/Graph/PCGGraphPerExecutionCache.h b/Engine/Plugins/PCG/Source/PCG/Private/Graph/PCGGraphPerExecutionCache.h
--- a/Engine/Plugins/PCG/Source/PCG/Private/Graph/PCGGraphPerExecutionCache.h
+++ b/Engine/Plugins/PCG/Source/PCG/Private/Graph/PCGGraphPerExecutionCache.h
@@ -56,10 +56,10 @@ enum class EPCGPerExecutionCacheDataType : uint8
 {
 	PCGData = 0,
 	InputData,
-	ActorData,
+	SelfData,
 	LandscapeData,
 	LandscapeHeightData,
-	OriginalActorData,
+	OriginalSelfData,
 	Bounds,
 	OriginalBounds,
 	LocalSpaceBounds,
diff --git a/Engine/Plugins/PCG/Source/PCG/Private/PCGComponent.cpp b/Engine/Plugins/PCG/Source/PCG/Private/PCGComponent.cpp
--- a/Engine/Plugins/PCG/Source/PCG/Private/PCGComponent.cpp
+++ b/Engine/Plugins/PCG/Source/PCG/Private/PCGComponent.cpp
@@ -2941,14 +2941,14 @@ UPCGData* UPCGComponent::GetActorPCGData() const
 	if (UPCGSubsystem* PCGSubsystem = GetSubsystem())
 	{
 		bool bFound = false;
-		Data = PCGSubsystem->GetActorPCGData(CurrentGenerationTask, bFound);
+		Data = PCGSubsystem->GetSelfPCGData(CurrentGenerationTask, bFound);
 		if (!bFound)
 		{
 			PCG_EXECUTION_CACHE_VALIDATION_CHECK(this);
 			Data = CachedActorData && !IsLandscapeCachedDataDirty(CachedActorData) ? CachedActorData.Get() : CreateActorPCGData();
 			CachedActorData = Data;
 
-			PCGSubsystem->SetActorPCGData(CurrentGenerationTask, Data);
+			PCGSubsystem->SetSelfPCGData(CurrentGenerationTask, Data);
 			if (PCGSubsystem->IsGraphCacheDebuggingEnabled() && CurrentGenerationTask != InvalidPCGTaskId)
 			{
 				UE_LOG(LogPCG, Log, TEXT("         [%s] CACHE REFRESH CachedActorData"), *GetOwner()->GetName());
@@ -3017,7 +3017,7 @@ UPCGData* UPCGComponent::GetOriginalActorPCGData() const
 	if (UPCGSubsystem* PCGSubsystem = GetSubsystem())
 	{
 		bool bFound = false;
-		Data = PCGSubsystem->GetOriginalActorPCGData(CurrentGenerationTask, bFound);
+		Data = PCGSubsystem->GetOriginalSelfPCGData(CurrentGenerationTask, bFound);
 		if (!bFound)
 		{
 			if (UPCGComponent* OriginalComponent = PartitionActor->GetOriginalComponent(this))
@@ -3028,7 +3028,7 @@ UPCGData* UPCGComponent::GetOriginalActorPCGData() const
 					PCG_EXECUTION_CACHE_VALIDATION_CHECK(this);
 					Data = OriginalComponent->CreateActorPCGData();
 				}
-				PCGSubsystem->SetOriginalActorPCGData(CurrentGenerationTask, Data);
+				PCGSubsystem->SetOriginalSelfPCGData(CurrentGenerationTask, Data);
 				if (PCGSubsystem->IsGraphCacheDebuggingEnabled() && CurrentGenerationTask != InvalidPCGTaskId)
 				{
 					UE_LOG(LogPCG, Log, TEXT("         [%s] CACHE REFRESH CachedOriginalActorData"), *GetOwner()->GetName());
diff --git a/Engine/Plugins/PCG/Source/PCG/Private/Subsystems/PCGSubsystem.cpp b/Engine/Plugins/PCG/Source/PCG/Private/Subsystems/PCGSubsystem.cpp
--- a/Engine/Plugins/PCG/Source/PCG/Private/Subsystems/PCGSubsystem.cpp
+++ b/Engine/Plugins/PCG/Source/PCG/Private/Subsystems/PCGSubsystem.cpp
@@ -2084,10 +2084,10 @@ UPCGData* UPCGSubsystem::GetInputPCGData(FPCGTaskId InGraphExecutionTaskId, bool
 	return GraphExecutor ? GraphExecutor->GetInputPCGData(InGraphExecutionTaskId, bOutFound) : nullptr;
 }
 
-UPCGData* UPCGSubsystem::GetActorPCGData(FPCGTaskId InGraphExecutionTaskId, bool& bOutFound)
+UPCGData* UPCGSubsystem::GetSelfPCGData(FPCGTaskId InGraphExecutionTaskId, bool& bOutFound)
 {
 	bOutFound = false;
-	return GraphExecutor ? GraphExecutor->GetActorPCGData(InGraphExecutionTaskId, bOutFound) : nullptr;
+	return GraphExecutor ? GraphExecutor->GetSelfPCGData(InGraphExecutionTaskId, bOutFound) : nullptr;
 }
 
 UPCGData* UPCGSubsystem::GetLandscapePCGData(FPCGTaskId InGraphExecutionTaskId, bool& bOutFound)
@@ -2102,10 +2102,10 @@ UPCGData* UPCGSubsystem::GetLandscapeHeightPCGData(FPCGTaskId InGraphExecutionTa
 	return GraphExecutor ? GraphExecutor->GetLandscapeHeightPCGData(InGraphExecutionTaskId, bOutFound) : nullptr;
 }
 
-UPCGData* UPCGSubsystem::GetOriginalActorPCGData(FPCGTaskId InGraphExecutionTaskId, bool& bOutFound)
+UPCGData* UPCGSubsystem::GetOriginalSelfPCGData(FPCGTaskId InGraphExecutionTaskId, bool& bOutFound)
 {
 	bOutFound = false;
-	return GraphExecutor ? GraphExecutor->GetOriginalActorPCGData(InGraphExecutionTaskId, bOutFound) : nullptr;
+	return GraphExecutor ? GraphExecutor->GetOriginalSelfPCGData(InGraphExecutionTaskId, bOutFound) : nullptr;
 }
 
 TOptional<FBox> UPCGSubsystem::GetBounds(FPCGTaskId InGraphExecutionTaskId)
@@ -2144,11 +2144,11 @@ void UPCGSubsystem::SetInputPCGData(FPCGTaskId InGraphExecutionTaskId, UPCGData*
 	}
 }
 
-void UPCGSubsystem::SetActorPCGData(FPCGTaskId InGraphExecutionTaskId, UPCGData* InData)
+void UPCGSubsystem::SetSelfPCGData(FPCGTaskId InGraphExecutionTaskId, UPCGData* InData)
 {
 	if (GraphExecutor)
 	{
-		GraphExecutor->SetActorPCGData(InGraphExecutionTaskId, InData);
+		GraphExecutor->SetSelfPCGData(InGraphExecutionTaskId, InData);
 	}
 }
 
@@ -2168,11 +2168,11 @@ void UPCGSubsystem::SetLandscapeHeightPCGData(FPCGTaskId InGraphExecutionTaskId,
 	}
 }
 
-void UPCGSubsystem::SetOriginalActorPCGData(FPCGTaskId InGraphExecutionTaskId, UPCGData* InData)
+void UPCGSubsystem::SetOriginalSelfPCGData(FPCGTaskId InGraphExecutionTaskId, UPCGData* InData)
 {
 	if (GraphExecutor)
 	{
-		GraphExecutor->SetOriginalActorPCGData(InGraphExecutionTaskId, InData);
+		GraphExecutor->SetOriginalSelfPCGData(InGraphExecutionTaskId, InData);
 	}
 }
 
diff --git a/Engine/Plugins/PCG/Source/PCG/Public/Subsystems/PCGSubsystem.h b/Engine/Plugins/PCG/Source/PCG/Public/Subsystems/PCGSubsystem.h
--- a/Engine/Plugins/PCG/Source/PCG/Public/Subsystems/PCGSubsystem.h
+++ b/Engine/Plugins/PCG/Source/PCG/Public/Subsystems/PCGSubsystem.h
@@ -319,6 +319,30 @@ public:
 	void ResetCacheStats();
 	void LogCacheStats();
 
+	/** Getters for per execution data */
+	PCG_API UPCGData* GetPCGData(FPCGTaskId InGraphExecutionTaskId, bool& bOutFound);
+	PCG_API UPCGData* GetInputPCGData(FPCGTaskId InGraphExecutionTaskId, bool& bOutFound);
+	PCG_API UPCGData* GetSelfPCGData(FPCGTaskId InGraphExecutionTaskId, bool& bOutFound);
+	PCG_API UPCGData* GetLandscapePCGData(FPCGTaskId InGraphExecutionTaskId, bool& bOutFound);
+	PCG_API UPCGData* GetLandscapeHeightPCGData(FPCGTaskId InGraphExecutionTaskId, bool& bOutFound);
+	PCG_API UPCGData* GetOriginalSelfPCGData(FPCGTaskId InGraphExecutionTaskId, bool& bOutFound);
+	PCG_API TOptional<FBox> GetBounds(FPCGTaskId InGraphExecutionTaskId);
+	PCG_API TOptional<FBox> GetOriginalBounds(FPCGTaskId InGraphExecutionTaskId);
+	PCG_API TOptional<FBox> GetLocalSpaceBounds(FPCGTaskId InGraphExecutionTaskId);
+	PCG_API TOptional<FBox> GetOriginalLocalSpaceBounds(FPCGTaskId InGraphExecutionTaskId);
+
+	/** Setters for per execution data. Usually called from execution state ExecutePreGraph implementations. */
+	PCG_API void SetPCGData(FPCGTaskId InGraphExecutionTaskId, UPCGData* InData);
+	PCG_API void SetInputPCGData(FPCGTaskId InGraphExecutionTaskId, UPCGData* InData);
+	PCG_API void SetSelfPCGData(FPCGTaskId InGraphExecutionTaskId, UPCGData* InData);
+	PCG_API void SetLandscapePCGData(FPCGTaskId InGraphExecutionTaskId, UPCGData* InData);
+	PCG_API void SetLandscapeHeightPCGData(FPCGTaskId InGraphExecutionTaskId, UPCGData* InData);
+	PCG_API void SetOriginalSelfPCGData(FPCGTaskId InGraphExecutionTaskId, UPCGData* InData);
+	PCG_API void SetBounds(FPCGTaskId InGraphExecutionTaskId, const FBox& InBounds);
+	PCG_API void SetOriginalBounds(FPCGTaskId InGraphExecutionTaskId, const FBox& InBounds);
+	PCG_API void SetLocalSpaceBounds(FPCGTaskId InGraphExecutionTaskId, const FBox& InBounds);
+	PCG_API void SetOriginalLocalSpaceBounds(FPCGTaskId InGraphExecutionTaskId, const FBox& InBounds);
+
 #if WITH_EDITOR
 public:
 	/** Schedule refresh on the current or next frame */
@@ -402,33 +426,11 @@ private:
 	void OnLoadedActorRemovedFromLevelPreEvent(const TArray<AActor*>& InActors);
 	void OnLoadedActorRemovedFromLevelPostEvent(const TArray<AActor*>& InActors);
 #endif // WITH_EDITOR
-	
+
 private:
 	/** When registering PAs, we might not have a PCG World Actor already registered (at runtime). In that case we'll look for a World Actor in the same level.*/
 	APCGWorldActor* GetPCGWorldActorForPartitionActor(APCGPartitionActor* InActor);
 
-	UPCGData* GetPCGData(FPCGTaskId InGraphExecutionTaskId, bool& bOutFound);
-	UPCGData* GetInputPCGData(FPCGTaskId InGraphExecutionTaskId, bool& bOutFound);
-	UPCGData* GetActorPCGData(FPCGTaskId InGraphExecutionTaskId, bool& bOutFound);
-	UPCGData* GetLandscapePCGData(FPCGTaskId InGraphExecutionTaskId, bool& bOutFound);
-	UPCGData* GetLandscapeHeightPCGData(FPCGTaskId InGraphExecutionTaskId, bool& bOutFound);
-	UPCGData* GetOriginalActorPCGData(FPCGTaskId InGraphExecutionTaskId, bool& bOutFound);
-	TOptional<FBox> GetBounds(FPCGTaskId InGraphExecutionTaskId);
-	TOptional<FBox> GetOriginalBounds(FPCGTaskId InGraphExecutionTaskId);
-	TOptional<FBox> GetLocalSpaceBounds(FPCGTaskId InGraphExecutionTaskId);
-	TOptional<FBox> GetOriginalLocalSpaceBounds(FPCGTaskId InGraphExecutionTaskId);
-
-	void SetPCGData(FPCGTaskId InGraphExecutionTaskId, UPCGData* InData);
-	void SetInputPCGData(FPCGTaskId InGraphExecutionTaskId, UPCGData* InData);
-	void SetActorPCGData(FPCGTaskId InGraphExecutionTaskId, UPCGData* InData);
-	void SetLandscapePCGData(FPCGTaskId InGraphExecutionTaskId, UPCGData* InData);
-	void SetLandscapeHeightPCGData(FPCGTaskId InGraphExecutionTaskId, UPCGData* InData);
-	void SetOriginalActorPCGData(FPCGTaskId InGraphExecutionTaskId, UPCGData* InData);
-	void SetBounds(FPCGTaskId InGraphExecutionTaskId, const FBox& InBounds);
-	void SetOriginalBounds(FPCGTaskId InGraphExecutionTaskId, const FBox& InBounds);
-	void SetLocalSpaceBounds(FPCGTaskId InGraphExecutionTaskId, const FBox& InBounds);
-	void SetOriginalLocalSpaceBounds(FPCGTaskId InGraphExecutionTaskId, const FBox& InBounds);
-
 	void ExecuteBeginTickActions();
 
 	void OnGlobalComponentReregisterContextDestroyed();
