[PCG] Graph parameters: fix crash when copy/pasting then renaming

#Synced-CL 50299325

diff --git a/Engine/Plugins/PCG/Source/PCGEditor/Private/PCGEditorGraph.cpp b/Engine/Plugins/PCG/Source/PCGEditor/Private/PCGEditorGraph.cpp
--- a/Engine/Plugins/PCG/Source/PCGEditor/Private/PCGEditorGraph.cpp
+++ b/Engine/Plugins/PCG/Source/PCGEditor/Private/PCGEditorGraph.cpp
@@ -409,7 +409,18 @@ void UPCGEditorGraph::OnGraphChanged(UPCGGraphInterface* InGraph, EPCGChangeType
 
 void UPCGEditorGraph::OnGraphUserParametersChanged(UPCGGraphInterface* InGraph, EPCGGraphParameterEvent ChangeType, FName ChangedPropertyName)
 {
-	if ((ChangeType != EPCGGraphParameterEvent::RemovedUnused && ChangeType != EPCGGraphParameterEvent::RemovedUsed) || InGraph != PCGGraph)
+	if (InGraph != PCGGraph)
+	{
+		return;
+	}
+
+	// Those change types can result in a parameter node being removed from the graph. (same condition found in UPCGGraph::OnGraphParametersChanged)
+	const bool bCheckForRemovedNodes = ChangeType == EPCGGraphParameterEvent::RemovedUnused
+		|| ChangeType == EPCGGraphParameterEvent::RemovedUsed
+		|| ChangeType == EPCGGraphParameterEvent::PropertyRenamed
+		|| ChangeType == EPCGGraphParameterEvent::PropertyTypeModified;
+	
+	if (!bCheckForRemovedNodes)
 	{
 		return;
 	}
