[PCG] Embedded Subgraphs support

#Synced-CL 50299325

#Synced-CL 50299325

diff --git a/Engine/Plugins/Experimental/ProceduralVegetationEditor/Source/ProceduralVegetation/Public/ProceduralVegetation.h b/Engine/Plugins/Experimental/ProceduralVegetationEditor/Source/ProceduralVegetation/Public/ProceduralVegetation.h
--- a/Engine/Plugins/Experimental/ProceduralVegetationEditor/Source/ProceduralVegetation/Public/ProceduralVegetation.h
+++ b/Engine/Plugins/Experimental/ProceduralVegetationEditor/Source/ProceduralVegetation/Public/ProceduralVegetation.h
@@ -37,6 +37,8 @@ public:
 	virtual bool ShowGraphCustomization() const override { return false; }
 	virtual bool IsTemplatePropertyEnabled() const override { return false; }
 #endif
+
+	virtual bool SupportsEmbeddedSubgraphs() const override { return false; }
 };
 
 /**
diff --git a/Engine/Plugins/Experimental/ProceduralVegetationEditor/Source/ProceduralVegetationEditor/Private/PVEditor.cpp b/Engine/Plugins/Experimental/ProceduralVegetationEditor/Source/ProceduralVegetationEditor/Private/PVEditor.cpp
--- a/Engine/Plugins/Experimental/ProceduralVegetationEditor/Source/ProceduralVegetationEditor/Private/PVEditor.cpp
+++ b/Engine/Plugins/Experimental/ProceduralVegetationEditor/Source/ProceduralVegetationEditor/Private/PVEditor.cpp
@@ -101,7 +101,7 @@ void FPVEditor::Initialize(const EToolkitMode::Type InMode, const TSharedPtr<cla
 	ExecutionSource->Generate();
 
 	// Select the first available output
-	for (auto EdNode : GetPCGEditorGraph()->Nodes)
+	for (auto EdNode : GetMainEditorGraph()->Nodes)
 	{
 		if (const UPCGEditorGraphNodeBase* PCGEdNode = Cast<UPCGEditorGraphNodeBase>(EdNode))
 		{
@@ -109,7 +109,7 @@ void FPVEditor::Initialize(const EToolkitMode::Type InMode, const TSharedPtr<cla
 			{
 				if(PCGNode->GetSettings()->IsA<UPVOutputSettings>())
 				{
-					GetPCGEditorGraph()->SelectNodeSet({PCGEdNode}, true);
+					GetMainEditorGraph()->SelectNodeSet({PCGEdNode}, true);
 					break;
 				}
 			}
@@ -268,7 +268,7 @@ void FPVEditor::RegisterToolbarInternal(FToolBarBuilder& ToolbarBuilder) const
 {
 	RegisterToolbarButton(ToolbarBuilder, EPCGToolbarButtons::Find);
 	RegisterToolbarButton(ToolbarBuilder, EPCGToolbarButtons::ForceRegen);
-	
+	RegisterToolbarButton(ToolbarBuilder, EPCGToolbarButtons::Graph);
 	const FPVEditorCommands& Commands = FPVEditorCommands::Get();
 
 	ToolbarBuilder.AddToolBarButton(Commands.Export, NAME_None, TAttribute<FText>(), TAttribute<FText>(),
@@ -394,7 +394,7 @@ void FPVEditor::OnExport_Clicked()
 
 FReply FPVEditor::ExportInternal()
 {
-	const UProceduralVegetationGraph* ProceduralVegetationGraph = Cast<UProceduralVegetationGraph>(GetPCGGraph());
+	const UProceduralVegetationGraph* ProceduralVegetationGraph = Cast<UProceduralVegetationGraph>(GetMainGraph());
 	
 	TSharedPtr<SPVExportSelectionDialog> ExportSelectionWidget =
 		SNew(SPVExportSelectionDialog)
diff --git a/Engine/Plugins/PCG/Source/PCG/Private/Helpers/PCGSubgraphHelpers.cpp b/Engine/Plugins/PCG/Source/PCG/Private/Helpers/PCGSubgraphHelpers.cpp
--- a/Engine/Plugins/PCG/Source/PCG/Private/Helpers/PCGSubgraphHelpers.cpp
+++ b/Engine/Plugins/PCG/Source/PCG/Private/Helpers/PCGSubgraphHelpers.cpp
@@ -38,11 +38,12 @@ namespace PCGSubgraphHelpersExtra
 
 		void SetNewGraphAndUpdatePositions(UPCGGraph* NewGraph);
 
+		void CopyAndRemapEmbeddedSubgraphs();
 		void AddDiscoveredGraphParameters();
 		void CopyNodes(const TArray<UObject*>& ExtraNodesToCollapse);
 		void DeleteNodes(const TArray<UObject*>& ExtraNodesToCollapse);
 		void CreateSubgraphNodeAndConnectRemainingPins();
-
+		
 		/**
 		* Go through all the collapsed nodes and reconnect all of them according to the original connections.
 		* Will also handle any special cases (named reroutes, graph parameters, ...)
@@ -64,6 +65,8 @@ namespace PCGSubgraphHelpersExtra
 		TArray<UPCGNode*> NodesToRemove;
 		// Also some node can be mark superfluous, meaning that they can be removed if at the end, there is no edges remaining.
 		TArray<UPCGNode*> SuperfluousNodes;
+		// List used embedded subgraphs
+		TArray<UPCGNode*> EmbeddedSubgraphNodes;
 
 		// Keep a list of all outside pins and the name of the label to connect to.
 		TArray<TPair<UPCGPin*, FName>> OutsideSubgraphPinToSubgraphPinLabel;
@@ -78,6 +81,8 @@ namespace PCGSubgraphHelpersExtra
 		UPCGGraph* OriginalGraph = nullptr;
 		UPCGGraph* NewGraph = nullptr;
 
+		TMap<UPCGGraph*, UPCGGraph*> DuplicatedEmbeddedSubgraphs;
+
 		/** 
 		* Connect two pins that cross the subgraph boundaries. The Inside Pin is the pin that is inside the subgraph (NewGraph) and outside pin the pin in the main graph.
 		* bIsSubgraphInputBoundary is set if it is the input boundary that is crossed. Otherwise, it is the output boundary.
@@ -162,6 +167,14 @@ namespace PCGSubgraphHelpersExtra
 				}
 			}
 
+			if (const UPCGBaseSubgraphSettings* SubgraphSettings = Cast<UPCGBaseSubgraphSettings>(NodeSettings))
+			{
+				if (SubgraphSettings->GetSubgraph() && SubgraphSettings->GetSubgraph()->IsEmbeddedSubgraph())
+				{
+					CollapseInfo.EmbeddedSubgraphNodes.Add(PCGNode);
+				}
+			}
+
 			CollapseInfo.ValidNodesToCollapse.Add(PCGNode);
 
 #if WITH_EDITOR
@@ -212,6 +225,41 @@ namespace PCGSubgraphHelpersExtra
 		NewGraph->AddUserParameters(GraphParametersUsed, OriginalGraph);
 	}
 
+	void FCollapsingInformation::CopyAndRemapEmbeddedSubgraphs()
+	{
+		check(OriginalGraph && NewGraph);
+
+		// Not needed if we are creating a new embedded graph (and not an asset)
+		if(NewGraph->GetEmbeddedParentGraph() == OriginalGraph)
+		{
+			return;
+		}
+
+		TArray<UPCGGraph*> EmbeddedSubgraphsToCopy;
+		for (UPCGNode* EmbeddedSubgraphNode : EmbeddedSubgraphNodes)
+		{
+			UPCGBaseSubgraphSettings* SubgraphSettings = CastChecked<UPCGBaseSubgraphSettings>(EmbeddedSubgraphNode->GetSettings());
+			UPCGGraph* EmbeddedSubgraph = SubgraphSettings->GetSubgraph();
+			check(EmbeddedSubgraph && EmbeddedSubgraph->IsEmbeddedSubgraph());
+			EmbeddedSubgraphsToCopy.AddUnique(EmbeddedSubgraph);
+		}
+
+		DuplicatedEmbeddedSubgraphs.Empty();
+		NewGraph->DuplicateEmbeddedSubgraphs(EmbeddedSubgraphsToCopy, DuplicatedEmbeddedSubgraphs);
+
+		for (UPCGNode* EmbeddedSubgraphNode : EmbeddedSubgraphNodes)
+		{
+			check(OriginalToNewNodesMapping.Contains(EmbeddedSubgraphNode));
+			UPCGNode* NewEmbeddedSubgraphNode = OriginalToNewNodesMapping[EmbeddedSubgraphNode];
+			check(NewEmbeddedSubgraphNode);
+
+			UPCGBaseSubgraphSettings* SubgraphSettings = CastChecked<UPCGBaseSubgraphSettings>(NewEmbeddedSubgraphNode->GetSettings());
+			UPCGGraph* DuplicatedEmbeddedSubgraph = DuplicatedEmbeddedSubgraphs.FindChecked(SubgraphSettings->GetSubgraph());
+
+			SubgraphSettings->SetSubgraph(DuplicatedEmbeddedSubgraph);
+		}
+	}
+
 	void FCollapsingInformation::CopyNodes(const TArray<UObject*>& ExtraNodesToCollapse)
 	{
 		check(NewGraph);
@@ -798,6 +846,7 @@ UPCGGraph* FPCGSubgraphHelpers::CollapseIntoSubgraphWithReason(UPCGGraph* InOrig
 
 	CollapseInfo.AddDiscoveredGraphParameters();
 	CollapseInfo.CopyNodes(InExtraNodesToCollapse);
+	CollapseInfo.CopyAndRemapEmbeddedSubgraphs();
 	CollapseInfo.ReconnectAllNodes();
 	CollapseInfo.CleanUpOriginalEdges();
 	CollapseInfo.DeleteNodes(InExtraNodesToCollapse);
diff --git a/Engine/Plugins/PCG/Source/PCG/Private/PCGGraph.cpp b/Engine/Plugins/PCG/Source/PCG/Private/PCGGraph.cpp
--- a/Engine/Plugins/PCG/Source/PCG/Private/PCGGraph.cpp
+++ b/Engine/Plugins/PCG/Source/PCG/Private/PCGGraph.cpp
@@ -24,6 +24,7 @@
 
 #include "AssetRegistry/AssetData.h"
 #include "Logging/StructuredLog.h"
+#include "Serialization/ArchiveReplaceObjectRef.h"
 
 #if WITH_EDITOR
 #include "CoreGlobals.h"
@@ -930,8 +931,83 @@ void UPCGGraph::AddReferencedObjects(UObject* InThis, FReferenceCollector& Colle
 	Super::AddReferencedObjects(InThis, Collector);
 }
 
+bool UPCGGraph::ContainsEmbeddedSubgraph(UPCGGraph* InEmbeddedSubgraph) const
+{
+	return EmbeddedSubgraphs.Contains(InEmbeddedSubgraph);
+}
+
+void UPCGGraph::DuplicateEmbeddedSubgraphs(const TArray<UPCGGraph*>& InEmbeddedSubgraphs, TMap<UPCGGraph*, UPCGGraph*>& OutDuplicatedGraphs)
+{
+	check(OutDuplicatedGraphs.IsEmpty());
+
+	// Find graphs to duplicate recursively
+	TSet<UPCGGraph*> EmbeddedSubgraphsToCopy;
+	for (UPCGGraph* EmbeddedSubgraphToCopy : InEmbeddedSubgraphs)
+	{
+		bool bAlreadySet = false;
+		EmbeddedSubgraphsToCopy.Add(EmbeddedSubgraphToCopy, &bAlreadySet);
+		if (!bAlreadySet)
+		{
+			EmbeddedSubgraphToCopy->ForEachNodeRecursively([&EmbeddedSubgraphsToCopy](UPCGNode* InNode)
+			{
+				if (UPCGBaseSubgraphSettings* RecursiveSubgraphSettings = Cast<UPCGBaseSubgraphSettings>(InNode->GetSettings()))
+				{
+					if (UPCGGraph* RecursiveSubgraph = RecursiveSubgraphSettings->GetSubgraph(); RecursiveSubgraph && RecursiveSubgraph->IsEmbeddedSubgraph())
+					{
+						EmbeddedSubgraphsToCopy.Add(RecursiveSubgraph);
+					}
+				}
+
+				return true;
+			});
+		}
+	}
+
+	if (EmbeddedSubgraphsToCopy.IsEmpty())
+	{
+		return;
+	}
+
+	Modify();
+
+	// Duplicate graphs
+	for (UPCGGraph* EmbeddedSubgraphToCopy : EmbeddedSubgraphsToCopy)
+	{
+		check(EmbeddedSubgraphToCopy->IsEmbeddedSubgraph());
+
+		UPCGGraph* DuplicatedEmbeddedSubgraph = DuplicateObject<UPCGGraph>(EmbeddedSubgraphToCopy, this);
+		EmbeddedSubgraphs.Add(DuplicatedEmbeddedSubgraph);
+
+		check(DuplicatedEmbeddedSubgraph->GetEmbeddedParentGraph() == this);
+
+		OutDuplicatedGraphs.Add(EmbeddedSubgraphToCopy, DuplicatedEmbeddedSubgraph);
+	}
+
+	// Replace internal references
+	for (auto It : OutDuplicatedGraphs)
+	{
+		FArchiveReplaceObjectRef<UPCGGraph>(It.Value, OutDuplicatedGraphs);
+	}
+}
+
+EPCGHiGenGrid UPCGGraph::GetDefaultGrid() const 
+{ 
+	if (UPCGGraph* EmbeddedParentGraph = GetEmbeddedParentGraph())
+	{
+		return EmbeddedParentGraph->GetDefaultGrid();
+	}
+
+	ensure(IsHierarchicalGenerationEnabled()); 
+	return HiGenGridSize; 
+}
+
 uint32 UPCGGraph::GetDefaultGridSize() const
 {
+	if (UPCGGraph* EmbeddedParentGraph = GetEmbeddedParentGraph())
+	{
+		return EmbeddedParentGraph->GetDefaultGridSize();
+	}
+
 	if (IsHierarchicalGenerationEnabled() && PCGHiGenGrid::IsValidGrid(HiGenGridSize))
 	{
 		return PCGHiGenGrid::GridToGridSize(HiGenGridSize) * (1 << HiGenExponential);
@@ -942,6 +1018,26 @@ uint32 UPCGGraph::GetDefaultGridSize() const
 	}
 }
 
+bool UPCGGraph::IsHierarchicalGenerationEnabled() const 
+{ 
+	if (UPCGGraph* EmbeddedParentGraph = GetEmbeddedParentGraph())
+	{
+		return EmbeddedParentGraph->IsHierarchicalGenerationEnabled();
+	}
+
+	return bUseHierarchicalGeneration; 
+}
+
+bool UPCGGraph::Use2DGrid() const 
+{ 
+	if (UPCGGraph* EmbeddedParentGraph = GetEmbeddedParentGraph())
+	{
+		return EmbeddedParentGraph->Use2DGrid();
+	}
+
+	return bUse2DGrid; 
+}
+
 TSharedPtr<FComputeGraphInstance> UPCGGraph::RetrieveComputeGraphInstanceFromPool(const FComputeGraphInstanceKey& InKey, bool& bOutNewInstance) const
 {
 	check(IsInGameThread());
@@ -1795,6 +1891,16 @@ void UPCGGraph::OnPCGQualityLevelChanged()
 	}
 }
 
+const TArray<TSubclassOf<UObject>>& UPCGGraph::GetChangeTrackingIgnoredClasses() const
+{
+	if (UPCGGraph* EmbeddedParentGraph = GetEmbeddedParentGraph())
+	{
+		return EmbeddedParentGraph->GetChangeTrackingIgnoredClasses();
+	}
+
+	return ChangeTrackingIgnoredClasses;
+}
+
 FPCGSelectionKeyToSettingsMap UPCGGraph::GetTrackedSelectionKeysToSettings() const
 {
 	FPCGSelectionKeyToSettingsMap TagsToSettings;
@@ -1974,6 +2080,12 @@ void UPCGGraph::NotifyGraphChanged(EPCGChangeType ChangeType)
 		}
 	}
 
+	// Notify Outer Graph if we are a embedded subgraph
+	if (UPCGGraph* Parent = GetEmbeddedParentGraph())
+	{
+		Parent->NotifyGraphChanged(ChangeType);
+	}
+
 	OnGraphChangedDelegate.Broadcast(this, ChangeType);
 
 	bIsNotifying = false;
@@ -2130,9 +2242,11 @@ void UPCGGraph::PostEditChangeChainProperty(FPropertyChangedChainEvent& Property
 		|| PropertyName == GET_MEMBER_NAME_CHECKED(UPCGGraphInterface, bOverrideTitle)
 		|| PropertyName == GET_MEMBER_NAME_CHECKED(UPCGGraphInterface, Color)
 		|| PropertyName == GET_MEMBER_NAME_CHECKED(UPCGGraphInterface, bOverrideColor)
-		|| PropertyName == GET_MEMBER_NAME_CHECKED(UPCGGraph, bIsStandaloneGraph))
+		|| PropertyName == GET_MEMBER_NAME_CHECKED(UPCGGraph, bIsStandaloneGraph)
+		|| PropertyName == GET_MEMBER_NAME_CHECKED(UPCGGraph, Category)
+		|| PropertyName == GET_MEMBER_NAME_CHECKED(UPCGGraph, Description))
 	{
-		NotifyGraphChanged(EPCGChangeType::Cosmetic);
+		NotifyGraphChanged(EPCGChangeType::Cosmetic | EPCGChangeType::GraphCustomization);
 	}
 	else if (VerifyIfGraphCustomizationChanged(PropertyChangedEvent))
 	{
@@ -2155,6 +2269,35 @@ void UPCGGraph::PostEditUndo()
 	}
 }
 
+UPCGGraph* UPCGGraph::AddNewEmbeddedSubgraph()
+{
+	check(SupportsEmbeddedSubgraphs());
+
+	Modify();
+	UPCGGraph* EmbeddedSubgraph = NewObject<UPCGGraph>(this, GetEmbeddedSubgraphClass(), NAME_None, RF_Transactional);
+	EmbeddedSubgraph->bOverrideTitle = true;
+	EmbeddedSubgraph->Title = FText::FromName(EmbeddedSubgraph->GetFName());
+	EmbeddedSubgraphs.Add(EmbeddedSubgraph);
+	NotifyGraphChanged(EPCGChangeType::Structural | EPCGChangeType::GraphCustomization);
+	return EmbeddedSubgraph;
+}
+
+void UPCGGraph::DeleteEmbeddedSubgraph(UPCGGraph* InEmbeddedSubgraph)
+{
+	check(SupportsEmbeddedSubgraphs());
+
+	if (int32 Index = EmbeddedSubgraphs.Find(InEmbeddedSubgraph); Index != INDEX_NONE)
+	{
+		Modify();
+		EmbeddedSubgraphs.RemoveAt(Index);
+
+		InEmbeddedSubgraph->Modify();
+		InEmbeddedSubgraph->Rename(nullptr, GetTransientPackage(), REN_DontCreateRedirectors);
+		InEmbeddedSubgraph->MarkAsGarbage();
+		NotifyGraphChanged(EPCGChangeType::Structural | EPCGChangeType::GraphCustomization);
+	}
+}
+
 void UPCGGraph::FixInvalidEdges()
 {
 	auto ValidatePins = [this](const TArray<TObjectPtr<UPCGPin>>& Pins, bool bPinsAreInputs)
diff --git a/Engine/Plugins/PCG/Source/PCG/Public/PCGGraph.h b/Engine/Plugins/PCG/Source/PCG/Public/PCGGraph.h
--- a/Engine/Plugins/PCG/Source/PCG/Public/PCGGraph.h
+++ b/Engine/Plugins/PCG/Source/PCG/Public/PCGGraph.h
@@ -104,6 +104,41 @@ public:
 	TSet<FGuid> PropertiesIDsOverridden;
 };
 
+/**
+ * Struct used to restore graph editor per document settings
+ */
+USTRUCT()
+struct FPCGGraphDocumentInfo
+{
+	GENERATED_BODY()
+
+	/** Edited object */
+	UPROPERTY()
+	FSoftObjectPath EditedObjectPath;
+
+	/** Saved view position */
+	UPROPERTY()
+	FVector2f SavedViewOffset;
+
+	/** Saved zoom amount */
+	UPROPERTY()
+	float SavedZoomAmount = -1.0f;
+
+	FPCGGraphDocumentInfo() = default;
+
+	FPCGGraphDocumentInfo(UObject* InEditedObject, const FVector2f& InSavedViewOffset, float InSavedZoomAmount)
+		: EditedObjectPath(InEditedObject)
+		, SavedViewOffset(InSavedViewOffset)
+		, SavedZoomAmount(InSavedZoomAmount)
+	{
+	}
+
+	friend bool operator==(const FPCGGraphDocumentInfo& LHS, const FPCGGraphDocumentInfo& RHS)
+	{
+		return LHS.EditedObjectPath == RHS.EditedObjectPath && LHS.SavedViewOffset == RHS.SavedViewOffset && LHS.SavedZoomAmount == RHS.SavedZoomAmount;
+	}
+};
+
 UCLASS(MinimalAPI, BlueprintType, Abstract, ClassGroup = (Procedural))
 class UPCGGraphInterface : public UObject
 {
@@ -192,6 +227,10 @@ public:
 	/** When enabled and standalone, it adds the action in the asset explorer to trigger a generation. Invisible, to be set in the CDO of child class to disable. */
 	UPROPERTY(AssetRegistrySearchable, meta = (PCGNoHash))
 	bool bExposeGenerationInAssetExplorer = true;
+
+	/** Set of documents that were being edited in this pcg graph, so we can open them right away */
+	UPROPERTY()
+	TArray<FPCGGraphDocumentInfo> LastEditedDocuments;
 #endif
 
 	template <typename T>
@@ -300,7 +339,23 @@ public:
 	UE_API virtual void PostEditUndo() override;
 	
 	virtual TOptional<FPCGGraphToolData> GetGraphToolData() const override { return ToolData; }
+
+	virtual TSubclassOf<UPCGGraph> GetEmbeddedSubgraphClass() const { return UPCGGraph::StaticClass(); }
+
+	UE_API UPCGGraph* AddNewEmbeddedSubgraph();
+	UE_API void DeleteEmbeddedSubgraph(UPCGGraph* InEmbeddedSubgraph);
 #endif
+	UE_API bool ContainsEmbeddedSubgraph(UPCGGraph* InEmbeddedSubgraph) const;
+	UE_API void DuplicateEmbeddedSubgraphs(const TArray<UPCGGraph*>& InEmbeddedSubgraphs, TMap<UPCGGraph*, UPCGGraph*>& OutDuplicatedGraphs);
+
+	virtual bool SupportsEmbeddedSubgraphs() const { return !IsEmbeddedSubgraph(); }
+	UPCGGraph* GetEmbeddedParentGraph() const { return GetTypedOuter<UPCGGraph>(); }
+
+	UFUNCTION()
+	bool IsEmbeddedSubgraph() const { return !!GetEmbeddedParentGraph(); }
+	
+	const TArray<TObjectPtr<UPCGGraph>> GetEmbeddedSubgraphs() { return EmbeddedSubgraphs; }
+
 	/** ~End UObject interface */
 
 	/** ~Begin UPCGGraphInterface interface */
@@ -309,10 +364,11 @@ public:
 	/** ~End UPCGGraphInterface interface */
 
 	// Default grid size for generation. For hierarchical generation, nodes outside of grid size graph ranges will generate on this grid.
-	EPCGHiGenGrid GetDefaultGrid() const { ensure(IsHierarchicalGenerationEnabled()); return HiGenGridSize; }
+	UE_API EPCGHiGenGrid GetDefaultGrid() const;
 	UE_API uint32 GetDefaultGridSize() const;
-	bool IsHierarchicalGenerationEnabled() const { return bUseHierarchicalGeneration; }
-	bool Use2DGrid() const { return bUse2DGrid; }
+	UE_API bool IsHierarchicalGenerationEnabled() const;
+	UE_API bool Use2DGrid() const;
+
 	bool HasDefaultConstructedInputs() const { return bHasDefaultConstructedInputs; }
 
 	using FComputeGraphInstanceKey = TPair</*GridSize*/uint32, /*Compute graph index*/int32>;
@@ -332,14 +388,14 @@ public:
 	FText Description;
 
 	/** Marks the graph to be not refreshed automatically when the landscape changes, even if it is used. */
-	UPROPERTY(BlueprintReadWrite, EditAnywhere, Category = "Settings|Advanced", meta=(PCGNoHash, EditCondition="!IsStandaloneGraph()", EditConditionHides))
+	UPROPERTY(BlueprintReadWrite, EditAnywhere, Category = "Settings|Advanced", meta=(PCGNoHash, EditCondition="!IsStandaloneGraph() && !IsEmbeddedSubgraph()", EditConditionHides))
 	bool bIgnoreLandscapeTracking = false;
 
 	UPROPERTY(EditAnywhere, Category = Customization, meta=(PCGNoHash, EditCondition="ShowGraphCustomization()", EditConditionHides))
 	FPCGGraphEditorCustomization GraphCustomization;
 
 	/** Contains the data relevant for PCG Editor Mode usage. */
-	UPROPERTY(BlueprintReadWrite, EditAnywhere, Category = AssetInfo, AssetRegistrySearchable, meta=(PCGNoHash))
+	UPROPERTY(BlueprintReadWrite, EditAnywhere, Category = AssetInfo, AssetRegistrySearchable, meta=(PCGNoHash, EditCondition="ShowToolData()", EditConditionHides))
 	FPCGGraphToolData ToolData;
 
 	/** When enabled, this graph can be executed outside of the world using an editor execution source. */
@@ -347,7 +403,7 @@ public:
 	bool bIsStandaloneGraph = false;
 #endif
 
-	UPROPERTY(EditAnywhere, Category = "Settings|Advanced", meta=(EditCondition="!IsStandaloneGraph()", EditConditionHides))
+	UPROPERTY(EditAnywhere, Category = "Settings|Advanced", meta=(EditCondition="!IsStandaloneGraph() && !IsEmbeddedSubgraph()", EditConditionHides))
 	bool bLandscapeUsesMetadata = true;
 
 	/** Creates a node using the given settings interface. Does not manage ownership - done outside of this method. */
@@ -448,6 +504,8 @@ public:
 	bool UseActorComponentlessGeneration() const;
 
 #if WITH_EDITOR
+	UPCGEditorGraph* GetEditorGraph() const { return PCGEditorGraph; }
+
 	/** Can be overriden by child class to disable debug globally on all settings. */
 	virtual bool ShouldDisplayDebuggingProperties() const { return true; };
 	
@@ -488,7 +546,7 @@ public:
 	/** Overridable function for child classes to have a graph-wide node title sanitization when the title of a node changes. Returns true if the title have changed. */
 	virtual bool SanitizeNodeTitle(FName& InOutTitle, const UPCGNode* InNode) const { return false; }
 
-	const TArray<TSubclassOf<UObject>>& GetChangeTrackingIgnoredClasses() const { return ChangeTrackingIgnoredClasses; }
+	UE_API const TArray<TSubclassOf<UObject>>& GetChangeTrackingIgnoredClasses() const;
 #endif
 
 #if WITH_EDITOR
@@ -512,13 +570,19 @@ public:
 protected:
 #if WITH_EDITOR
 	UFUNCTION()
-	virtual bool CanToggleStandaloneGraph() const { return true; }
+	virtual bool CanToggleStandaloneGraph() const { return !IsEmbeddedSubgraph(); }
 
 	UFUNCTION()
 	virtual bool SupportHierarchicalGeneration() const { return !IsStandaloneGraph(); }
 
 	UFUNCTION()
-	virtual bool ShowGraphCustomization() const { return true; }
+	virtual bool CanEditHierarchicalGeneration() const { return SupportHierarchicalGeneration() && !IsEmbeddedSubgraph(); }
+
+	UFUNCTION()
+	virtual bool ShowGraphCustomization() const { return !IsEmbeddedSubgraph(); }
+
+	UFUNCTION()
+	virtual bool ShowToolData() const { return !IsEmbeddedSubgraph(); }
 
 	virtual bool IsTemplatePropertyEnabled() const override { return IsAsset(); }
 
@@ -577,6 +641,9 @@ protected:
 	UPROPERTY(EditAnywhere, Category = "Graph Parameters", DisplayName = "Parameters", meta = (NoResetToDefault, DefaultType = "EPropertyBagPropertyType::Double", IsPinTypeAccepted = "UserParametersIsPinTypeAccepted", CanRemoveProperty = "UserParametersCanRemoveProperty", ChildRowFeatures = "Extended", PropertyBagSchemaClass="/Script/PCGEditor.PCGGraphParametersSchema"))
 	FInstancedPropertyBag UserParameters;
 
+	UPROPERTY()
+	TArray<TObjectPtr<UPCGGraph>> EmbeddedSubgraphs;
+
 #if WITH_EDITOR
 	UFUNCTION(BlueprintInternalUseOnly)
 	UE_API bool UserParametersIsPinTypeAccepted(FEdGraphPinType InPinType, bool bIsChild);
@@ -587,21 +654,21 @@ protected:
 
 	UE_API virtual FInstancedPropertyBag* GetMutableUserParametersStruct() override;
 
-	UPROPERTY(EditAnywhere, Category = Settings, meta=(EditCondition="SupportHierarchicalGeneration", EditConditionHides))
+	UPROPERTY(EditAnywhere, Category = Settings, meta=(EditCondition="CanEditHierarchicalGeneration()", EditConditionHides))
 	bool bUseHierarchicalGeneration = false;
 
-	UPROPERTY(EditAnywhere, Category = Settings, meta = (DisplayName = "HiGen Default Grid Size", EditCondition = "SupportHierarchicalGeneration && bUseHierarchicalGeneration", EditConditionHides))
+	UPROPERTY(EditAnywhere, Category = Settings, meta = (DisplayName = "HiGen Default Grid Size", EditCondition = "CanEditHierarchicalGeneration() && bUseHierarchicalGeneration", EditConditionHides))
 	EPCGHiGenGrid HiGenGridSize = EPCGHiGenGrid::Grid256;
 
 	/** Shifts the grid sizes upwards based on the value, which allows to use larger grids. A value of 1 will effectively use the graph's Grid-400 values x 2 for the actual Grid-800 sizes and so on. */
 	UPROPERTY(EditAnywhere, Category = Settings, meta = (ClampMax = "10", DisplayName = "HiGen Grid Size Exponent", EditCondition = "SupportHierarchicalGeneration && bUseHierarchicalGeneration", EditConditionHides))
 	uint32 HiGenExponential = 0;
 
-	UPROPERTY(EditAnywhere, Category = Settings, meta = (DisplayName = "2D Grid", EditCondition = "SupportHierarchicalGeneration && bUseHierarchicalGeneration", EditConditionHides))
+	UPROPERTY(EditAnywhere, Category = Settings, meta = (DisplayName = "2D Grid", EditCondition = "CanEditHierarchicalGeneration() && bUseHierarchicalGeneration", EditConditionHides))
 	bool bUse2DGrid = true;
 
 	/** Controls whether the default pins (In, ...) are provided data on top-level graphs automatically. Note this will eventually change to be turned off by default. */
-	UPROPERTY(BlueprintReadWrite, EditAnywhere, Category = Settings, meta=(EditCondition="!IsStandaloneGraph()", EditConditionHides))
+	UPROPERTY(BlueprintReadWrite, EditAnywhere, Category = Settings, meta=(EditCondition="!IsStandaloneGraph() && !IsEmbeddedSubgraph()", EditConditionHides))
 	bool bHasDefaultConstructedInputs = true;
 
 	/** Execution grid size for nodes. */
@@ -624,17 +691,17 @@ protected:
 	mutable PCG::FLock CachedGridInfoLock;
 
 	/** Sets whether this graph is marked as editor-only; note that the IsEditorOnly call depends on the local graph value and the value in all subgraphs, recursively. */
-	UPROPERTY(EditAnywhere, BlueprintReadWrite, Category = Cooking, meta=(EditCondition="!IsStandaloneGraph()", EditConditionHides))
+	UPROPERTY(EditAnywhere, BlueprintReadWrite, Category = Cooking, meta=(EditCondition="!IsStandaloneGraph() && !IsEmbeddedSubgraph()", EditConditionHides))
 	bool bIsEditorOnly = false;
 
 #if WITH_EDITORONLY_DATA
 	/** When true the Debug flag in the graph editor will display debug information contextually for the selected debug object. Otherwise
 	* debug information is displayed for all components using a graph (requires regenerate).
 	*/
-	UPROPERTY(EditAnywhere, Category = Debug, meta=(EditCondition="!IsStandaloneGraph()", EditConditionHides))
+	UPROPERTY(EditAnywhere, Category = Debug, meta=(EditCondition="!IsStandaloneGraph() && !IsEmbeddedSubgraph()", EditConditionHides))
 	bool bDebugFlagAppliesToIndividualComponents = true;
 
-	UPROPERTY(EditAnywhere, Category = ChangeTracking, DisplayName = "Ignored Classes")
+	UPROPERTY(EditAnywhere, Category = ChangeTracking, DisplayName = "Ignored Classes", meta = (EditCondition = "!IsEmbeddedSubgraph()", EditConditionHides))
 	TArray<TSubclassOf<UObject>> ChangeTrackingIgnoredClasses;
 #endif // WITH_EDITORONLY_DATA
 
@@ -654,12 +721,12 @@ public:
 	// Best used on a brand new PCG Graph.
 	UE_API EPropertyBagAlterationResult AddUserParameters(const TArray<FPropertyBagPropertyDesc>& InDescs, const UPCGGraph* InOptionalOriginalGraph = nullptr);
 
-	UPROPERTY(BlueprintReadWrite, EditAnywhere, Category = "Runtime Generation", meta=(EditCondition = "!IsStandaloneGraph()", EditConditionHides))
+	UPROPERTY(BlueprintReadWrite, EditAnywhere, Category = "Runtime Generation", meta=(EditCondition = "!IsStandaloneGraph() && !IsEmbeddedSubgraph()", EditConditionHides))
 	FPCGRuntimeGenerationRadii GenerationRadii;
 
 	// @todo_pcg: This could be on by default/hidden and graphs which need actors/actor-components could create a target actor on demand when necessary.
 	/** [EXPERIMENTAL] Avoids creating partition actors and local PCG components. Any functionality which depends on a target actor/component are not supported in this mode. */
-	UPROPERTY(BlueprintReadWrite, EditAnywhere, Category = "Runtime Generation|Advanced", meta=(EditCondition = "!IsStandaloneGraph()", EditConditionHides))
+	UPROPERTY(BlueprintReadWrite, EditAnywhere, Category = "Runtime Generation|Advanced", meta=(EditCondition = "!IsStandaloneGraph() && !IsEmbeddedSubgraph()", EditConditionHides))
 	bool bUseActorComponentlessGeneration = false;
 
 	UE_API virtual void OnGraphParametersChanged(EPCGGraphParameterEvent InChangeType, FName InChangedPropertyName) override;
diff --git a/Engine/Plugins/PCG/Source/PCGEditor/Private/AssetDefinitions/AssetDefinition_PCGGraph.cpp b/Engine/Plugins/PCG/Source/PCGEditor/Private/AssetDefinitions/AssetDefinition_PCGGraph.cpp
--- a/Engine/Plugins/PCG/Source/PCGEditor/Private/AssetDefinitions/AssetDefinition_PCGGraph.cpp
+++ b/Engine/Plugins/PCG/Source/PCGEditor/Private/AssetDefinitions/AssetDefinition_PCGGraph.cpp
@@ -23,12 +23,7 @@ TSoftClassPtr<UObject> UAssetDefinition_PCGGraph::GetAssetClass() const
 
 EAssetCommandResult UAssetDefinition_PCGGraph::OpenAssets(const FAssetOpenArgs& OpenArgs) const
 {
-	for (UPCGGraph* PCGGraph : OpenArgs.LoadObjects<UPCGGraph>())
-	{
-		const TSharedRef<FPCGEditor> PCGEditor = MakeShared<FPCGEditor>();
-		PCGEditor->Initialize(EToolkitMode::Standalone, OpenArgs.ToolkitHost, PCGGraph);
-	}
-
+	FPCGEditor::OpenAssets(OpenArgs);
 	return EAssetCommandResult::Handled;
 }
 
diff --git a/Engine/Plugins/PCG/Source/PCGEditor/Private/Details/PCGGraphDetails.cpp b/Engine/Plugins/PCG/Source/PCGEditor/Private/Details/PCGGraphDetails.cpp
--- a/Engine/Plugins/PCG/Source/PCGEditor/Private/Details/PCGGraphDetails.cpp
+++ b/Engine/Plugins/PCG/Source/PCGEditor/Private/Details/PCGGraphDetails.cpp
@@ -96,7 +96,7 @@ void FPCGGraphDetails::CustomizeDetails(IDetailLayoutBuilder& DetailBuilder)
 			];
 
 		// Add a "Run Graph Determinism Test" button in the debug category, if we're not a standalone graph
-		if (!Algo::AnyOf(SelectedGraphs, [](const TWeakObjectPtr<UPCGGraph> Graph) { return !Graph.IsValid() || Graph->IsStandaloneGraph(); }))
+		if (!Algo::AnyOf(SelectedGraphs, [](const TWeakObjectPtr<UPCGGraph> Graph) { return !Graph.IsValid() || Graph->IsStandaloneGraph() || Graph->IsEmbeddedSubgraph(); }))
 		{
 			IDetailCategoryBuilder& DebugCategory = DetailBuilder.EditCategory(TEXT("Debug"));
 
diff --git a/Engine/Plugins/PCG/Source/PCGEditor/Private/Nodes/PCGEditorGraphNodeBase.cpp b/Engine/Plugins/PCG/Source/PCGEditor/Private/Nodes/PCGEditorGraphNodeBase.cpp
--- a/Engine/Plugins/PCG/Source/PCGEditor/Private/Nodes/PCGEditorGraphNodeBase.cpp
+++ b/Engine/Plugins/PCG/Source/PCGEditor/Private/Nodes/PCGEditorGraphNodeBase.cpp
@@ -360,7 +360,16 @@ void UPCGEditorGraphNodeBase::GetNodeContextMenuActions(UToolMenu* Menu, class U
 		Section.AddMenuEntry(FPCGEditorCommands::Get().ToggleInspect, LOCTEXT("ToggleinspectionLabel", "Inspect"));
 		Section.AddMenuEntry(FGraphEditorCommands::Get().BreakNodeLinks);
 		Section.AddMenuEntry(FPCGEditorCommands::Get().ExportNodes, TAttribute<FText>(), TAttribute<FText>(), FSlateIcon(FPCGEditorStyle::Get().GetStyleSetName(), "ClassIcon.PCGSettings"));
-		Section.AddMenuEntry(FPCGEditorCommands::Get().CollapseNodes);
+		
+		auto CollapseEntries = [this](UToolMenu* CollapseMenu)
+		{
+			FToolMenuSection& CollapseSubSection = CollapseMenu->AddSection(TEXT("CollapseSection"));
+			CollapseSubSection.AddMenuEntry(FPCGEditorCommands::Get().CollapseNodes);
+			CollapseSubSection.AddMenuEntry(FPCGEditorCommands::Get().CollapseNodesToEmbeddedSubgraph);
+		};
+
+		Section.AddSubMenu(TEXT("CollapseSubMenu"), LOCTEXT("CollapseSubMenuLabel", "Collapse Nodes"), FText(), FNewToolMenuDelegate::CreateLambda(CollapseEntries));
+
 		Section.AddMenuEntry(FPCGEditorCommands::Get().ConvertToStandaloneNodes);
 		Section.AddMenuEntry(FPCGEditorCommands::Get().JumpToSource);
 	}
diff --git a/Engine/Plugins/PCG/Source/PCGEditor/Private/PCGEditor.cpp b/Engine/Plugins/PCG/Source/PCGEditor/Private/PCGEditor.cpp
--- a/Engine/Plugins/PCG/Source/PCGEditor/Private/PCGEditor.cpp
+++ b/Engine/Plugins/PCG/Source/PCGEditor/Private/PCGEditor.cpp
@@ -9,6 +9,7 @@
 #include "PCGGraph.h"
 #include "PCGGraphFactory.h"
 #include "PCGPin.h"
+#include "PCGSubgraph.h"
 #include "PCGSubsystem.h"
 #include "Editor/IPCGEditorModule.h"
 #include "Elements/PCGReroute.h"
@@ -41,13 +42,16 @@
 #include "Widgets/SPCGEditorGraphDetailsView.h"
 #include "Widgets/SPCGEditorGraphDeterminism.h"
 #include "Widgets/SPCGEditorGraphFind.h"
+#include "Widgets/SPCGEditorGraphEmbeddedSubgraphsView.h"
 #include "Widgets/SPCGEditorGraphLogView.h"
 #include "Widgets/SPCGEditorGraphNodePalette.h"
 #include "Widgets/SPCGEditorGraphParamsView.h"
 #include "Widgets/SPCGEditorGraphProfilingView.h"
+#include "Widgets/SPCGEditorGraphTitleBar.h"
 #include "Widgets/SPCGEditorNodeSource.h"
 #include "Widgets/AssetEditorViewport/SPCGEditorViewport.h"
 
+#include "AssetDefinition.h"
 #include "AssetToolsModule.h"
 #include "EdGraphUtilities.h"
 #include "EditorAssetLibrary.h"
@@ -80,6 +84,7 @@
 #include "Widgets/Notifications/SNotificationList.h"
 #include "WorkflowOrientedApp/WorkflowTabManager.h"
 #include "WorkflowOrientedApp/WorkflowTabFactory.h"
+#include "WorkflowOrientedApp/WorkflowUObjectDocuments.h"
 
 #define LOCTEXT_NAMESPACE "PCGGraphEditor"
 
@@ -107,7 +112,8 @@ namespace FPCGEditor_private
 		PCGEditorTabs::ViewportID[0],
 		PCGEditorTabs::ViewportID[1],
 		PCGEditorTabs::ViewportID[2],
-		PCGEditorTabs::ViewportID[3]
+		PCGEditorTabs::ViewportID[3],
+		PCGEditorTabs::EmbeddedSubgraphsID,
 	};
 
 	const FName AllRestrictedPanels[] =
@@ -133,6 +139,40 @@ UPCGEditorGraph* FPCGEditor::GetPCGEditorGraph(const UPCGSettings* InSettings)
 	return GetPCGEditorGraph(PCGNode);
 }
 
+void FPCGEditor::InitializePCGEditorGraph(UPCGGraph* InGraph)
+{
+	// Initializes the UPCGEditorGraph if needed
+	const TSubclassOf<UPCGEditorGraphSchema> SchemaClass = GetSchemaClass();
+	const bool bShouldUpdateSchema = !InGraph->PCGEditorGraph || InGraph->PCGEditorGraph->Schema != SchemaClass || !InGraph->PCGEditorGraph->GetPCGGraph();
+	if (!InGraph->PCGEditorGraph)
+	{
+		InGraph->PCGEditorGraph = NewObject<UPCGEditorGraph>(InGraph, UPCGEditorGraph::StaticClass(), NAME_None, RF_Transactional | RF_Transient);
+	}
+
+	if (bShouldUpdateSchema)
+	{
+		InGraph->PCGEditorGraph->Schema = SchemaClass;
+		InGraph->PCGEditorGraph->InitFromNodeGraph(InGraph);
+	}
+
+	// Only set editor if this graph is owned by us
+	if (InGraph == MainGraph || InGraph->GetEmbeddedParentGraph() == MainGraph)
+	{
+		InGraph->PCGEditorGraph->SetEditor(SharedThis(this));
+	}
+}
+
+UPCGEditorGraph* FPCGEditor::GetOrCreatePCGEditorGraph(UPCGGraph* InPCGGraph)
+{
+	UPCGEditorGraph* PCGEditorGraph = GetPCGEditorGraph(InPCGGraph);
+	if (!PCGEditorGraph)
+	{
+		InitializePCGEditorGraph(InPCGGraph);
+	}
+	
+	return GetPCGEditorGraph(InPCGGraph);
+}
+
 void FPCGEditor::Initialize(const EToolkitMode::Type InMode, const TSharedPtr<class IToolkitHost>& InToolkitHost, UPCGGraph* InPCGGraph, UObject* InObjectToEdit)
 {
 	InObjectToEdit = InObjectToEdit != nullptr ? InObjectToEdit : InPCGGraph;
@@ -146,34 +186,18 @@ void FPCGEditor::Initialize(const EToolkitMode::Type InMode, const TSharedPtr<cl
 
 	FPCGModule::GetGraphExecutionRegistry().OnGraphExecutionSourcesChanged().AddRaw(this, &FPCGEditor::OnGraphExecutionSourcesChanged);
 
-	PCGGraphBeingEdited = InPCGGraph;
+	MainGraph = InPCGGraph;
 
 	UpdateDefaultExecutionSource();
 
-	// Initializes the UPCGEditorGraph if needed
-	const TSubclassOf<UPCGEditorGraphSchema> SchemaClass = GetSchemaClass();
-	const bool bShouldUpdateSchema = !InPCGGraph->PCGEditorGraph || InPCGGraph->PCGEditorGraph->Schema != SchemaClass;
-	if (!InPCGGraph->PCGEditorGraph)
-	{
-		InPCGGraph->PCGEditorGraph = NewObject<UPCGEditorGraph>(InPCGGraph, UPCGEditorGraph::StaticClass(), NAME_None, RF_Transactional | RF_Transient);
-	}
-
-	if (bShouldUpdateSchema)
-	{
-		InPCGGraph->PCGEditorGraph->Schema = SchemaClass;
-		InPCGGraph->PCGEditorGraph->InitFromNodeGraph(InPCGGraph);
-	}
-
-	PCGGraphBeingEdited->PCGEditorGraph->SetEditor(SharedThis(this));
-	PCGEditorGraph = PCGGraphBeingEdited->PCGEditorGraph;
+	InitializePCGEditorGraph(MainGraph);
 
-	GetOrCreateGraphEditorWidget();
 	GetOrCreateNodeSourceWidget();
-	
+
 	BindCommands();
 
-	PCGGraphBeingEdited->OnGraphChangedDelegate.AddRaw(this, &FPCGEditor::OnGraphChanged);
-	PCGGraphBeingEdited->OnNodeSourceCompiledDelegate.AddRaw(this, &FPCGEditor::OnNodeSourceCompiled);
+	MainGraph->OnGraphChangedDelegate.AddRaw(this, &FPCGEditor::OnGraphChanged);
+	MainGraph->OnNodeSourceCompiledDelegate.AddRaw(this, &FPCGEditor::OnNodeSourceCompiled);
 
 	// Hook to map change / delete actor to refresh debug object selection list, to help prevent it going stale.
 	FLevelEditorModule& LevelEditor = FModuleManager::LoadModuleChecked<FLevelEditorModule>("LevelEditor");
@@ -197,7 +221,7 @@ void FPCGEditor::Initialize(const EToolkitMode::Type InMode, const TSharedPtr<cl
 	}
 
 	// Clear inspection flag on all nodes.
-	for (UEdGraphNode* EdGraphNode : PCGEditorGraph->Nodes)
+	for (UEdGraphNode* EdGraphNode : GetMainEditorGraph()->Nodes)
 	{
 		if (UPCGEditorGraphNodeBase* PCGEditorGraphNode = Cast<UPCGEditorGraphNodeBase>(EdGraphNode))
 		{
@@ -205,23 +229,129 @@ void FPCGEditor::Initialize(const EToolkitMode::Type InMode, const TSharedPtr<cl
 		}
 	}
 
-	DocumentManager = MakeShared<FDocumentTracker>(PCGEditorTabs::PCGEditorDocumentID);
+	DocumentManager = MakeShared<FDocumentTracker>(PCGEditorTabs::GraphEditorID);
 
 	const FName PCGGraphEditorAppName = FName(TEXT("PCGEditorApp"));
 	InitAssetEditor(InMode, InToolkitHost, PCGGraphEditorAppName, FTabManager::FLayout::NullLayout, /*bCreateDefaultStandaloneMenu=*/ true, /*bCreateDefaultToolbar=*/ true, InObjectToEdit);
 
 	TSharedPtr<FPCGEditor> ThisPtr = SharedThis<FPCGEditor>(this);
 	DocumentManager->Initialize(AsShared());
+	TSharedPtr<FDocumentTabFactory> DocumentTabFactoryPtr = MakeShareable(new FPCGGraphEditorDocumentTabFactory(ThisPtr,
+		FPCGGraphEditorDocumentTabFactory::FOnCreateGraphEditorWidget::CreateSP(this, &FPCGEditor::CreateGraphEditorWidget)));
 
-	// @todo_pcg: add document factories here
+	DocumentTabFactory = DocumentTabFactoryPtr;
+	DocumentManager->RegisterDocumentFactory(DocumentTabFactoryPtr);
 
 	AddApplicationMode(FPCGEditorDefaultMode::StaticName(), MakeShared<FPCGEditorDefaultMode>(ThisPtr));
 	SetCurrentMode(FPCGEditorDefaultMode::StaticName());
 }
 
-UPCGEditorGraph* FPCGEditor::GetPCGEditorGraph()
+UPCGGraph* FPCGEditor::GetMainGraph() const
 {
-	return PCGEditorGraph;
+	return MainGraph;
+}
+
+UPCGGraph* FPCGEditor::GetFocusedGraph() const
+{
+	if (TSharedPtr<SGraphEditor> GraphEditor = GetFocusedGraphEditorWidget())
+	{
+		if (UPCGEditorGraph* EditorGraph = Cast<UPCGEditorGraph>(GraphEditor->GetCurrentGraph()))
+		{
+			return EditorGraph->GetPCGGraph();
+		}
+	}
+
+	return GetMainGraph();
+}
+
+UPCGEditorGraph* FPCGEditor::GetMainEditorGraph() const
+{
+	return MainGraph ? MainGraph->PCGEditorGraph : nullptr;
+}
+
+UPCGEditorGraph* FPCGEditor::GetFocusedEditorGraph() const
+{
+	UPCGGraph* FocusedGraph = GetFocusedGraph();
+	return FocusedGraph ? FocusedGraph->PCGEditorGraph.Get() : GetMainEditorGraph();
+}
+
+void FPCGEditor::NavigateTab(FDocumentTracker::EOpenDocumentCause InCause)
+{
+	TSharedRef<FTabPayload_UObject> Payload = FTabPayload_UObject::Make(nullptr);
+	DocumentManager->OpenDocument(Payload, InCause);
+}
+
+TSharedPtr<SDockTab> FPCGEditor::OpenDocument(UPCGGraph* InPCGGraph, FDocumentTracker::EOpenDocumentCause InCause)
+{
+	check(InPCGGraph);
+
+	InitializePCGEditorGraph(InPCGGraph);
+
+	UPCGEditorGraph* PCGEditorGraph = InPCGGraph->GetEditorGraph();
+	if (!PCGEditorGraph)
+	{
+		return nullptr;
+	}
+
+	TSharedRef<FTabPayload_UObject> Payload = FTabPayload_UObject::Make(PCGEditorGraph);
+	return DocumentManager->OpenDocument(Payload, InCause);
+}
+
+void FPCGEditor::CloseDocument(UPCGGraph* InPCGGraph)
+{
+	check(InPCGGraph);
+
+	UPCGEditorGraph* PCGEditorGraph = InPCGGraph->GetEditorGraph();
+	if (!PCGEditorGraph)
+	{
+		return;
+	}
+
+	TSharedRef<FTabPayload_UObject> Payload = FTabPayload_UObject::Make(PCGEditorGraph);
+	DocumentManager->CloseTab(Payload);
+}
+
+FPCGEditor* FPCGEditor::OpenEditorForGraph(UPCGGraph* InPCGGraph)
+{
+	if (MainGraph == InPCGGraph || MainGraph->ContainsEmbeddedSubgraph(InPCGGraph))
+	{
+		OpenDocument(InPCGGraph, FDocumentTracker::EOpenDocumentCause::OpenNewDocument);
+		return this;
+	}
+	else if(GEditor->GetEditorSubsystem<UAssetEditorSubsystem>()->OpenEditorForAsset(InPCGGraph))
+	{
+		TSharedPtr<FPCGEditor> PCGEditor = InPCGGraph->GetEditorGraph() && InPCGGraph->GetEditorGraph()->GetEditor().IsValid() ? InPCGGraph->GetEditorGraph()->GetEditor().Pin() : nullptr;
+		PCGEditor->OpenDocument(InPCGGraph, FDocumentTracker::OpenNewDocument);
+		return PCGEditor.Get();
+	}
+
+	return nullptr;
+}
+
+void FPCGEditor::OpenAssets(const FAssetOpenArgs& OpenArgs)
+{
+	for (UPCGGraph* PCGGraph : OpenArgs.LoadObjects<UPCGGraph>())
+	{
+		UPCGGraph* GraphToOpen = PCGGraph;
+		if (UPCGGraph* MainGraph = PCGGraph->GetEmbeddedParentGraph())
+		{
+			GraphToOpen = MainGraph;
+		}
+
+		TSharedPtr<FPCGEditor> PCGEditor = GraphToOpen->GetEditorGraph() && GraphToOpen->GetEditorGraph()->GetEditor().IsValid() ? GraphToOpen->GetEditorGraph()->GetEditor().Pin() : nullptr;
+		if (!PCGEditor)
+		{
+			PCGEditor = MakeShared<FPCGEditor>();
+			PCGEditor->Initialize(EToolkitMode::Standalone, OpenArgs.ToolkitHost, GraphToOpen);
+			GraphToOpen->GetEditorGraph()->SetEditor(PCGEditor);
+		}
+		
+		// Now that we have the editor, if we are opening an embedded subgraph, open the tab
+		if (PCGGraph != GraphToOpen)
+		{
+			PCGEditor->OpenDocument(PCGGraph, FDocumentTracker::OpenNewDocument);
+		}
+	}
 }
 
 void FPCGEditor::SetStackBeingInspectedFromAnotherEditor(const FPCGStack& FullStack)
@@ -275,11 +405,11 @@ void FPCGEditor::UpdateAfterInspectedStackChanged(const FPCGStack& FullStack)
 
 	if (bValidStack)
 	{
-		if (PCGGraphBeingEdited)
+		if (MainGraph)
 		{
-			PCGGraphBeingEdited->EnableInspection(FullStack);
+			MainGraph->EnableInspection(FullStack);
 		}
-		
+
 		if (Source)
 		{
 			// Implementation note: if we're inspecting and have not pre-run the graph, then it probably makes sense to enable inspection by default. 
@@ -301,13 +431,13 @@ void FPCGEditor::UpdateAfterInspectedStackChanged(const FPCGStack& FullStack)
 			}
 		}
 	}
-	else if (PCGGraphBeingEdited)
+	else if (MainGraph)
 	{
-		PCGGraphBeingEdited->DisableInspection();
+		MainGraph->DisableInspection();
 	}
 
-	check(PCGEditorGraph);
-	for (UEdGraphNode* Node : PCGEditorGraph->Nodes)
+	check(GetMainEditorGraph());
+	for (UEdGraphNode* Node : GetMainEditorGraph()->Nodes)
 	{
 		if (UPCGEditorGraphNodeBase* PCGNode = Cast<UPCGEditorGraphNodeBase>(Node))
 		{
@@ -344,7 +474,7 @@ IPCGGraphExecutionSource* FPCGEditor::GetPCGSourceBeingInspected() const
 
 void FPCGEditor::UpdateDebugAfterSourceSelection(IPCGGraphExecutionSource* InOldSource, IPCGGraphExecutionSource* InNewSource, bool bInNewSourceStartedInspecting)
 {
-	if (!ensure(PCGGraphBeingEdited))
+	if (!ensure(MainGraph))
 	{
 		return;
 	}
@@ -374,8 +504,8 @@ void FPCGEditor::UpdateDebugAfterSourceSelection(IPCGGraphExecutionSource* InOld
 		}
 	};
 
-	// If individual source debugging is disabled, just generate the new source if required.
-	if (!PCGGraphBeingEdited->DebugFlagAppliesToIndividualComponents())
+	// If individual component debugging is disabled, just generate the new component if required.
+	if (!MainGraph->DebugFlagAppliesToIndividualComponents())
 	{
 		if (InNewSource && bInNewSourceStartedInspecting)
 		{
@@ -397,7 +527,7 @@ void FPCGEditor::UpdateDebugAfterSourceSelection(IPCGGraphExecutionSource* InOld
 	}
 	else
 	{
-		const bool bDebugFlagSetOnAnyNode = Algo::AnyOf(PCGGraphBeingEdited->GetNodes(), [](const UPCGNode* InNode)
+		const bool bDebugFlagSetOnAnyNode = Algo::AnyOf(MainGraph->GetNodes(), [](const UPCGNode* InNode)
 		{
 			return InNode && InNode->GetSettings() && InNode->GetSettings()->bDebug;
 		});
@@ -434,20 +564,26 @@ void FPCGEditor::SetSourceEditorTargetObject(UObject* InObject)
 
 void FPCGEditor::JumpToNode(const UEdGraphNode* InNode)
 {
-	if (GraphEditorWidget.IsValid())
+	if (TSharedPtr<SGraphEditor> FocusedGraphWidget = GetFocusedGraphEditorWidget())
 	{
-		GraphEditorWidget->JumpToNode(InNode);
+		FocusedGraphWidget->JumpToNode(InNode);
 	}
 }
 
 UPCGEditorGraphNodeBase* FPCGEditor::GetEditorNode(const UPCGNode* InNode)
 {
-	if (!ensure(PCGEditorGraph) || !InNode)
+	if (!InNode)
 	{
 		return nullptr;
 	}
 
-	for (UEdGraphNode* EdGraphNode : PCGEditorGraph->Nodes)
+	UPCGGraph* OuterGraph = InNode->GetTypedOuter<UPCGGraph>();
+	if (!OuterGraph || !OuterGraph->GetEditorGraph())
+	{
+		return nullptr;
+	}
+
+	for (UEdGraphNode* EdGraphNode : OuterGraph->GetEditorGraph()->Nodes)
 	{
 		if (UPCGEditorGraphNodeBase* PCGEdGraphNode = Cast<UPCGEditorGraphNodeBase>(EdGraphNode))
 		{
@@ -551,12 +687,17 @@ bool FPCGEditor::IsPanelCurrentlyOpen(FName PanelID) const
 
 bool FPCGEditor::IsPanelCurrentlyForeground(FName PanelID) const
 {
-		const TSharedPtr<SDockTab> DockTab = TabManager.IsValid() ? TabManager->FindExistingLiveTab(PanelID) : nullptr;
-		return DockTab.IsValid() && DockTab->IsForeground();
+	const TSharedPtr<SDockTab> DockTab = TabManager.IsValid() ? TabManager->FindExistingLiveTab(PanelID) : nullptr;
+	return DockTab.IsValid() && DockTab->IsForeground();
 }
 
 bool FPCGEditor::IsPanelAvailable(const FName PanelID) const
 {
+	if (PanelID == PCGEditorTabs::EmbeddedSubgraphsID && (!GetMainGraph() || !GetMainGraph()->SupportsEmbeddedSubgraphs()))
+	{
+		return false;
+	}
+
 	bool bIsPanelAvailable = Algo::AnyOf(FPCGEditor_private::AllDefaultPanels, [PanelID](const FName& It) { return It == PanelID; });
 
 	IPCGEditorModule* PCGEditorModule = IPCGEditorModule::Get();
@@ -579,7 +720,7 @@ void FPCGEditor::RegisterTabSpawners(const TSharedRef<FTabManager>& InTabManager
 
 void FPCGEditor::AddReferencedObjects(FReferenceCollector& Collector)
 {
-	Collector.AddReferencedObject(PCGGraphBeingEdited);
+	Collector.AddReferencedObject(MainGraph);
 	Collector.AddReferencedObject(PCGDefaultExecutionSource);
 
 	InspectionDataManager.AddReferencedObjects(Collector);
@@ -606,7 +747,7 @@ bool FPCGEditor::MatchesContext(const FTransactionContext& InContext, const TArr
 		const UObject* Object = TransactionObjectContext.Key;
 		while (Object != nullptr)
 		{
-			if (Object == PCGGraphBeingEdited)
+			if (Object == MainGraph)
 			{
 				return true;
 			}
@@ -621,16 +762,19 @@ void FPCGEditor::PostUndo(bool bSuccess)
 {
 	if (bSuccess)
 	{
-		if (PCGGraphBeingEdited)
+		if (MainGraph)
 		{
 			// Deepest change type to catch all types of change (like redoing adding a grid size node or etc).
-			PCGGraphBeingEdited->NotifyGraphChanged(EPCGChangeType::Structural | EPCGChangeType::GenerationGrid);
+			MainGraph->NotifyGraphChanged(EPCGChangeType::Structural | EPCGChangeType::GenerationGrid);
 		}
 
-		if (GraphEditorWidget.IsValid())
+		CloseInvalidDocuments();
+		DocumentManager->RefreshAllTabs();
+
+		if (TSharedPtr<SGraphEditor> FocusedGraphWidget = GetFocusedGraphEditorWidget())
 		{
-			GraphEditorWidget->ClearSelectionSet();
-			GraphEditorWidget->NotifyGraphChanged();
+			FocusedGraphWidget->ClearSelectionSet();
+			FocusedGraphWidget->NotifyGraphChanged();
 
 			FSlateApplication::Get().DismissAllMenus();
 		}
@@ -639,7 +783,7 @@ void FPCGEditor::PostUndo(bool bSuccess)
 
 FText FPCGEditor::GetToolkitName() const
 {
-	return PCGGraphBeingEdited->GetDisplayName();
+	return MainGraph ? MainGraph->GetDisplayName() : FWorkflowCentricApplication::GetToolkitName();
 }
 
 FName FPCGEditor::GetToolkitFName() const
@@ -679,6 +823,7 @@ void FPCGEditor::RegisterToolbarInternal(FToolBarBuilder& ToolbarBuilder) const
 
 	ToolbarBuilder.AddSeparator(NAME_None);
 
+	RegisterToolbarButton(ToolbarBuilder, EPCGToolbarButtons::Graph);
 	RegisterToolbarButton(ToolbarBuilder, EPCGToolbarButtons::GraphParams);
 	RegisterToolbarButton(ToolbarBuilder, EPCGToolbarButtons::GraphSettings);
 }
@@ -742,11 +887,19 @@ void FPCGEditor::RegisterToolbarButton(FToolBarBuilder& ToolbarBuilder, EPCGTool
 				FSlateIcon(FPCGEditorStyle::Get().GetStyleSetName(), "PCG.Command.GraphSettings"));
 			break;
 		}
+	case EPCGToolbarButtons::Graph:
+		{
+			ToolbarBuilder.AddToolBarButton(PCGEditorCommands.EditGraph, NAME_None, TAttribute<FText>(), TAttribute<FText>(),
+				FSlateIcon(FPCGEditorStyle::Get().GetStyleSetName(), "PCG.Command.EditGraph"));
+			break;
+		}
 	}
 }
 
 void FPCGEditor::BindCommands()
 {
+	CreateGraphEditorCommands();
+
 	const FPCGEditorCommands& PCGEditorCommands = FPCGEditorCommands::Get();
 
 	ToolkitCommands->MapAction(
@@ -796,10 +949,20 @@ void FPCGEditor::BindCommands()
 		FCanExecuteAction(),
 		FIsActionChecked::CreateSP(this, &FPCGEditor::IsEditGraphSettingsToggled));
 
+	ToolkitCommands->MapAction(
+		PCGEditorCommands.EditGraph,
+		FExecuteAction::CreateSP(this, &FPCGEditor::OnEditGraph_Clicked),
+		FCanExecuteAction::CreateSP(this, &FPCGEditor::IsEditGraphTabClosed));
+
 	GraphEditorCommands->MapAction(
 		PCGEditorCommands.CollapseNodes,
-		FExecuteAction::CreateSP(this, &FPCGEditor::OnCollapseNodesInSubgraph),
-		FCanExecuteAction::CreateSP(this, &FPCGEditor::CanCollapseNodesInSubgraph));
+		FExecuteAction::CreateSP(this, &FPCGEditor::OnCollapseNodesInSubgraph, false),
+		FCanExecuteAction::CreateSP(this, &FPCGEditor::CanCollapseNodesInSubgraph, false));
+
+	GraphEditorCommands->MapAction(
+		PCGEditorCommands.CollapseNodesToEmbeddedSubgraph,
+		FExecuteAction::CreateSP(this, &FPCGEditor::OnCollapseNodesInSubgraph, true),
+		FCanExecuteAction::CreateSP(this, &FPCGEditor::CanCollapseNodesInSubgraph, true));
 
 	GraphEditorCommands->MapAction(
 		PCGEditorCommands.ExportNodes,
@@ -983,22 +1146,22 @@ void FPCGEditor::OnViewportViewTabClosed(FName Id, int32 Index)
 
 void FPCGEditor::OnPauseAutomaticRegeneration_Clicked()
 {
-	if (!PCGGraphBeingEdited)
+	if (!MainGraph)
 	{
 		return;
 	}
 
-	PCGGraphBeingEdited->ToggleUserPausedNotificationsForEditor();
+	MainGraph->ToggleUserPausedNotificationsForEditor();
 }
 
 bool FPCGEditor::IsAutomaticRegenerationPaused() const
 {
-	return PCGGraphBeingEdited && PCGGraphBeingEdited->NotificationsForEditorArePausedByUser();
+	return MainGraph && MainGraph->NotificationsForEditorArePausedByUser();
 }
 
 void FPCGEditor::OnForceGraphRegeneration_Clicked()
 {
-	if (PCGGraphBeingEdited)
+	if (MainGraph)
 	{
 		EPCGChangeType ChangeType = EPCGChangeType::Structural;
 
@@ -1015,25 +1178,25 @@ void FPCGEditor::OnForceGraphRegeneration_Clicked()
 			ChangeType |= EPCGChangeType::ShaderSource;
 		}
 
-		PCGGraphBeingEdited->ForceNotificationForEditor(ChangeType);
+		MainGraph->ForceNotificationForEditor(ChangeType);
 	}
 }
 
 void FPCGEditor::OnCancelExecution_Clicked()
 {
 	IPCGBaseSubsystem* Subsystem = GetSubsystem();
-	if (PCGEditorGraph && Subsystem)
+	if (MainGraph && Subsystem)
 	{
-		Subsystem->CancelGeneration(PCGEditorGraph->GetPCGGraph());
+		Subsystem->CancelGeneration(MainGraph);
 	}
 }
 
 bool FPCGEditor::IsCurrentlyGenerating() const
 {
 	IPCGBaseSubsystem* Subsystem = GetSubsystem();
-	if (PCGGraphBeingEdited && Subsystem)
+	if (MainGraph && Subsystem)
 	{
-		return Subsystem->IsGraphCurrentlyExecuting(PCGGraphBeingEdited);
+		return Subsystem->IsGraphCurrentlyExecuting(MainGraph);
 	}
 
 	return false;
@@ -1049,11 +1212,27 @@ void FPCGEditor::OnOpenDebugObjectTreeTab_Clicked()
 	TabManager->TryInvokeTab(FTabId(PCGEditorTabs::DebugObjectID));
 }
 
+void FPCGEditor::OnEditGraph_Clicked()
+{
+	OpenDocument(GetMainGraph(), FDocumentTracker::EOpenDocumentCause::OpenNewDocument);
+}
+
+bool FPCGEditor::IsEditGraphTabClosed() const
+{
+	TSharedRef<FTabPayload_UObject> Payload = FTabPayload_UObject::Make(GetMainEditorGraph());
+
+	TArray<TSharedPtr<SDockTab>> OpenedTabs;
+	DocumentManager->FindMatchingTabs(Payload, OpenedTabs);
+
+	return OpenedTabs.IsEmpty();
+}
+
 bool FPCGEditor::CanRunDeterminismNodeTest() const
 {
-	check(GraphEditorWidget.IsValid());
+	TSharedPtr<SGraphEditor> FocusedGraphWidget = GetFocusedGraphEditorWidget();
+	check(FocusedGraphWidget.IsValid());
 
-	for (const UObject* Object : GraphEditorWidget->GetSelectedNodes())
+	for (const UObject* Object : FocusedGraphWidget->GetSelectedNodes())
 	{
 		if (Cast<const UPCGEditorGraphNodeBase>(Object) && !Cast<const UPCGEditorGraphNodeInput>(Object) && !Cast<const UPCGEditorGraphNodeOutput>(Object))
 		{
@@ -1066,7 +1245,11 @@ bool FPCGEditor::CanRunDeterminismNodeTest() const
 
 void FPCGEditor::OnDeterminismNodeTest() const
 {
-	check(GraphEditorWidget.IsValid());
+	TSharedPtr<SGraphEditor> FocusedGraphWidget = GetFocusedGraphEditorWidget();
+	check(FocusedGraphWidget.IsValid());
+
+	// Lazy create
+	const_cast<FPCGEditor*>(this)->GetOrCreateDeterminismWidget();
 
 	if (!DeterminismWidget.IsValid() || !DeterminismWidget->WidgetIsConstructed())
 	{
@@ -1078,7 +1261,7 @@ void FPCGEditor::OnDeterminismNodeTest() const
 	DeterminismWidget->BuildBaseColumns();
 
 	int64 TestIndex = 0;
-	for (UObject* Object : GraphEditorWidget->GetSelectedNodes())
+	for (UObject* Object : FocusedGraphWidget->GetSelectedNodes())
 	{
 		check(Object);
 
@@ -1166,21 +1349,25 @@ void FPCGEditor::OnDeterminismNodeTest() const
 
 bool FPCGEditor::CanRunDeterminismGraphTest() const
 {
-	return PCGEditorGraph && InspectionDataManager.GetPCGSourceBeingInspected();
+	return GetMainEditorGraph() && InspectionDataManager.GetPCGSourceBeingInspected();
 }
 
 void FPCGEditor::OnDeterminismGraphTest() const
 {
-	check(GraphEditorWidget.IsValid());
+	TSharedPtr<SGraphEditor> FocusedGraphWidget = GetFocusedGraphEditorWidget();
+	check(FocusedGraphWidget.IsValid());
+
+	// Lazy create
+	const_cast<FPCGEditor*>(this)->GetOrCreateDeterminismWidget();
 
 	const IPCGGraphExecutionSource* PCGSourceBeingInspected = InspectionDataManager.GetPCGSourceBeingInspected();
 
-	if (!DeterminismWidget.IsValid() || !DeterminismWidget->WidgetIsConstructed() || !PCGGraphBeingEdited || !PCGSourceBeingInspected)
+	if (!DeterminismWidget.IsValid() || !DeterminismWidget->WidgetIsConstructed() || !MainGraph || !PCGSourceBeingInspected)
 	{
 		return;
 	}
 
-	if (PCGSourceBeingInspected->GetExecutionState().GetGraph() != PCGGraphBeingEdited)
+	if (PCGSourceBeingInspected->GetExecutionState().GetGraph() != MainGraph)
 	{
 		// TODO: Should we alert the user more directly or disable this altogether?
 		UE_LOG(LogPCGEditor, Warning, TEXT("Running Determinism on a PCG Component with different/no attached PCG Graph"));
@@ -1195,10 +1382,10 @@ void FPCGEditor::OnDeterminismGraphTest() const
 	TSharedPtr<FDeterminismTestResult> TestResult = MakeShared<FDeterminismTestResult>();
 	TestResult->Index = 0;
 	TestResult->TestResultTitle = TEXT("Full Graph Test");
-	TestResult->TestResultName = PCGGraphBeingEdited->GetName();
+	TestResult->TestResultName = MainGraph->GetDisplayName().ToString();
 	TestResult->Seed = PCGSourceBeingInspected->GetExecutionState().GetSeed();
 
-	PCGDeterminismTests::RunDeterminismTest(PCGGraphBeingEdited, PCGSourceBeingInspected, *TestResult);
+	PCGDeterminismTests::RunDeterminismTest(MainGraph, PCGSourceBeingInspected, *TestResult);
 
 	DeterminismWidget->AddItem(TestResult);
 	DeterminismWidget->AddDetailsColumn();
@@ -1226,16 +1413,17 @@ TAttribute<FGraphAppearanceInfo> FPCGEditor::GetAppearanceInfo() const
 
 void FPCGEditor::OnEditGraphSettings()
 {
-	check(GraphEditorWidget)
+	TSharedPtr<SGraphEditor> FocusedGraphWidget = GetFocusedGraphEditorWidget();
+	check(FocusedGraphWidget.IsValid());
 
 	// Clear any selected nodes.
-	GraphEditorWidget->ClearSelectionSet();
+	FocusedGraphWidget->ClearSelectionSet();
 
 	for (TSharedPtr<SPCGEditorGraphDetailsView> PropertyDetailsWidget : PropertyDetailsWidgets)
 	{
 		if (PropertyDetailsWidget.IsValid())
 		{
-			PropertyDetailsWidget->SetObject(PCGGraphBeingEdited);
+			PropertyDetailsWidget->SetObject(GetFocusedGraph());
 		}
 	}
 
@@ -1258,7 +1446,7 @@ bool FPCGEditor::IsEditGraphSettingsToggled() const
 
 		const TArray<TWeakObjectPtr<UObject>>& SelectedObjects = PropertyDetailsWidgets[WidgetIndex]->GetSelectedObjects();
 		// The only object selected should be the graph. If there is no details view panel open, leave it disabled.
-		if (SelectedObjects.Num() == 1 && SelectedObjects[0] == PCGGraphBeingEdited.Get())
+		if (SelectedObjects.Num() == 1 && SelectedObjects[0] == GetFocusedGraph())
 		{
 			if (const TSharedPtr<SDockTab>& Tab = TabManager->FindExistingLiveTab(FTabId(PCGEditorTabs::PropertyDetailsID[WidgetIndex])))
 			{
@@ -1290,11 +1478,17 @@ bool FPCGEditor::IsToggleGraphParamsToggled() const
 	return IsPanelCurrentlyOpen(PCGEditorTabs::UserParamsID);
 }
 
-bool FPCGEditor::CanCollapseNodesInSubgraph() const
+bool FPCGEditor::CanCollapseNodesInSubgraph(bool bCollapseToEmbeddedSubgraph) const
 {
+	TSharedPtr<SGraphEditor> FocusedGraphWidget = GetFocusedGraphEditorWidget();
+	if (!FocusedGraphWidget)
+	{
+		return false;
+	}
+
 	bool HasPCGNode = false;
 
-	for (UObject* Object : GraphEditorWidget->GetSelectedNodes())
+	for (UObject* Object : FocusedGraphWidget->GetSelectedNodes())
 	{
 		check(Object);
 
@@ -1320,9 +1514,10 @@ bool FPCGEditor::CanCollapseNodesInSubgraph() const
 
 void FPCGEditor::OnAddDynamicInputPin()
 {
-	check(GraphEditorWidget.IsValid());
+	TSharedPtr<SGraphEditor> FocusedGraphWidget = GetFocusedGraphEditorWidget();
+	check(FocusedGraphWidget.IsValid());
 
-	const FGraphPanelSelectionSet SelectedNodes = GraphEditorWidget->GetSelectedNodes();
+	const FGraphPanelSelectionSet SelectedNodes = FocusedGraphWidget->GetSelectedNodes();
 
 	if (!ensure(SelectedNodes.Num() == 1))
 	{
@@ -1336,9 +1531,10 @@ void FPCGEditor::OnAddDynamicInputPin()
 
 bool FPCGEditor::CanAddDynamicInputPin() const
 {
-	check(GraphEditorWidget.IsValid());
+	TSharedPtr<SGraphEditor> FocusedGraphWidget = GetFocusedGraphEditorWidget();
+	check(FocusedGraphWidget.IsValid());
 
-	const FGraphPanelSelectionSet SelectedNodes = GraphEditorWidget->GetSelectedNodes();
+	const FGraphPanelSelectionSet SelectedNodes = FocusedGraphWidget->GetSelectedNodes();
 	if (SelectedNodes.Num() != 1)
 	{
 		return false;
@@ -1350,9 +1546,10 @@ bool FPCGEditor::CanAddDynamicInputPin() const
 
 void FPCGEditor::OnRenameNode()
 {
-	check(GraphEditorWidget.IsValid());
+	TSharedPtr<SGraphEditor> FocusedGraphWidget = GetFocusedGraphEditorWidget();
+	check(FocusedGraphWidget.IsValid());
 
-	const FGraphPanelSelectionSet SelectedNodes = GraphEditorWidget->GetSelectedNodes();
+	const FGraphPanelSelectionSet SelectedNodes = FocusedGraphWidget->GetSelectedNodes();
 
 	if (!ensure(SelectedNodes.Num() == 1))
 	{
@@ -1363,15 +1560,16 @@ void FPCGEditor::OnRenameNode()
 	const UEdGraphNode* SelectedNode = Cast<UEdGraphNode>(*SelectedNodes.CreateConstIterator());
 	if (SelectedNode != nullptr && SelectedNode->GetCanRenameNode())
 	{
-		GraphEditorWidget->IsNodeTitleVisible(SelectedNode, true);
+		FocusedGraphWidget->IsNodeTitleVisible(SelectedNode, true);
 	}
 }
 
 bool FPCGEditor::CanRenameNode() const
 {
-	check(GraphEditorWidget.IsValid());
+	TSharedPtr<SGraphEditor> FocusedGraphWidget = GetFocusedGraphEditorWidget();
+	check(FocusedGraphWidget.IsValid());
 
-	const FGraphPanelSelectionSet SelectedNodes = GraphEditorWidget->GetSelectedNodes();
+	const FGraphPanelSelectionSet SelectedNodes = FocusedGraphWidget->GetSelectedNodes();
 
 	// You cannot enter renaming mode on multiple nodes at once, since they will not all enter synchronously.
 	// Simultaneous editing of multiple InlineEditableTextBlocks may not even be possible with default behavior.
@@ -1397,13 +1595,15 @@ bool FPCGEditor::CanRenameNode() const
 
 bool FPCGEditor::InternalValidationOnAction()
 {
-	if (!GraphEditorWidget.IsValid() || PCGEditorGraph == nullptr)
+	TSharedPtr<SGraphEditor> FocusedGraphWidget = GetFocusedGraphEditorWidget();
+	UPCGEditorGraph* FocusedEditorGraph = GetFocusedEditorGraph();
+	if (!FocusedGraphWidget || FocusedEditorGraph == nullptr)
 	{
-		UE_LOG(LogPCGEditor, Error, TEXT("GraphEditorWidget or PCGEditorGraph is null, aborting"));
+		UE_LOG(LogPCGEditor, Error, TEXT("FocusedGraphWidget or FocusedEditorGraph is null, aborting"));
 		return false;
 	}
 
-	UPCGGraph* PCGGraph = PCGEditorGraph->GetPCGGraph();
+	UPCGGraph* PCGGraph = FocusedEditorGraph->GetPCGGraph();
 	if (PCGGraph == nullptr)
 	{
 		UE_LOG(LogPCGEditor, Error, TEXT("PCGGraph is null, aborting"));
@@ -1420,7 +1620,8 @@ void FPCGEditor::OnSelectNamedRerouteUsages()
 		return;
 	}
 
-	const FGraphPanelSelectionSet SelectedNodes = GraphEditorWidget->GetSelectedNodes();
+	TSharedPtr<SGraphEditor> FocusedGraphWidget = GetFocusedGraphEditorWidget();
+	const FGraphPanelSelectionSet SelectedNodes = FocusedGraphWidget->GetSelectedNodes();
 
 	if (SelectedNodes.Num() != 1)
 	{
@@ -1439,7 +1640,7 @@ void FPCGEditor::OnSelectNamedRerouteUsages()
 		return;
 	}
 
-	GraphEditorWidget->ClearSelectionSet();
+	FocusedGraphWidget->ClearSelectionSet();
 
 	// Some assumptions below - that only usages are connected to the invisible pin.
 	if (const UPCGPin* InvisiblePin = DeclarationNode->GetPCGNode()->GetOutputPin(PCGNamedRerouteConstants::InvisiblePinLabel))
@@ -1448,22 +1649,28 @@ void FPCGEditor::OnSelectNamedRerouteUsages()
 		{
 			if (const UPCGNode* Usage = Edge->OutputPin->Node)
 			{
-				GraphEditorWidget->SetNodeSelection(GetEditorNode(Usage), true);
+				FocusedGraphWidget->SetNodeSelection(GetEditorNode(Usage), true);
 			}
 		}
 	}
 
-	GraphEditorWidget->ZoomToFit(true);
+	FocusedGraphWidget->ZoomToFit(true);
 }
 
 bool FPCGEditor::CanSelectNamedRerouteUsages() const
 {
-	if (!GraphEditorWidget || GraphEditorWidget->GetSelectedNodes().Num() != 1)
+	TSharedPtr<SGraphEditor> FocusedGraphWidget = GetFocusedGraphEditorWidget();
+	if (!FocusedGraphWidget)
 	{
 		return false;
 	}
 
-	if (auto It = GraphEditorWidget->GetSelectedNodes().CreateConstIterator(); It)
+	if (FocusedGraphWidget->GetSelectedNodes().Num() != 1)
+	{
+		return false;
+	}
+
+	if (auto It = FocusedGraphWidget->GetSelectedNodes().CreateConstIterator(); It)
 	{
 		const UObject* Object = *It;
 		return Object && Object->IsA<UPCGEditorGraphNodeNamedRerouteDeclaration>();
@@ -1479,7 +1686,8 @@ void FPCGEditor::OnSelectNamedRerouteDeclaration()
 		return;
 	}
 
-	const FGraphPanelSelectionSet SelectedNodes = GraphEditorWidget->GetSelectedNodes();
+	TSharedPtr<SGraphEditor> FocusedGraphWidget = GetFocusedGraphEditorWidget();
+	const FGraphPanelSelectionSet SelectedNodes = FocusedGraphWidget->GetSelectedNodes();
 
 	if (SelectedNodes.Num() != 1)
 	{
@@ -1495,7 +1703,7 @@ void FPCGEditor::OnSelectNamedRerouteDeclaration()
 			continue;
 		}
 
-		GraphEditorWidget->ClearSelectionSet();
+		FocusedGraphWidget->ClearSelectionSet();
 
 		if (!UsageNode->GetPCGNode())
 		{
@@ -1516,12 +1724,18 @@ void FPCGEditor::OnSelectNamedRerouteDeclaration()
 
 bool FPCGEditor::CanSelectNamedRerouteDeclaration() const
 {
-	if (!GraphEditorWidget || GraphEditorWidget->GetSelectedNodes().Num() != 1)
+	TSharedPtr<SGraphEditor> FocusedGraphWidget = GetFocusedGraphEditorWidget();
+	if (!FocusedGraphWidget)
 	{
 		return false;
 	}
 
-	if (auto It = GraphEditorWidget->GetSelectedNodes().CreateConstIterator(); It)
+	if (FocusedGraphWidget->GetSelectedNodes().Num() != 1)
+	{
+		return false;
+	}
+
+	if (auto It = FocusedGraphWidget->GetSelectedNodes().CreateConstIterator(); It)
 	{
 		const UObject* Object = *It;
 		return Object && Object->IsA<UPCGEditorGraphNodeNamedRerouteUsage>();
@@ -1537,7 +1751,13 @@ void FPCGEditor::OnJumpToSource()
 		return;
 	}
 
-	for (UObject* Object : GraphEditorWidget->GetSelectedNodes())
+	TSharedPtr<SGraphEditor> FocusedGraphWidget = GetFocusedGraphEditorWidget();
+	if (!FocusedGraphWidget)
+	{
+		return;
+	}
+
+	for (UObject* Object : FocusedGraphWidget->GetSelectedNodes())
 	{
 		const UPCGEditorGraphNodeBase* PCGEditorGraphNode = Cast<const UPCGEditorGraphNodeBase>(Object);
 		const UPCGNode* PCGNode = PCGEditorGraphNode ? PCGEditorGraphNode->GetPCGNode() : nullptr;
@@ -1587,22 +1807,23 @@ FActionMenuContent FPCGEditor::OnCreateActionMenuContent(UEdGraph* InGraph, cons
 	return FActionMenuContent(Menu, Menu->GetFilterTextBox());
 }
 
-void FPCGEditor::OnCollapseNodesInSubgraph()
+void FPCGEditor::OnCollapseNodesInSubgraph(bool bCollapseToEmbeddedSubgraph)
 {
 	if (!InternalValidationOnAction())
 	{
 		return;
 	}
 
-	UPCGGraph* PCGGraph = PCGEditorGraph->GetPCGGraph();
+	UPCGGraph* PCGGraph = GetFocusedGraph();
 	check(PCGGraph);
 
 	// Gather all nodes that will be included in the subgraph, and the extra nodes
 	TArray<UPCGNode*> NodesToCollapse;
 	TArray<UObject*> ExtraNodesToCollapse;
 
-	check(GraphEditorWidget);
-	for (UObject* Object : GraphEditorWidget->GetSelectedNodes())
+	TSharedPtr<SGraphEditor> FocusedGraphWidget = GetFocusedGraphEditorWidget();
+	check(FocusedGraphWidget);
+	for (UObject* Object : FocusedGraphWidget->GetSelectedNodes())
 	{
 		check(Object);
 
@@ -1634,25 +1855,39 @@ void FPCGEditor::OnCollapseNodesInSubgraph()
 	// Create a new subgraph, by creating a new PCGGraph asset.
 	TObjectPtr<UPCGGraph> NewPCGGraph = nullptr;
 
-	IAssetTools& AssetTools = FModuleManager::Get().LoadModuleChecked<FAssetToolsModule>("AssetTools").Get();
+	if(!bCollapseToEmbeddedSubgraph)
+	{
+		IAssetTools& AssetTools = FModuleManager::Get().LoadModuleChecked<FAssetToolsModule>("AssetTools").Get();
 
-	TObjectPtr<UPCGGraphFactory> Factory = NewObject<UPCGGraphFactory>();
-	Factory->bSkipTemplateSelection = true;
+		TObjectPtr<UPCGGraphFactory> Factory = NewObject<UPCGGraphFactory>();
+		Factory->bSkipTemplateSelection = true;
 
-	FString NewPackageName;
-	FString NewAssetName;
-	PCGEditorUtils::GetParentPackagePathAndUniqueName(PCGGraph, LOCTEXT("NewPCGSubgraphAsset", "NewPCGSubgraph").ToString(), NewPackageName, NewAssetName);
+		FString NewPackageName;
+		FString NewAssetName;
+		PCGEditorUtils::GetParentPackagePathAndUniqueName(PCGGraph, LOCTEXT("NewPCGSubgraphAsset", "NewPCGSubgraph").ToString(), NewPackageName, NewAssetName);
 
-	NewPCGGraph = Cast<UPCGGraph>(AssetTools.CreateAssetWithDialog(NewAssetName, NewPackageName, PCGGraph->GetClass(), Factory, "PCGEditor_CollapseInSubgraph"));
+		NewPCGGraph = Cast<UPCGGraph>(AssetTools.CreateAssetWithDialog(NewAssetName, NewPackageName, PCGGraph->GetClass(), Factory, "PCGEditor_CollapseInSubgraph"));
 
-	if (NewPCGGraph == nullptr)
-	{
-		UE_LOG(LogPCGEditor, Warning, TEXT("Subgraph asset creation was aborted or failed, abort."));
-		return;
+		if (NewPCGGraph == nullptr)
+		{
+			UE_LOG(LogPCGEditor, Warning, TEXT("Subgraph asset creation was aborted or failed, abort."));
+			return;
+		}
 	}
 
 	{
 		FScopedTransaction Transaction(LOCTEXT("PCGCollapseInSubgraphMessage", "[PCG] Collapse into Subgraph"));
+
+		if (bCollapseToEmbeddedSubgraph)
+		{
+			UPCGGraph* RootGraph = PCGGraph;
+			if (PCGGraph->IsEmbeddedSubgraph())
+			{
+				RootGraph = PCGGraph->GetEmbeddedParentGraph();
+			}
+			NewPCGGraph = RootGraph->AddNewEmbeddedSubgraph();
+		}
+
 		FText OutFailReason;
 		NewPCGGraph = FPCGSubgraphHelpers::CollapseIntoSubgraphWithReason(PCGGraph, NodesToCollapse, ExtraNodesToCollapse, OutFailReason, NewPCGGraph);
 
@@ -1664,22 +1899,31 @@ void FPCGEditor::OnCollapseNodesInSubgraph()
 		}
 
 		// Force a refresh
-		PCGEditorGraph->ReconstructGraph();
+		PCGGraph->GetEditorGraph()->ReconstructGraph();
 	}
 
 	if (NewPCGGraph)
 	{
-		// Save the new asset
-		UEditorAssetLibrary::SaveLoadedAsset(NewPCGGraph);
+		if (!NewPCGGraph->IsEmbeddedSubgraph())
+		{
+			// Save the new asset
+			UEditorAssetLibrary::SaveLoadedAsset(NewPCGGraph);
+		}
 
 		// Notify the widget
-		GraphEditorWidget->NotifyGraphChanged();
+		FocusedGraphWidget->NotifyGraphChanged();
 	}
 }
 
 bool FPCGEditor::CanExportNodes() const
 {
-	for (const UObject* Object : GraphEditorWidget->GetSelectedNodes())
+	TSharedPtr<SGraphEditor> FocusedGraphWidget = GetFocusedGraphEditorWidget();
+	if (!FocusedGraphWidget)
+	{
+		return false;
+	}
+
+	for (const UObject* Object : FocusedGraphWidget->GetSelectedNodes())
 	{
 		check(Object);
 
@@ -1695,8 +1939,17 @@ bool FPCGEditor::CanExportNodes() const
 			continue;
 		}
 
-		if (Object->IsA<UPCGEditorGraphNodeBase>())
+		if (const UPCGEditorGraphNodeBase* GraphNode = Cast<UPCGEditorGraphNodeBase>(Object))
 		{
+			// Exclude embedded subgraph nodes
+			if (const UPCGBaseSubgraphNode* SubgraphNode = Cast<UPCGBaseSubgraphNode>(GraphNode->GetPCGNode()))
+			{
+				if (SubgraphNode->GetSubgraph() && SubgraphNode->GetSubgraph()->IsEmbeddedSubgraph())
+				{
+					continue;
+				}
+			}
+
 			return true;
 		}
 	}
@@ -1706,21 +1959,17 @@ bool FPCGEditor::CanExportNodes() const
 
 void FPCGEditor::OnExportNodes()
 {
-	if (!GraphEditorWidget.IsValid() || PCGEditorGraph == nullptr)
-	{
-		UE_LOG(LogPCGEditor, Error, TEXT("GraphEditorWidget or PCGEditorGraph is null, aborting"));
-		return;
-	}
-
-	if (PCGGraphBeingEdited == nullptr)
+	UPCGEditorGraph* FocusedEditorGraph = GetFocusedEditorGraph();
+	TSharedPtr<SGraphEditor> FocusedGraphWidget = GetFocusedGraphEditorWidget();
+	if (!FocusedEditorGraph || !FocusedGraphWidget)
 	{
-		UE_LOG(LogPCGEditor, Error, TEXT("Editor has no graph loaded, aborting"));
+		UE_LOG(LogPCGEditor, Error, TEXT("FocusedGraphWidget or Focused editor graph is null, aborting"));
 		return;
 	}
 
 	IAssetTools& AssetTools = FModuleManager::Get().LoadModuleChecked<FAssetToolsModule>("AssetTools").Get();
 
-	for (UObject* Object : GraphEditorWidget->GetSelectedNodes())
+	for (UObject* Object : FocusedGraphWidget->GetSelectedNodes())
 	{
 		check(Object);
 
@@ -1736,6 +1985,16 @@ void FPCGEditor::OnExportNodes()
 		{
 			UPCGNode* PCGNode = PCGEditorGraphNode->GetPCGNode();
 			check(PCGNode);
+
+			// Exclude embedded subgraph nodes
+			if (const UPCGBaseSubgraphNode* SubgraphNode = Cast<UPCGBaseSubgraphNode>(PCGNode))
+			{
+				if (SubgraphNode->GetSubgraph() && SubgraphNode->GetSubgraph()->IsEmbeddedSubgraph())
+				{
+					continue;
+				}
+			}
+
 			Settings = PCGNode->GetSettings();
 		}
 
@@ -1747,7 +2006,7 @@ void FPCGEditor::OnExportNodes()
 		// Create new settings asset
 		FString NewPackageName;
 		FString NewAssetName;
-		PCGEditorUtils::GetParentPackagePathAndUniqueName(PCGGraphBeingEdited, LOCTEXT("NewPCGSettingsAsset", "NewPCGSettings").ToString(), NewPackageName, NewAssetName);
+		PCGEditorUtils::GetParentPackagePathAndUniqueName(GetFocusedGraph(), LOCTEXT("NewPCGSettingsAsset", "NewPCGSettings").ToString(), NewPackageName, NewAssetName);
 
 		UObject* NewSettings = AssetTools.DuplicateAssetWithDialogAndTitle(NewAssetName, NewPackageName, Settings, NSLOCTEXT("PCGEditor_ExportNodes", "PCGEditor_ExportNodesTitle", "Export Settings As..."));
 
@@ -1761,9 +2020,15 @@ void FPCGEditor::OnExportNodes()
 
 void FPCGEditor::OnConvertToStandaloneNodes()
 {
+	TSharedPtr<SGraphEditor> FocusedGraphWidget = GetFocusedGraphEditorWidget();
+	if (!FocusedGraphWidget)
+	{
+		return;
+	}
+
 	const FScopedTransaction Transaction(*FPCGEditorCommon::ContextIdentifier, LOCTEXT("PCGEditorConvertToStandaloneMessage", "PCG Editor: Converting instanced nodes to standalone"), nullptr);
 
-	for (UObject* Object : GraphEditorWidget->GetSelectedNodes())
+	for (UObject* Object : FocusedGraphWidget->GetSelectedNodes())
 	{
 		check(Object);
 		// Exclude input and output nodes from the subgraph.
@@ -1795,15 +2060,18 @@ void FPCGEditor::OnConvertToStandaloneNodes()
 	}
 
 	// Notify the widget
-	if (GraphEditorWidget)
-	{
-		GraphEditorWidget->NotifyGraphChanged();
-	}
+	FocusedGraphWidget->NotifyGraphChanged();
 }
 
 bool FPCGEditor::CanConvertToStandaloneNodes() const
 {
-	for (const UObject* Object : GraphEditorWidget->GetSelectedNodes())
+	TSharedPtr<SGraphEditor> FocusedGraphWidget = GetFocusedGraphEditorWidget();
+	if (!FocusedGraphWidget)
+	{
+		return false;
+	}
+
+	for (const UObject* Object : FocusedGraphWidget->GetSelectedNodes())
 	{
 		check(Object);
 
@@ -1830,12 +2098,13 @@ bool FPCGEditor::CanConvertToStandaloneNodes() const
 
 void FPCGEditor::OnToggleInspected()
 {
-	if (!GraphEditorWidget.IsValid())
+	TSharedPtr<SGraphEditor> FocusedGraphWidget = GetFocusedGraphEditorWidget();
+	if (!FocusedGraphWidget)
 	{
 		return;
 	}
 
-	UEdGraphNode* GraphNode = GraphEditorWidget->GetSingleSelectedNode();
+	UEdGraphNode* GraphNode = FocusedGraphWidget->GetSingleSelectedNode();
 	UPCGEditorGraphNodeBase* PCGGraphNodeBase = Cast<UPCGEditorGraphNodeBase>(GraphNode);
 
 	const UPCGNode* PCGNode = PCGGraphNodeBase ? PCGGraphNodeBase->GetPCGNode() : nullptr;
@@ -1969,19 +2238,20 @@ void FPCGEditor::SetNodeInspected(UPCGEditorGraphNodeBase* InspectedNode, bool b
 
 bool FPCGEditor::CanToggleInspected() const
 {
-	if (!GraphEditorWidget.IsValid())
+	TSharedPtr<SGraphEditor> FocusedGraphWidget = GetFocusedGraphEditorWidget();
+	if (!FocusedGraphWidget)
 	{
 		return false;
 	}
 
-	const FGraphPanelSelectionSet& SelectedNodes = GraphEditorWidget->GetSelectedNodes();
+	const FGraphPanelSelectionSet& SelectedNodes = FocusedGraphWidget->GetSelectedNodes();
 	if (SelectedNodes.Num() != 1)
 	{
 		// Can only inspect one node.
 		return false;
 	}
 
-	for (const UObject* Object : GraphEditorWidget->GetSelectedNodes())
+	for (const UObject* Object : FocusedGraphWidget->GetSelectedNodes())
 	{
 		const UPCGEditorGraphNodeBase* PCGEditorGraphNode = Cast<UPCGEditorGraphNodeBase>(Object);
 		if (!PCGEditorGraphNode)
@@ -2001,38 +2271,41 @@ bool FPCGEditor::CanToggleInspected() const
 
 ECheckBoxState FPCGEditor::GetInspectedCheckState() const
 {
-	if (GraphEditorWidget.IsValid())
+	TSharedPtr<SGraphEditor> FocusedGraphWidget = GetFocusedGraphEditorWidget();
+	if (!FocusedGraphWidget)
 	{
-		const FGraphPanelSelectionSet& SelectedNodes = GraphEditorWidget->GetSelectedNodes();
-
-		if (SelectedNodes.IsEmpty())
-		{
-			return ECheckBoxState::Unchecked;
-		}
+		return ECheckBoxState::Unchecked;
+	}
 
-		bool bAllEnabled = true;
-		bool bAnyEnabled = false;
+	const FGraphPanelSelectionSet& SelectedNodes = FocusedGraphWidget->GetSelectedNodes();
 
-		for (UObject* Object : SelectedNodes)
-		{
-			const UPCGEditorGraphNodeBase* PCGEditorGraphNode = Cast<UPCGEditorGraphNodeBase>(Object);
-			if (!PCGEditorGraphNode)
-			{
-				continue;
-			}
+	if (SelectedNodes.IsEmpty())
+	{
+		return ECheckBoxState::Unchecked;
+	}
 
-			bAllEnabled &= PCGEditorGraphNode->GetInspected();
-			bAnyEnabled |= PCGEditorGraphNode->GetInspected();
-		}
+	bool bAllEnabled = true;
+	bool bAnyEnabled = false;
 
-		if (bAllEnabled)
-		{
-			return ECheckBoxState::Checked;
-		}
-		else if (bAnyEnabled)
+	for (UObject* Object : SelectedNodes)
+	{
+		const UPCGEditorGraphNodeBase* PCGEditorGraphNode = Cast<UPCGEditorGraphNodeBase>(Object);
+		if (!PCGEditorGraphNode)
 		{
-			return ECheckBoxState::Undetermined;
+			continue;
 		}
+
+		bAllEnabled &= PCGEditorGraphNode->GetInspected();
+		bAnyEnabled |= PCGEditorGraphNode->GetInspected();
+	}
+
+	if (bAllEnabled)
+	{
+		return ECheckBoxState::Checked;
+	}
+	else if (bAnyEnabled)
+	{
+		return ECheckBoxState::Undetermined;
 	}
 
 	return ECheckBoxState::Unchecked;
@@ -2040,66 +2313,70 @@ ECheckBoxState FPCGEditor::GetInspectedCheckState() const
 
 void FPCGEditor::OnToggleEnabled()
 {
+	TSharedPtr<SGraphEditor> FocusedGraphWidget = GetFocusedGraphEditorWidget();
+	if (!FocusedGraphWidget)
+	{
+		return;
+	}
+
 	const ECheckBoxState CheckState = GetEnabledCheckState();
 	const bool bNewCheckState = !(CheckState != ECheckBoxState::Unchecked);
 
 	// To prevent the changes on the editor node from being in the transaction, we delay reconstruction.
 	TArray<FPCGDeferNodeReconstructScope> DeferredEditorNodes;
 
-	if (GraphEditorWidget.IsValid())
+	FScopedTransaction Transaction(*FPCGEditorCommon::ContextIdentifier, LOCTEXT("PCGEditorToggleEnableTransactionMessage", "PCG Editor: Toggle Enable Nodes"), nullptr);
+
+	UPCGGraph* PCGGraph = GetFocusedGraph();
+	if (!ensure(PCGGraph))
 	{
-		FScopedTransaction Transaction(*FPCGEditorCommon::ContextIdentifier, LOCTEXT("PCGEditorToggleEnableTransactionMessage", "PCG Editor: Toggle Enable Nodes"), nullptr);
+		return;
+	}
 
-		UPCGGraph* PCGGraph = PCGEditorGraph ? PCGEditorGraph->GetPCGGraph() : nullptr;
-		if (!ensure(PCGGraph))
-		{
-			return;
-		}
+	PCGGraph->DisableNotificationsForEditor();
 
-		PCGGraph->DisableNotificationsForEditor();
+	bool bChanged = false;
+	for (UObject* Object : FocusedGraphWidget->GetSelectedNodes())
+	{
+		UPCGEditorGraphNodeBase* PCGEditorGraphNode = Cast<UPCGEditorGraphNodeBase>(Object);
+		UPCGNode* PCGNode = PCGEditorGraphNode ? PCGEditorGraphNode->GetPCGNode() : nullptr;
+		UPCGSettingsInterface* PCGSettingsInterface = PCGNode ? PCGNode->GetSettingsInterface() : nullptr;
 
-		bool bChanged = false;
-		for (UObject* Object : GraphEditorWidget->GetSelectedNodes())
+		if (!PCGSettingsInterface || !PCGSettingsInterface->CanBeDisabled())
 		{
-			UPCGEditorGraphNodeBase* PCGEditorGraphNode = Cast<UPCGEditorGraphNodeBase>(Object);
-			UPCGNode* PCGNode = PCGEditorGraphNode ? PCGEditorGraphNode->GetPCGNode() : nullptr;
-			UPCGSettingsInterface* PCGSettingsInterface = PCGNode ? PCGNode->GetSettingsInterface() : nullptr;
-
-			if (!PCGSettingsInterface || !PCGSettingsInterface->CanBeDisabled())
-			{
-				continue;
-			}
-
-			if (PCGSettingsInterface->bEnabled != bNewCheckState)
-			{
-				DeferredEditorNodes.Emplace(PCGEditorGraphNode);
-				PCGSettingsInterface->Modify();
-				PCGSettingsInterface->SetEnabled(bNewCheckState);
-				bChanged = true;
-			}
+			continue;
 		}
 
-		PCGGraph->EnableNotificationsForEditor();
-
-		if (bChanged)
+		if (PCGSettingsInterface->bEnabled != bNewCheckState)
 		{
-			GraphEditorWidget->NotifyGraphChanged();
-		}
-		else
-		{
-			Transaction.Cancel();
+			DeferredEditorNodes.Emplace(PCGEditorGraphNode);
+			PCGSettingsInterface->Modify();
+			PCGSettingsInterface->SetEnabled(bNewCheckState);
+			bChanged = true;
 		}
 	}
+
+	PCGGraph->EnableNotificationsForEditor();
+
+	if (bChanged)
+	{
+		FocusedGraphWidget->NotifyGraphChanged();
+	}
+	else
+	{
+		Transaction.Cancel();
+	}
 }
 
 bool FPCGEditor::CanToggleEnabled() const
 {
-	if (!GraphEditorWidget.IsValid())
+	TSharedPtr<SGraphEditor> FocusedGraphWidget = GetFocusedGraphEditorWidget();
+	if (!FocusedGraphWidget)
 	{
 		return false;
 	}
 
-	for (const UObject* Object : GraphEditorWidget->GetSelectedNodes())
+	for (const UObject* Object : FocusedGraphWidget->GetSelectedNodes())
 	{
 		const UPCGEditorGraphNodeBase* PCGEditorGraphNode = Cast<UPCGEditorGraphNodeBase>(Object);
 		const UPCGNode* PCGNode = PCGEditorGraphNode ? PCGEditorGraphNode->GetPCGNode() : nullptr;
@@ -2120,34 +2397,37 @@ bool FPCGEditor::CanToggleEnabled() const
 
 ECheckBoxState FPCGEditor::GetEnabledCheckState() const
 {
-	if (GraphEditorWidget.IsValid())
+	TSharedPtr<SGraphEditor> FocusedGraphWidget = GetFocusedGraphEditorWidget();
+	if (!FocusedGraphWidget)
 	{
-		bool bAllEnabled = true;
-		bool bAnyEnabled = false;
-
-		for (const UObject* Object : GraphEditorWidget->GetSelectedNodes())
-		{
-			const UPCGEditorGraphNodeBase* PCGEditorGraphNode = Cast<UPCGEditorGraphNodeBase>(Object);
-			const UPCGNode* PCGNode = PCGEditorGraphNode ? PCGEditorGraphNode->GetPCGNode() : nullptr;
-			const UPCGSettingsInterface* PCGSettingsInterface = PCGNode ? PCGNode->GetSettingsInterface() : nullptr;
+		return ECheckBoxState::Unchecked;
+	}
 
-			if (!PCGSettingsInterface || !PCGSettingsInterface->CanBeDisabled())
-			{
-				continue;
-			}
+	bool bAllEnabled = true;
+	bool bAnyEnabled = false;
 
-			bAllEnabled &= PCGSettingsInterface->bEnabled;
-			bAnyEnabled |= PCGSettingsInterface->bEnabled;
-		}
+	for (const UObject* Object : FocusedGraphWidget->GetSelectedNodes())
+	{
+		const UPCGEditorGraphNodeBase* PCGEditorGraphNode = Cast<UPCGEditorGraphNodeBase>(Object);
+		const UPCGNode* PCGNode = PCGEditorGraphNode ? PCGEditorGraphNode->GetPCGNode() : nullptr;
+		const UPCGSettingsInterface* PCGSettingsInterface = PCGNode ? PCGNode->GetSettingsInterface() : nullptr;
 
-		if (bAllEnabled)
+		if (!PCGSettingsInterface || !PCGSettingsInterface->CanBeDisabled())
 		{
-			return ECheckBoxState::Checked;
-		}
-		else if (bAnyEnabled)
-		{
-			return ECheckBoxState::Undetermined;
+			continue;
 		}
+
+		bAllEnabled &= PCGSettingsInterface->bEnabled;
+		bAnyEnabled |= PCGSettingsInterface->bEnabled;
+	}
+
+	if (bAllEnabled)
+	{
+		return ECheckBoxState::Checked;
+	}
+	else if (bAnyEnabled)
+	{
+		return ECheckBoxState::Undetermined;
 	}
 
 	return ECheckBoxState::Unchecked;
@@ -2155,43 +2435,47 @@ ECheckBoxState FPCGEditor::GetEnabledCheckState() const
 
 void FPCGEditor::OnToggleDebug()
 {
+	TSharedPtr<SGraphEditor> FocusedGraphWidget = GetFocusedGraphEditorWidget();
+	if (!FocusedGraphWidget)
+	{
+		return;
+	}
+
 	const ECheckBoxState CheckState = GetDebugCheckState();
 	const bool bNewCheckState = !(CheckState != ECheckBoxState::Unchecked);
 
-	if (GraphEditorWidget.IsValid())
+	FScopedTransaction Transaction(*FPCGEditorCommon::ContextIdentifier, LOCTEXT("PCGEditorToggleDebugTransactionMessage", "PCG Editor: Toggle Debug Nodes"), nullptr);
+
+	bool bChanged = false;
+	for (UObject* Object : FocusedGraphWidget->GetSelectedNodes())
 	{
-		FScopedTransaction Transaction(*FPCGEditorCommon::ContextIdentifier, LOCTEXT("PCGEditorToggleDebugTransactionMessage", "PCG Editor: Toggle Debug Nodes"), nullptr);
+		UPCGEditorGraphNodeBase* PCGEditorGraphNode = Cast<UPCGEditorGraphNodeBase>(Object);
+		UPCGNode* PCGNode = PCGEditorGraphNode ? PCGEditorGraphNode->GetPCGNode() : nullptr;
+		UPCGSettingsInterface* PCGSettingsInterface = PCGNode ? PCGNode->GetSettingsInterface() : nullptr;
 
-		bool bChanged = false;
-		for (UObject* Object : GraphEditorWidget->GetSelectedNodes())
+		if (!PCGSettingsInterface || !PCGSettingsInterface->CanBeDebugged())
 		{
-			UPCGEditorGraphNodeBase* PCGEditorGraphNode = Cast<UPCGEditorGraphNodeBase>(Object);
-			UPCGNode* PCGNode = PCGEditorGraphNode ? PCGEditorGraphNode->GetPCGNode() : nullptr;
-			UPCGSettingsInterface* PCGSettingsInterface = PCGNode ? PCGNode->GetSettingsInterface() : nullptr;
-
-			if (!PCGSettingsInterface || !PCGSettingsInterface->CanBeDebugged())
-			{
-				continue;
-			}
-
-			bChanged |= PCGSettingsInterface->SetDebugged(bNewCheckState);
+			continue;
 		}
 
-		if (!bChanged)
-		{
-			Transaction.Cancel();
-		}
+		bChanged |= PCGSettingsInterface->SetDebugged(bNewCheckState);
+	}
+
+	if (!bChanged)
+	{
+		Transaction.Cancel();
 	}
 }
 
 bool FPCGEditor::CanToggleDebug() const
 {
-	if (!GraphEditorWidget.IsValid())
+	TSharedPtr<SGraphEditor> FocusedGraphWidget = GetFocusedGraphEditorWidget();
+	if (!FocusedGraphWidget)
 	{
 		return false;
 	}
 
-	for (const UObject* Object : GraphEditorWidget->GetSelectedNodes())
+	for (const UObject* Object : FocusedGraphWidget->GetSelectedNodes())
 	{
 		const UPCGEditorGraphNodeBase* PCGEditorGraphNode = Cast<UPCGEditorGraphNodeBase>(Object);
 		const UPCGNode* PCGNode = PCGEditorGraphNode ? PCGEditorGraphNode->GetPCGNode() : nullptr;
@@ -2208,123 +2492,134 @@ bool FPCGEditor::CanToggleDebug() const
 
 void FPCGEditor::OnDebugOnlySelected()
 {
-	if (GraphEditorWidget.IsValid() && PCGEditorGraph)
+	UPCGEditorGraph* FocusedEditorGraph = GetFocusedEditorGraph();
+	TSharedPtr<SGraphEditor> FocusedGraphWidget = GetFocusedGraphEditorWidget();
+	if (!FocusedEditorGraph || !FocusedGraphWidget)
 	{
-		bool bChanged = false;
+		return;
+	}
 
-		const FGraphPanelSelectionSet& SelectedNodes = GraphEditorWidget->GetSelectedNodes();
+	bool bChanged = false;
 
-		FScopedTransaction Transaction(*FPCGEditorCommon::ContextIdentifier, LOCTEXT("PCGEditorDebugOnlySelectedTransactionMessage", "PCG Editor: Debug only selected nodes"), nullptr);
+	const FGraphPanelSelectionSet& SelectedNodes = FocusedGraphWidget->GetSelectedNodes();
 
-		bool bAnyNonSelectedNodesDebugged = false;
-		bool bAllSelectedNodesDebugged = true;
+	FScopedTransaction Transaction(*FPCGEditorCommon::ContextIdentifier, LOCTEXT("PCGEditorDebugOnlySelectedTransactionMessage", "PCG Editor: Debug only selected nodes"), nullptr);
 
-		// Initial pass - inspect state of selected and non-selected nodes.
-		for (const UEdGraphNode* Node : PCGEditorGraph->Nodes)
-		{
-			const UPCGEditorGraphNodeBase* PCGEditorGraphNode = Cast<UPCGEditorGraphNodeBase>(Node);
-			const UPCGNode* PCGNode = PCGEditorGraphNode ? PCGEditorGraphNode->GetPCGNode() : nullptr;
-			const UPCGSettingsInterface* PCGSettingsInterface = PCGNode ? PCGNode->GetSettingsInterface() : nullptr;
-			if (!PCGSettingsInterface)
-			{
-				continue;
-			}
+	bool bAnyNonSelectedNodesDebugged = false;
+	bool bAllSelectedNodesDebugged = true;
 
-			if (SelectedNodes.Contains(PCGEditorGraphNode))
-			{
-				bAllSelectedNodesDebugged &= PCGSettingsInterface->bDebug;
-			}
-			else
-			{
-				bAnyNonSelectedNodesDebugged |= PCGSettingsInterface->bDebug;
-			}
+	// Initial pass - inspect state of selected and non-selected nodes.
+	for (const UEdGraphNode* Node : FocusedEditorGraph->Nodes)
+	{
+		const UPCGEditorGraphNodeBase* PCGEditorGraphNode = Cast<UPCGEditorGraphNodeBase>(Node);
+		const UPCGNode* PCGNode = PCGEditorGraphNode ? PCGEditorGraphNode->GetPCGNode() : nullptr;
+		const UPCGSettingsInterface* PCGSettingsInterface = PCGNode ? PCGNode->GetSettingsInterface() : nullptr;
+		if (!PCGSettingsInterface)
+		{
+			continue;
 		}
 
-		// The selected nodes should be debugged if any non-selected nodes are being debugged, or if the selected
-		// nodes are partially being debugged.
-		const bool bTargetDebugState = bAnyNonSelectedNodesDebugged || !bAllSelectedNodesDebugged;
-
-		for (UEdGraphNode* Node : PCGEditorGraph->Nodes)
+		if (SelectedNodes.Contains(PCGEditorGraphNode))
 		{
-			UPCGEditorGraphNodeBase* PCGEditorGraphNode = Cast<UPCGEditorGraphNodeBase>(Node);
-			UPCGNode* PCGNode = PCGEditorGraphNode ? PCGEditorGraphNode->GetPCGNode() : nullptr;
-			UPCGSettingsInterface* PCGSettingsInterface = PCGNode ? PCGNode->GetSettingsInterface() : nullptr;
-
-			if (!PCGSettingsInterface || !PCGSettingsInterface->CanBeDebugged())
-			{
-				continue;
-			}
+			bAllSelectedNodesDebugged &= PCGSettingsInterface->bDebug;
+		}
+		else
+		{
+			bAnyNonSelectedNodesDebugged |= PCGSettingsInterface->bDebug;
+		}
+	}
 
-			// Selected set to target state, non-selected should not be debugged.
-			const bool bShouldBeDebug = SelectedNodes.Contains(PCGEditorGraphNode) ? bTargetDebugState : false;
+	// The selected nodes should be debugged if any non-selected nodes are being debugged, or if the selected
+	// nodes are partially being debugged.
+	const bool bTargetDebugState = bAnyNonSelectedNodesDebugged || !bAllSelectedNodesDebugged;
 
-			bChanged |= PCGSettingsInterface->SetDebugged(bShouldBeDebug);
-		}
+	for (UEdGraphNode* Node : FocusedEditorGraph->Nodes)
+	{
+		UPCGEditorGraphNodeBase* PCGEditorGraphNode = Cast<UPCGEditorGraphNodeBase>(Node);
+		UPCGNode* PCGNode = PCGEditorGraphNode ? PCGEditorGraphNode->GetPCGNode() : nullptr;
+		UPCGSettingsInterface* PCGSettingsInterface = PCGNode ? PCGNode->GetSettingsInterface() : nullptr;
 
-		if (!bChanged)
+		if (!PCGSettingsInterface || !PCGSettingsInterface->CanBeDebugged())
 		{
-			Transaction.Cancel();
+			continue;
 		}
+
+		// Selected set to target state, non-selected should not be debugged.
+		const bool bShouldBeDebug = SelectedNodes.Contains(PCGEditorGraphNode) ? bTargetDebugState : false;
+
+		bChanged |= PCGSettingsInterface->SetDebugged(bShouldBeDebug);
+	}
+
+	if (!bChanged)
+	{
+		Transaction.Cancel();
 	}
 }
 
 void FPCGEditor::OnDisableDebugOnAllNodes()
 {
-	if (GraphEditorWidget.IsValid() && PCGEditorGraph)
+	UPCGEditorGraph* FocusedEditorGraph = GetFocusedEditorGraph();
+	TSharedPtr<SGraphEditor> FocusedGraphWidget = GetFocusedGraphEditorWidget();
+	if (!FocusedEditorGraph || !FocusedGraphWidget)
 	{
-		bool bChanged = false;
-		FScopedTransaction Transaction(*FPCGEditorCommon::ContextIdentifier, LOCTEXT("PCGEditorDisableDebugAllNodesTransactionMessage", "PCG Editor: Disable debug on all nodes"), nullptr);
-
-		for (UEdGraphNode* Node : PCGEditorGraph->Nodes)
-		{
-			UPCGEditorGraphNodeBase* PCGEditorGraphNode = Cast<UPCGEditorGraphNodeBase>(Node);
-			UPCGNode* PCGNode = PCGEditorGraphNode ? PCGEditorGraphNode->GetPCGNode() : nullptr;
-			UPCGSettingsInterface* PCGSettingsInterface = PCGNode ? PCGNode->GetSettingsInterface() : nullptr;
-			if (!PCGSettingsInterface)
-			{
-				continue;
-			}
+		return;
+	}
 
-			bChanged |= PCGSettingsInterface->SetDebugged(false);
-		}
+	bool bChanged = false;
+	FScopedTransaction Transaction(*FPCGEditorCommon::ContextIdentifier, LOCTEXT("PCGEditorDisableDebugAllNodesTransactionMessage", "PCG Editor: Disable debug on all nodes"), nullptr);
 
-		if (!bChanged)
+	for (UEdGraphNode* Node : FocusedEditorGraph->Nodes)
+	{
+		UPCGEditorGraphNodeBase* PCGEditorGraphNode = Cast<UPCGEditorGraphNodeBase>(Node);
+		UPCGNode* PCGNode = PCGEditorGraphNode ? PCGEditorGraphNode->GetPCGNode() : nullptr;
+		UPCGSettingsInterface* PCGSettingsInterface = PCGNode ? PCGNode->GetSettingsInterface() : nullptr;
+		if (!PCGSettingsInterface)
 		{
-			Transaction.Cancel();
+			continue;
 		}
+
+		bChanged |= PCGSettingsInterface->SetDebugged(false);
+	}
+
+	if (!bChanged)
+	{
+		Transaction.Cancel();
 	}
 }
 
 ECheckBoxState FPCGEditor::GetDebugCheckState() const
 {
-	if (GraphEditorWidget.IsValid())
+	TSharedPtr<SGraphEditor> FocusedGraphWidget = GetFocusedGraphEditorWidget();
+	if (!FocusedGraphWidget)
 	{
-		bool bAllDebug = true;
-		bool bAnyDebug = false;
-
-		for (const UObject* Object : GraphEditorWidget->GetSelectedNodes())
-		{
-			const UPCGEditorGraphNodeBase* PCGEditorGraphNode = Cast<UPCGEditorGraphNodeBase>(Object);
-			const UPCGNode* PCGNode = PCGEditorGraphNode ? PCGEditorGraphNode->GetPCGNode() : nullptr;
-			const UPCGSettingsInterface* PCGSettingsInterface = PCGNode ? PCGNode->GetSettingsInterface() : nullptr;
+		return ECheckBoxState::Unchecked;
+	}
 
-			if (!PCGSettingsInterface || !PCGSettingsInterface->CanBeDebugged())
-			{
-				continue;
-			}
+	bool bAllDebug = true;
+	bool bAnyDebug = false;
 
-			bAllDebug &= PCGSettingsInterface->bDebug;
-			bAnyDebug |= PCGSettingsInterface->bDebug;
-		}
+	for (const UObject* Object : FocusedGraphWidget->GetSelectedNodes())
+	{
+		const UPCGEditorGraphNodeBase* PCGEditorGraphNode = Cast<UPCGEditorGraphNodeBase>(Object);
+		const UPCGNode* PCGNode = PCGEditorGraphNode ? PCGEditorGraphNode->GetPCGNode() : nullptr;
+		const UPCGSettingsInterface* PCGSettingsInterface = PCGNode ? PCGNode->GetSettingsInterface() : nullptr;
 
-		if (bAllDebug)
+		if (!PCGSettingsInterface || !PCGSettingsInterface->CanBeDebugged())
 		{
-			return ECheckBoxState::Checked;
-		}
-		else if (bAnyDebug)
-		{
-			return ECheckBoxState::Undetermined;
+			continue;
 		}
+
+		bAllDebug &= PCGSettingsInterface->bDebug;
+		bAnyDebug |= PCGSettingsInterface->bDebug;
+	}
+
+	if (bAllDebug)
+	{
+		return ECheckBoxState::Checked;
+	}
+	else if (bAnyDebug)
+	{
+		return ECheckBoxState::Undetermined;
 	}
 
 	return ECheckBoxState::Unchecked;
@@ -2332,81 +2627,84 @@ ECheckBoxState FPCGEditor::GetDebugCheckState() const
 
 void FPCGEditor::SelectAllNodes()
 {
-	if (GraphEditorWidget.IsValid())
+	if (TSharedPtr<SGraphEditor> FocusedGraphWidget = GetFocusedGraphEditorWidget())
 	{
-		GraphEditorWidget->SelectAllNodes();
+		FocusedGraphWidget->SelectAllNodes();
 	}
 }
 
 bool FPCGEditor::CanSelectAllNodes() const
 {
-	return GraphEditorWidget.IsValid();
+	return !!GetFocusedGraphEditorWidget();
 }
 
 void FPCGEditor::DeleteSelectedNodes()
 {
-	if (GraphEditorWidget.IsValid())
+	TSharedPtr<SGraphEditor> FocusedGraphWidget = GetFocusedGraphEditorWidget();
+	if (!FocusedGraphWidget)
 	{
-		UPCGGraph* PCGGraph = PCGEditorGraph->GetPCGGraph();
-		check(PCGEditorGraph && PCGGraph);
+		return;
+	}
 
-		// DeleteSelectedNodes is called directly from UI command 
-		PCGGraph->PrimeGraphCompilationCache();
+	UPCGGraph* PCGGraph = GetFocusedGraph();
+	check(PCGGraph && PCGGraph->GetEditorGraph());
 
-		bool bChanged = false;
+	// DeleteSelectedNodes is called directly from UI command 
+	PCGGraph->PrimeGraphCompilationCache();
 
-		{
-			const FScopedTransaction Transaction(*FPCGEditorCommon::ContextIdentifier, LOCTEXT("PCGEditorDeleteTransactionMessage", "PCG Editor: Delete"), nullptr);
-			PCGEditorGraph->Modify();
-		
-			TArray<UPCGNode*> NodesToRemove;
+	bool bChanged = false;
 
-			for (UObject* Object : GraphEditorWidget->GetSelectedNodes())
+	{
+		const FScopedTransaction Transaction(*FPCGEditorCommon::ContextIdentifier, LOCTEXT("PCGEditorDeleteTransactionMessage", "PCG Editor: Delete"), nullptr);
+		PCGGraph->GetEditorGraph()->Modify();
+
+		TArray<UPCGNode*> NodesToRemove;
+
+		for (UObject* Object : FocusedGraphWidget->GetSelectedNodes())
+		{
+			if (UPCGEditorGraphNodeBase* PCGEditorGraphNode = Cast<UPCGEditorGraphNodeBase>(Object))
 			{
-				if (UPCGEditorGraphNodeBase* PCGEditorGraphNode = Cast<UPCGEditorGraphNodeBase>(Object))
+				if (PCGEditorGraphNode->CanUserDeleteNode())
 				{
-					if (PCGEditorGraphNode->CanUserDeleteNode())
-					{
-						UPCGNode* PCGNode = PCGEditorGraphNode->GetPCGNode();
-						check(PCGNode);
+					UPCGNode* PCGNode = PCGEditorGraphNode->GetPCGNode();
+					check(PCGNode);
 
-						NodesToRemove.Add(PCGNode);
+					NodesToRemove.Add(PCGNode);
 
-						PCGEditorGraphNode->DestroyNode();
-						bChanged = true;
-					}
-				}
-				else if (UEdGraphNode* GraphNode = Cast<UEdGraphNode>(Object))
-				{
-					if (GraphNode->CanUserDeleteNode())
-					{
-						GraphNode->DestroyNode();
-						bChanged = true;
-					}
+					PCGEditorGraphNode->DestroyNode();
+					bChanged = true;
 				}
 			}
-
-			if (bChanged)
+			else if (UEdGraphNode* GraphNode = Cast<UEdGraphNode>(Object))
 			{
-				// Need to modify the pcg graph so comments are also caught.
-				PCGGraph->Modify();
-				PCGGraph->RemoveNodes(NodesToRemove);
+				if (GraphNode->CanUserDeleteNode())
+				{
+					GraphNode->DestroyNode();
+					bChanged = true;
+				}
 			}
 		}
 
 		if (bChanged)
 		{
-			GraphEditorWidget->ClearSelectionSet();
-			GraphEditorWidget->NotifyGraphChanged();
+			// Need to modify the pcg graph so comments are also caught.
+			PCGGraph->Modify();
+			PCGGraph->RemoveNodes(NodesToRemove);
 		}
 	}
+
+	if (bChanged)
+	{
+		FocusedGraphWidget->ClearSelectionSet();
+		FocusedGraphWidget->NotifyGraphChanged();
+	}
 }
 
 bool FPCGEditor::CanDeleteSelectedNodes() const
 {
-	if (GraphEditorWidget.IsValid())
+	if (TSharedPtr<SGraphEditor> FocusedGraphWidget = GetFocusedGraphEditorWidget())
 	{
-		for (UObject* Object : GraphEditorWidget->GetSelectedNodes())
+		for (UObject* Object : FocusedGraphWidget->GetSelectedNodes())
 		{
 			UEdGraphNode* GraphNode = CastChecked<UEdGraphNode>(Object);
 
@@ -2422,9 +2720,9 @@ bool FPCGEditor::CanDeleteSelectedNodes() const
 
 void FPCGEditor::CopySelectedNodes()
 {
-	if (GraphEditorWidget.IsValid())
+	if (TSharedPtr<SGraphEditor> FocusedGraphWidget = GetFocusedGraphEditorWidget())
 	{
-		const FGraphPanelSelectionSet SelectedNodes = GraphEditorWidget->GetSelectedNodes();
+		const FGraphPanelSelectionSet SelectedNodes = FocusedGraphWidget->GetSelectedNodes();
 
 		//TODO: evaluate creating a clipboard object instead of ownership hack
 		for (UObject* SelectedNode : SelectedNodes)
@@ -2449,9 +2747,9 @@ void FPCGEditor::CopySelectedNodes()
 
 bool FPCGEditor::CanCopySelectedNodes() const
 {
-	if (GraphEditorWidget.IsValid())
+	if (TSharedPtr<SGraphEditor> FocusedGraphWidget = GetFocusedGraphEditorWidget())
 	{
-		for (UObject* Object : GraphEditorWidget->GetSelectedNodes())
+		for (UObject* Object : FocusedGraphWidget->GetSelectedNodes())
 		{
 			if (UEdGraphNode* GraphNode = CastChecked<UEdGraphNode>(Object))
 			{
@@ -2479,24 +2777,31 @@ bool FPCGEditor::CanCutSelectedNodes() const
 
 void FPCGEditor::PasteNodes()
 {
-	if (GraphEditorWidget.IsValid())
+	if (TSharedPtr<SGraphEditor> FocusedGraphWidget = GetFocusedGraphEditorWidget())
 	{
-		PasteNodesHere(GraphEditorWidget->GetPasteLocation2f());
+		PasteNodesHere(FocusedGraphWidget->GetPasteLocation2f());
 	}
 }
 
 void FPCGEditor::PasteNodesHere(const FVector2D& Location)
 {
-	if (!GraphEditorWidget.IsValid() || !PCGEditorGraph)
+	UPCGEditorGraph* FocusedEditorGraph = GetFocusedEditorGraph();
+	if (!FocusedEditorGraph)
+	{
+		return;
+	}
+
+	TSharedPtr<SGraphEditor> FocusedGraphWidget = GetFocusedGraphEditorWidget();
+	if (!FocusedGraphWidget)
 	{
 		return;
 	}
 
 	const FScopedTransaction Transaction(*FPCGEditorCommon::ContextIdentifier, LOCTEXT("PCGEditorPasteTransactionMessage", "PCG Editor: Paste"), nullptr);
-	PCGEditorGraph->Modify();
+	FocusedEditorGraph->Modify();
 
 	// Clear the selection set (newly pasted stuff will be selected)
-	GraphEditorWidget->ClearSelectionSet();
+	FocusedGraphWidget->ClearSelectionSet();
 
 	// Grab the text to paste from the clipboard.
 	FString TextToImport;
@@ -2504,7 +2809,7 @@ void FPCGEditor::PasteNodesHere(const FVector2D& Location)
 
 	// Import the nodes
 	TSet<UEdGraphNode*> PastedNodes;
-	FEdGraphUtilities::ImportNodesFromText(PCGEditorGraph, TextToImport, /*out*/ PastedNodes);
+	FEdGraphUtilities::ImportNodesFromText(FocusedEditorGraph, TextToImport, /*out*/ PastedNodes);
 
 	//Average position of nodes so we can move them while still maintaining relative distances to each other
 	FVector2D AvgNodePosition(0.0f, 0.0f);
@@ -2533,7 +2838,7 @@ void FPCGEditor::PasteNodesHere(const FVector2D& Location)
 
 	for (UEdGraphNode* PastedNode : PastedNodes)
 	{
-		GraphEditorWidget->SetNodeSelection(PastedNode, true);
+		FocusedGraphWidget->SetNodeSelection(PastedNode, true);
 
 		PastedNode->NodePosX = (PastedNode->NodePosX - AvgNodePosition.X) + Location.X;
 		PastedNode->NodePosY = (PastedNode->NodePosY - AvgNodePosition.Y) + Location.Y;
@@ -2550,8 +2855,8 @@ void FPCGEditor::PasteNodesHere(const FVector2D& Location)
 	}
 
 	// Need to modify the pcg graph so comments are also caught.
-	PCGGraphBeingEdited->Modify();
-	PCGGraphBeingEdited->AddNodes(NodesToPaste);
+	FocusedEditorGraph->GetPCGGraph()->Modify();
+	FocusedEditorGraph->GetPCGGraph()->AddNodes(NodesToPaste);
 
 	for (UEdGraphNode* PastedNode : PastedNodes)
 	{
@@ -2576,7 +2881,7 @@ void FPCGEditor::PasteNodesHere(const FVector2D& Location)
 		}
 	}
 
-	GraphEditorWidget->NotifyGraphChanged();
+	FocusedGraphWidget->NotifyGraphChanged();
 }
 
 bool FPCGEditor::CanPasteNodes() const
@@ -2584,7 +2889,7 @@ bool FPCGEditor::CanPasteNodes() const
 	FString ClipboardContent;
 	FPlatformApplicationMisc::ClipboardPaste(ClipboardContent);
 
-	return FEdGraphUtilities::CanImportNodesFromText(PCGEditorGraph, ClipboardContent);
+	return FEdGraphUtilities::CanImportNodesFromText(GetFocusedEditorGraph(), ClipboardContent);
 }
 
 void FPCGEditor::DuplicateNodes()
@@ -2600,90 +2905,90 @@ bool FPCGEditor::CanDuplicateNodes() const
 
 void FPCGEditor::OnAlignTop()
 {
-	if (GraphEditorWidget.IsValid())
+	if (TSharedPtr<SGraphEditor> FocusedGraphWidget = GetFocusedGraphEditorWidget())
 	{
-		GraphEditorWidget->OnAlignTop();
+		FocusedGraphWidget->OnAlignTop();
 	}
 }
 
 void FPCGEditor::OnAlignMiddle()
 {
-	if (GraphEditorWidget.IsValid())
+	if (TSharedPtr<SGraphEditor> FocusedGraphWidget = GetFocusedGraphEditorWidget())
 	{
-		GraphEditorWidget->OnAlignMiddle();
+		FocusedGraphWidget->OnAlignMiddle();
 	}
 }
 
 void FPCGEditor::OnAlignBottom()
 {
-	if (GraphEditorWidget.IsValid())
+	if (TSharedPtr<SGraphEditor> FocusedGraphWidget = GetFocusedGraphEditorWidget())
 	{
-		GraphEditorWidget->OnAlignBottom();
+		FocusedGraphWidget->OnAlignBottom();
 	}
 }
 
 void FPCGEditor::OnAlignLeft()
 {
-	if (GraphEditorWidget.IsValid())
+	if (TSharedPtr<SGraphEditor> FocusedGraphWidget = GetFocusedGraphEditorWidget())
 	{
-		GraphEditorWidget->OnAlignLeft();
+		FocusedGraphWidget->OnAlignLeft();
 	}
 }
 
 void FPCGEditor::OnAlignCenter()
 {
-	if (GraphEditorWidget.IsValid())
+	if (TSharedPtr<SGraphEditor> FocusedGraphWidget = GetFocusedGraphEditorWidget())
 	{
-		GraphEditorWidget->OnAlignCenter();
+		FocusedGraphWidget->OnAlignCenter();
 	}
 }
 
 void FPCGEditor::OnAlignRight()
 {
-	if (GraphEditorWidget.IsValid())
+	if (TSharedPtr<SGraphEditor> FocusedGraphWidget = GetFocusedGraphEditorWidget())
 	{
-		GraphEditorWidget->OnAlignRight();
+		FocusedGraphWidget->OnAlignRight();
 	}
 }
 
 void FPCGEditor::OnStraightenConnections()
 {
-	if (GraphEditorWidget.IsValid())
+	if (TSharedPtr<SGraphEditor> FocusedGraphWidget = GetFocusedGraphEditorWidget())
 	{
-		GraphEditorWidget->OnStraightenConnections();
+		FocusedGraphWidget->OnStraightenConnections();
 	}
 }
 
 void FPCGEditor::OnDistributeNodesH()
 {
-	if (GraphEditorWidget.IsValid())
+	if (TSharedPtr<SGraphEditor> FocusedGraphWidget = GetFocusedGraphEditorWidget())
 	{
-		GraphEditorWidget->OnDistributeNodesH();
+		FocusedGraphWidget->OnDistributeNodesH();
 	}
 }
 
 void FPCGEditor::OnDistributeNodesV()
 {
-	if (GraphEditorWidget.IsValid())
+	if (TSharedPtr<SGraphEditor> FocusedGraphWidget = GetFocusedGraphEditorWidget())
 	{
-		GraphEditorWidget->OnDistributeNodesV();
+		FocusedGraphWidget->OnDistributeNodesV();
 	}
 }
 
 void FPCGEditor::OnCreateComment()
 {
-	if (PCGEditorGraph)
+	if (UPCGEditorGraph* FocusedEditorGraph = GetFocusedEditorGraph())
 	{
 		FPCGEditorGraphSchemaAction_NewComment CommentAction;
 
-		TSharedPtr<SGraphEditor> GraphEditorPtr = SGraphEditor::FindGraphEditorForGraph(PCGEditorGraph);
+		TSharedPtr<SGraphEditor> GraphEditorPtr = SGraphEditor::FindGraphEditorForGraph(FocusedEditorGraph);
 		FVector2f Location = FVector2f::ZeroVector;
 		if (GraphEditorPtr)
 		{
 			Location = GraphEditorPtr->GetPasteLocation2f();
 		}
 
-		CommentAction.PerformAction(PCGEditorGraph, nullptr, Location);
+		CommentAction.PerformAction(FocusedEditorGraph, nullptr, Location);
 	}
 }
 
@@ -2764,58 +3069,61 @@ void FPCGEditor::CreateGraphEditorCommands()
 	);
 }
 
-TSharedRef<SGraphEditor> FPCGEditor::CreateGraphEditorWidget()
+TSharedPtr<SGraphEditor> FPCGEditor::GetFocusedGraphEditorWidget() const
 {
-	CreateGraphEditorCommands();
-
-	return CreateGraphEditorWidget(nullptr, PCGEditorGraph);
+	if (TSharedPtr<SGraphEditor> FocusedGraphEditorPtr = FocusedGraphEditor.Pin())
+	{
+		return FocusedGraphEditorPtr;
+	}
+	else
+	{
+		return GraphEditorWidget;
+	}
 }
 
-TSharedRef<SGraphEditor> FPCGEditor::GetOrCreateGraphEditorWidget()
+TSharedRef<SGraphEditor> FPCGEditor::CreateGraphEditorWidget(TSharedPtr<FTabInfo> InTabInfo, UPCGEditorGraph* InGraph)
 {
-	if (!GraphEditorWidget.IsValid())
+	if (!InGraph)
 	{
-		GraphEditorWidget = CreateGraphEditorWidget();
+		InGraph = GetMainEditorGraph();
 	}
 
-	return GraphEditorWidget.ToSharedRef();
-}
+	TSharedPtr<SWidget> TitleBarWidget = SNew(SPCGGraphEditorTitleBar)
+		.Graph(InGraph)
+		.Editor(SharedThis(this))
+		.HistoryNavigationWidget(InTabInfo->CreateHistoryNavigationWidget());
 
-TSharedRef<SGraphEditor> FPCGEditor::CreateGraphEditorWidget(TSharedPtr<FTabInfo> InTabInfo, UEdGraph* InGraph)
-{
 	SGraphEditor::FGraphEditorEvents InEvents;
 	InEvents.OnSelectionChanged = SGraphEditor::FOnSelectionChanged::CreateSP(this, &FPCGEditor::OnSelectedNodesChanged);
 	InEvents.OnVerifyTextCommit = FOnNodeVerifyTextCommit::CreateSP(this, &FPCGEditor::OnValidateNodeTitle);
 	InEvents.OnTextCommitted = FOnNodeTextCommitted::CreateSP(this, &FPCGEditor::OnNodeTitleCommitted);
 	InEvents.OnNodeDoubleClicked = FSingleNodeEvent::CreateSP(this, &FPCGEditor::OnNodeDoubleClicked);
-	InEvents.OnSpawnNodeByShortcutAtLocation = SGraphEditor::FOnSpawnNodeByShortcutAtLocation::CreateSP(this, &FPCGEditor::OnSpawnNodeByShortcut, Cast<UPCGEditorGraph>(InGraph));
+	InEvents.OnSpawnNodeByShortcutAtLocation = SGraphEditor::FOnSpawnNodeByShortcutAtLocation::CreateSP(this, &FPCGEditor::OnSpawnNodeByShortcut, InGraph);
 	InEvents.OnCreateActionMenuAtLocation = SGraphEditor::FOnCreateActionMenuAtLocation::CreateSP(this, &FPCGEditor::OnCreateActionMenuContent);
 
-	return SNew(SGraphEditor)
+	TSharedRef<SGraphEditor> GraphWidget = SNew(SGraphEditor)
 		.AdditionalCommands(GraphEditorCommands)
 		.IsEditable(true)
 		.Appearance(GetAppearanceInfo())
 		.GraphToEdit(InGraph)
 		.GraphEvents(InEvents)
-		.ShowGraphStateOverlay(false);
-}
+		.ShowGraphStateOverlay(false)
+		.TitleBar(TitleBarWidget)
+		.OnNavigateHistoryBack(FSimpleDelegate::CreateSP(this, &FPCGEditor::NavigateTab, FDocumentTracker::EOpenDocumentCause::NavigateBackwards))
+		.OnNavigateHistoryForward(FSimpleDelegate::CreateSP(this, &FPCGEditor::NavigateTab, FDocumentTracker::EOpenDocumentCause::NavigateForwards));
 
-TSharedPtr<FWorkflowTabFactory> FPCGEditor::CreateGraphEditorTabFactory()
-{
-	FPCGGenericTabFactoryParams Params(PCGEditorTabs::GraphEditorID, SharedThis<FPCGEditor>(this), [this]() { return GetOrCreateGraphEditorWidget(); });
-	Params.Label = LOCTEXT("GraphLabel", "Graph");
-	return MakeShared<FPCGGenericTabFactory>(Params);
+	if (InGraph == GetMainEditorGraph())
+	{
+		GraphEditorWidget = GraphWidget;
+	}
+
+	return GraphWidget;
 }
 
 void FPCGEditor::OnClose()
 {
 	FWorkflowCentricApplication::OnClose();
 
-	if (PCGEditorGraph)
-	{
-		PCGEditorGraph->OnClose();
-	}
-
 	if (FLevelEditorModule* LevelEditor = FModuleManager::GetModulePtr<FLevelEditorModule>(TEXT("LevelEditor")))
 	{
 		LevelEditor->OnMapChanged().RemoveAll(this);
@@ -2825,22 +3133,22 @@ void FPCGEditor::OnClose()
 	FEditorDelegates::EndPIE.RemoveAll(this);
 
 	FAssetEditorToolkit::OnClose();
-	
+
 	InspectionDataManager.Cleanup();
 
-	if (PCGGraphBeingEdited)
+	if (MainGraph)
 	{
-		PCGGraphBeingEdited->OnGraphChangedDelegate.RemoveAll(this);
-		PCGGraphBeingEdited->OnNodeSourceCompiledDelegate.RemoveAll(this);
+		MainGraph->OnGraphChangedDelegate.RemoveAll(this);
+		MainGraph->OnNodeSourceCompiledDelegate.RemoveAll(this);
 
-		if (PCGGraphBeingEdited->IsInspecting())
+		if (MainGraph->IsInspecting())
 		{
-			PCGGraphBeingEdited->DisableInspection();
+			MainGraph->DisableInspection();
 		}
 
-		if (PCGGraphBeingEdited->NotificationsForEditorArePausedByUser())
+		if (MainGraph->NotificationsForEditorArePausedByUser())
 		{
-			PCGGraphBeingEdited->ToggleUserPausedNotificationsForEditor();
+			MainGraph->ToggleUserPausedNotificationsForEditor();
 		}
 	}
 
@@ -2874,12 +3182,62 @@ void FPCGEditor::InitToolMenuContext(FToolMenuContext& MenuContext)
 
 void FPCGEditor::SaveEditedObjectState()
 {
+	UPCGGraph* Graph = GetMainGraph();
+	if (!Graph)
+	{
+		return;
+	}
+	
+	Graph->LastEditedDocuments.Empty();
+
 	DocumentManager->SaveAllState();
 }
 
 void FPCGEditor::RestoreEditedObjectState()
 {
-	// @todo_pcg: implement
+	UPCGGraph* Graph = GetMainGraph();
+	if (!Graph)
+	{
+		return;
+	}
+
+	auto RestoreDocument = [this](UPCGGraph* InGraph, const FPCGGraphDocumentInfo* Info)
+	{
+		TSharedPtr<SDockTab> TabWithGraph = OpenDocument(InGraph, FDocumentTracker::EOpenDocumentCause::RestorePreviousDocument);
+		if (TabWithGraph)
+		{
+			TSharedRef<SGraphEditor> GraphEditor = StaticCastSharedRef<SGraphEditor>(TabWithGraph->GetContent());
+			
+			if (Info)
+			{
+				GraphEditor->SetViewLocation(Info->SavedViewOffset, Info->SavedZoomAmount);
+			}
+		}
+	};
+
+	// Find Main graph info
+	const FPCGGraphDocumentInfo* MainGraphInfo = Graph->LastEditedDocuments.FindByPredicate([Graph](const FPCGGraphDocumentInfo& Info)
+	{
+		return (Info.EditedObjectPath == Graph);
+	});
+
+	// Restore main graph first (so it is first in the tab list)
+	RestoreDocument(Graph, MainGraphInfo);
+
+	// Restore other graph(s)
+	for (const FPCGGraphDocumentInfo& Info : Graph->LastEditedDocuments)
+	{
+		if (UPCGGraph* EditedGraph = Cast<UPCGGraph>(Info.EditedObjectPath.ResolveObject()))
+		{
+			if (EditedGraph != Graph)
+			{
+				RestoreDocument(EditedGraph, &Info);
+			}
+		}
+	}
+
+	// Focus on main tab
+	OpenDocument(Graph, FDocumentTracker::EOpenDocumentCause::RestorePreviousDocument);
 }
 
 TSharedPtr<FWorkflowTabFactory> FPCGEditor::CreatePaletteTabFactory()
@@ -3082,7 +3440,7 @@ TSharedRef<SPCGEditorGraphDetailsView> FPCGEditor::GetOrCreatePropertyDetailsWid
 	{
 		PropertyDetailsWidgets[Index] = SNew(SPCGEditorGraphDetailsView);
 		PropertyDetailsWidgets[Index]->SetEditor(SharedThis(this));
-		PropertyDetailsWidgets[Index]->SetObject(PCGGraphBeingEdited);
+		PropertyDetailsWidgets[Index]->SetObject(GetFocusedGraph());
 	}
 
 	return PropertyDetailsWidgets[Index].ToSharedRef();
@@ -3136,13 +3494,61 @@ TSharedRef<SPCGEditorViewport> FPCGEditor::GetOrCreateViewportWidget(int32 Index
 	return ViewportWidgets[Index].ToSharedRef();
 }
 
+TSharedPtr<FWorkflowTabFactory> FPCGEditor::CreateEmbeddedSubgraphsTabFactory()
+{
+	FPCGGenericTabFactoryParams Params(PCGEditorTabs::EmbeddedSubgraphsID, SharedThis<FPCGEditor>(this), [this]() { return GetOrCreateEmbeddedSubgraphsWidget(); });
+	Params.Label = LOCTEXT("EmbeddedSubgraphsLabel", "Embedded Graphs");
+	return MakeShared<FPCGGenericTabFactory>(Params);
+}
+
+TSharedRef<SPCGEditorGraphEmbeddedSubgraphsView> FPCGEditor::GetOrCreateEmbeddedSubgraphsWidget()
+{
+	if (!EmbeddedSubgraphsWidget.IsValid())
+	{
+		EmbeddedSubgraphsWidget = SNew(SPCGEditorGraphEmbeddedSubgraphsView, SharedThis(this));
+	}
+
+	return EmbeddedSubgraphsWidget.ToSharedRef();
+}
+
+void FPCGEditor::OnGraphEditorFocused(const TSharedRef<SGraphEditor>& InGraphEditor)
+{
+	FocusedGraphEditor = InGraphEditor;
+	OnSelectedNodesChanged(InGraphEditor->GetSelectedNodes());
+
+	if (UserParamsWidget)
+	{
+		UserParamsWidget->SetGraph(GetFocusedGraph());
+	}
+}
+
+void FPCGEditor::OnGraphEditorBackgrounded(const TSharedRef<SGraphEditor>& InGraphEditor)
+{
+	FocusedGraphEditor = nullptr;
+	
+	OnSelectedNodesChanged(FGraphPanelSelectionSet());
+
+	if (UserParamsWidget)
+	{
+		UserParamsWidget->SetGraph(GetFocusedGraph());
+	}
+}
+
+void FPCGEditor::OnGraphPreSave(const UPCGGraph* InGraph)
+{
+	if (InGraph == GetMainGraph())
+	{
+		SaveEditedObjectState();
+	}
+}
+
 void FPCGEditor::OnSelectedNodesChanged(const TSet<UObject*>& NewSelection)
 {
 	TArray<TWeakObjectPtr<UObject>> SelectedObjects;
 
 	if (NewSelection.Num() == 0)
 	{
-		SelectedObjects.Add(PCGGraphBeingEdited);
+		SelectedObjects.Add(GetFocusedGraph());
 	}
 	else
 	{
@@ -3171,7 +3577,8 @@ void FPCGEditor::OnSelectedNodesChanged(const TSet<UObject*>& NewSelection)
 
 void FPCGEditor::OnNodeTitleCommitted(const FText& NewText, ETextCommit::Type CommitInfo, UEdGraphNode* NodeBeingChanged)
 {
-	check(PCGGraphBeingEdited);
+	UPCGGraph* PCGGraph = GetFocusedGraph();
+	check(PCGGraph);
 
 	if (NodeBeingChanged)
 	{
@@ -3188,7 +3595,7 @@ void FPCGEditor::OnNodeTitleCommitted(const FText& NewText, ETextCommit::Type Co
 				{
 					NodeBeingChanged->Modify();
 					// Modify the graph as well, as non-pcg editor nodes (like the comment nodes) are serialized in UPCGGraph::ExtraEditorNodes.
-					PCGGraphBeingEdited->Modify();
+					PCGGraph->Modify();
 				}
 
 				NodeBeingChanged->OnRenameNode(NewText.ToString());
@@ -3216,7 +3623,7 @@ void FPCGEditor::OnNodeDoubleClicked(UEdGraphNode* Node)
 		if (const UPCGSettings* PCGSettings = Cast<UPCGSettings>(Object))
 		{
 			// Functions may require the GraphEditorWidget's node selection, so set it manually to be safe.
-			GraphEditorWidget->SetNodeSelection(Node, /*bSelect=*/true);
+			GetFocusedGraphEditorWidget()->SetNodeSelection(Node, /*bSelect=*/true);
 
 			switch (GetDefault<UPCGEditorSettings>()->NodeDoubleClickAction)
 			{
@@ -3256,16 +3663,12 @@ void FPCGEditor::OnNodeDoubleClicked(UEdGraphNode* Node)
 				}
 			}
 
-			if (Object)
+			if (UPCGGraph* PCGGraph = Cast<UPCGGraph>(Object))
 			{
-				// Open other editor...
-				GEditor->GetEditorSubsystem<UAssetEditorSubsystem>()->OpenEditorForAsset(Object);
-				IAssetEditorInstance* EditorInstance = GEditor->GetEditorSubsystem<UAssetEditorSubsystem>()->FindEditorForAsset(Object, /*bFocusIfOpen*/true);
-
-				FPCGEditor* OtherPCGEditor = static_cast<FPCGEditor*>(EditorInstance);
-				if (OtherPCGEditor && !StackToInspect.GetStackFrames().IsEmpty())
+				FPCGEditor* CurrentPCGEditor = OpenEditorForGraph(PCGGraph);
+				if (CurrentPCGEditor && CurrentPCGEditor != this)
 				{
-					OtherPCGEditor->SetStackBeingInspectedFromAnotherEditor(StackToInspect);
+					CurrentPCGEditor->SetStackBeingInspectedFromAnotherEditor(StackToInspect);
 				}
 			}
 		}
@@ -3283,7 +3686,7 @@ void FPCGEditor::JumpToDefinition(const UClass* Class) const
 void FPCGEditor::OnComponentUnregistered(UPCGComponent* Component)
 {
 	// Refresh the debug object tree to avoid stale entries from components that have been unregistered.
-	if (!Component || Component->GetGraph() == PCGGraphBeingEdited)
+	if (!Component || Component->GetGraph() == MainGraph)
 	{
 		if (DebugObjectTreeWidget.IsValid())
 		{
@@ -3301,17 +3704,17 @@ void FPCGEditor::OnSourceGenerationDone(IPCGBaseSubsystem* Subsystem, IPCGGraphE
 {
 	const FPCGStack& StackBeingInspected = InspectionDataManager.GetStackBeingInspected();
 	const IPCGGraphExecutionSource* PCGSourceBeingInspected = InspectionDataManager.GetPCGSourceBeingInspected();
-	
+
 	// We want to refresh if the component that is done generating has generated the current graph being edited,
 	// or if it is the root of the current stack being inspected (for subgraphs to also be refreshed).
 	// If we don't have a component, we refresh nonetheless.
-	bool bShouldRefresh = !Source || StackBeingInspected.GetRootSource() == Source || Source->GetExecutionState().GetGraph() == PCGGraphBeingEdited;
+	bool bShouldRefresh = !Source || StackBeingInspected.GetRootSource() == Source || Source->GetExecutionState().GetGraph() == MainGraph;
 
 	// Additionally, if we are not inspecting but the component that's done executing contains this graph, then we should also update.
 	IPCGEditorModule* PCGEditorModule = IPCGEditorModule::Get();
-	if (!bShouldRefresh && !PCGSourceBeingInspected && PCGGraphBeingEdited && PCGEditorModule)
+	if (!bShouldRefresh && !PCGSourceBeingInspected && MainGraph && PCGEditorModule)
 	{
-		TArray<FPCGStackSharedPtr> ExecutedStacksContainingThisGraph = PCGEditorModule->GetExecutedStacksPtrs(Source, PCGGraphBeingEdited);
+		TArray<FPCGStackSharedPtr> ExecutedStacksContainingThisGraph = PCGEditorModule->GetExecutedStacksPtrs(Source, MainGraph);
 		bShouldRefresh |= !ExecutedStacksContainingThisGraph.IsEmpty();
 	}
 
@@ -3320,7 +3723,7 @@ void FPCGEditor::OnSourceGenerationDone(IPCGBaseSubsystem* Subsystem, IPCGGraphE
 		return;
 	}
 
-	LastExecutionStatus = {Source, Status};
+	LastExecutionStatus = { Source, Status };
 
 	OnSourceGenerated(Source);
 
@@ -3329,8 +3732,8 @@ void FPCGEditor::OnSourceGenerationDone(IPCGBaseSubsystem* Subsystem, IPCGGraphE
 	const bool CacheDebuggingEnabled = Subsystem && Subsystem->IsGraphCacheDebuggingEnabled();
 
 	// Refresh nodes to report any errors/warnings, and to display culling state after execution.
-	check(PCGEditorGraph);
-	for (UEdGraphNode* Node : PCGEditorGraph->Nodes)
+	check(GetMainEditorGraph());
+	for (UEdGraphNode* Node : GetMainEditorGraph()->Nodes)
 	{
 		if (UPCGEditorGraphNodeBase* PCGEditorGraphNode = Cast<UPCGEditorGraphNodeBase>(Node))
 		{
@@ -3384,6 +3787,14 @@ void FPCGEditor::OnGraphChanged(UPCGGraphInterface* InGraph, EPCGChangeType Chan
 {
 	UpdateDefaultExecutionSource();
 
+	if (!!(ChangeType & EPCGChangeType::Structural))
+	{
+		if (EmbeddedSubgraphsWidget.IsValid())
+		{
+			EmbeddedSubgraphsWidget->RequestRefresh();
+		}
+	}
+
 	if (!!(ChangeType & EPCGChangeType::ShaderSource))
 	{
 		// Flush the shader file cache in case we are editing engine or data interface shaders.
@@ -3391,7 +3802,7 @@ void FPCGEditor::OnGraphChanged(UPCGGraphInterface* InGraph, EPCGChangeType Chan
 		FlushShaderFileCache();
 	}
 
-	if (!!(ChangeType & EPCGChangeType::GraphCustomization))
+	if (!!(ChangeType & (EPCGChangeType::GraphCustomization | EPCGChangeType::Cosmetic)))
 	{
 		if (PaletteWidget)
 		{
@@ -3410,6 +3821,8 @@ void FPCGEditor::OnGraphChanged(UPCGGraphInterface* InGraph, EPCGChangeType Chan
 			}
 		}
 	}
+
+	CloseInvalidDocuments();
 }
 
 void FPCGEditor::OnNodeSourceCompiled(const UPCGNode* InNode, const FPCGCompilerDiagnostics& InDiagnostics)
@@ -3474,15 +3887,15 @@ void FPCGEditor::RefreshViews()
 
 void FPCGEditor::UpdateDefaultExecutionSource()
 {
-	if (PCGGraphBeingEdited->IsStandaloneGraph() && !PCGDefaultExecutionSource)
+	if (MainGraph->IsStandaloneGraph() && !PCGDefaultExecutionSource)
 	{
 		if (UPCGEngineSubsystem* EngineSubsystem = UPCGEngineSubsystem::Get())
 		{
-			PCGDefaultExecutionSource = EngineSubsystem->CreateExecutionSource({ PCGGraphBeingEdited });
+			PCGDefaultExecutionSource = EngineSubsystem->CreateExecutionSource({ MainGraph });
 			RefreshViews();
 		}
 	}
-	else if (!PCGGraphBeingEdited->IsStandaloneGraph() && PCGDefaultExecutionSource)
+	else if (!MainGraph->IsStandaloneGraph() && PCGDefaultExecutionSource)
 	{
 		if (IPCGEditorModule* PCGEditorModule = IPCGEditorModule::Get())
 		{
@@ -3499,4 +3912,26 @@ TSharedRef<FTabManager::FLayout> FPCGEditor::GetDefaultLayout() const
 	return PCGEditorDefaultMode::GetTabLayout().ToSharedRef();
 }
 
+void FPCGEditor::CloseInvalidDocuments()
+{
+	TArray<TSharedPtr<SDockTab>> AllDocumentTabs = DocumentManager->GetAllDocumentTabs();
+	for (TSharedPtr<SDockTab> Tab : AllDocumentTabs)
+	{
+		if (Tab->GetLayoutIdentifier() == FTabId(PCGEditorTabs::GraphEditorID))
+		{
+			TSharedRef<SGraphEditor> GraphEditor = StaticCastSharedRef<SGraphEditor>(Tab->GetContent());
+			if (UPCGEditorGraph* EditorGraph = Cast<UPCGEditorGraph>(GraphEditor->GetCurrentGraph()))
+			{
+				if (!EditorGraph->GetPCGGraph())
+				{
+					TSharedRef<FTabPayload_UObject> Payload = FTabPayload_UObject::Make(EditorGraph);
+					DocumentManager->CloseTab(Payload);
+				}
+			}
+		}
+	}
+
+	DocumentManager->CleanInvalidTabs();
+}
+
 #undef LOCTEXT_NAMESPACE
\ No newline at end of file
diff --git a/Engine/Plugins/PCG/Source/PCGEditor/Private/PCGEditorCommands.cpp b/Engine/Plugins/PCG/Source/PCGEditor/Private/PCGEditorCommands.cpp
--- a/Engine/Plugins/PCG/Source/PCGEditor/Private/PCGEditorCommands.cpp
+++ b/Engine/Plugins/PCG/Source/PCGEditor/Private/PCGEditorCommands.cpp
@@ -22,7 +22,8 @@ FPCGEditorCommands::FPCGEditorCommands()
 
 void FPCGEditorCommands::RegisterCommands()
 {
-	UI_COMMAND(CollapseNodes, "Collapse into Subgraph", "Collapse selected nodes into a separate PCGGraph asset.", EUserInterfaceActionType::Button, FInputChord(EModifierKey::Alt, EKeys::J));
+	UI_COMMAND(CollapseNodes, "Graph Asset", "Collapse selected nodes into a separate PCGGraph asset.", EUserInterfaceActionType::Button, FInputChord(EModifierKey::Alt, EKeys::J));
+	UI_COMMAND(CollapseNodesToEmbeddedSubgraph, "Embedded Graph", "Collapse selected nodes into an embedded subgraph.", EUserInterfaceActionType::Button, FInputChord());
 	UI_COMMAND(ExportNodes, "Export nodes to PCGSettings", "Exports selected nodes to separate and reusable PCGSettings assets.", EUserInterfaceActionType::Button, FInputChord());
 	UI_COMMAND(ConvertToStandaloneNodes, "Convert to Standalone Nodes", "Converts instanced nodes to standalone nodes.", EUserInterfaceActionType::Button, FInputChord());
 	UI_COMMAND(Find, "Find", "Finds PCG nodes and comments in the current graph.", EUserInterfaceActionType::Button, FInputChord(EModifierKey::Control, EKeys::F));
@@ -33,6 +34,7 @@ void FPCGEditorCommands::RegisterCommands()
 	UI_COMMAND(RunDeterminismNodeTest, "Run Determinism Test on Node", "Evaluate the current node for determinism.", EUserInterfaceActionType::Button, FInputChord(EModifierKey::Alt, EKeys::T));
 	UI_COMMAND(RunDeterminismGraphTest, "Graph Determinism Test", "Evaluate the current graph for determinism.", EUserInterfaceActionType::Button, FInputChord());
 	UI_COMMAND(EditGraphSettings, "Graph Settings", "Edit the graph settings.", EUserInterfaceActionType::ToggleButton, FInputChord());
+	UI_COMMAND(EditGraph, "Graph", "Edit the graph.", EUserInterfaceActionType::Button, FInputChord());
 	UI_COMMAND(ToggleGraphParams, "Graph Parameters", "Open the graph settings panel.", EUserInterfaceActionType::ToggleButton, FInputChord());
 	UI_COMMAND(CancelExecution, "Cancel Execution", "Cancels the execution of the current graph", EUserInterfaceActionType::Button, FInputChord(EModifierKey::Control, EKeys::Escape));
 	UI_COMMAND(ToggleEnabled, "Toggle Enabled", "Toggle node enabled state for selected nodes.", EUserInterfaceActionType::ToggleButton, FInputChord(EKeys::E));
diff --git a/Engine/Plugins/PCG/Source/PCGEditor/Private/PCGEditorDefaultMode.cpp b/Engine/Plugins/PCG/Source/PCGEditor/Private/PCGEditorDefaultMode.cpp
--- a/Engine/Plugins/PCG/Source/PCGEditor/Private/PCGEditorDefaultMode.cpp
+++ b/Engine/Plugins/PCG/Source/PCGEditor/Private/PCGEditorDefaultMode.cpp
@@ -46,7 +46,6 @@ namespace PCGEditorDefaultMode
 						FTabManager::NewStack()
 						->SetSizeCoefficient(0.6f)
 						->AddTab(PCGEditorTabs::GraphEditorID, ETabState::OpenedTab)
-						->AddTab(PCGEditorTabs::PCGEditorDocumentID, ETabState::ClosedTab)
 						->SetHideTabWell(true)
 					)
 					->Split // Details View
@@ -66,6 +65,7 @@ namespace PCGEditorDefaultMode
 						->SetSizeCoefficient(0.2f)
 						->AddTab(PCGEditorTabs::DebugObjectID, ETabState::OpenedTab)
 						->AddTab(PCGEditorTabs::UserParamsID, ETabState::OpenedTab)
+						->AddTab(PCGEditorTabs::EmbeddedSubgraphsID, ETabState::OpenedTab)
 						->SetForegroundTab(FTabId(PCGEditorTabs::DebugObjectID))
 					)
 					->Split // ALV, Profiling, Find, Determinism
@@ -104,7 +104,6 @@ FPCGEditorDefaultMode::FPCGEditorDefaultMode(TSharedPtr<FPCGEditor> InPCGEditor)
 
 	TabInfo DefaultTabInfo[] =
 	{
-		{PCGEditorTabs::GraphEditorID, PCGEditorDefaultMode::FCreateTabFactoryDelegate::CreateSP(InPCGEditor.Get(), &FPCGEditor::CreateGraphEditorTabFactory)},
 		{PCGEditorTabs::PropertyDetailsID[0], PCGEditorDefaultMode::FCreateTabFactoryDelegate::CreateSP(InPCGEditor.Get(), &FPCGEditor::CreatePropertyDetailsTabFactory, 0, PropertyDetailsGroup)},
 		{PCGEditorTabs::PropertyDetailsID[1], PCGEditorDefaultMode::FCreateTabFactoryDelegate::CreateSP(InPCGEditor.Get(), &FPCGEditor::CreatePropertyDetailsTabFactory, 1, PropertyDetailsGroup)},
 		{PCGEditorTabs::PropertyDetailsID[2], PCGEditorDefaultMode::FCreateTabFactoryDelegate::CreateSP(InPCGEditor.Get(), &FPCGEditor::CreatePropertyDetailsTabFactory, 2, PropertyDetailsGroup)},
@@ -125,6 +124,7 @@ FPCGEditorDefaultMode::FPCGEditorDefaultMode(TSharedPtr<FPCGEditor> InPCGEditor)
 		{PCGEditorTabs::ViewportID[1], PCGEditorDefaultMode::FCreateTabFactoryDelegate::CreateSP(InPCGEditor.Get(), &FPCGEditor::CreateViewportTabFactory, 1, ViewportGroup)},
 		{PCGEditorTabs::ViewportID[2], PCGEditorDefaultMode::FCreateTabFactoryDelegate::CreateSP(InPCGEditor.Get(), &FPCGEditor::CreateViewportTabFactory, 2, ViewportGroup)},
 		{PCGEditorTabs::ViewportID[3], PCGEditorDefaultMode::FCreateTabFactoryDelegate::CreateSP(InPCGEditor.Get(), &FPCGEditor::CreateViewportTabFactory, 3, ViewportGroup)},
+		{PCGEditorTabs::EmbeddedSubgraphsID, PCGEditorDefaultMode::FCreateTabFactoryDelegate::CreateSP(InPCGEditor.Get(), &FPCGEditor::CreateEmbeddedSubgraphsTabFactory)}
 	};
 
 	for (TabInfo& Info : DefaultTabInfo)
diff --git a/Engine/Plugins/PCG/Source/PCGEditor/Private/PCGEditorGraph.cpp b/Engine/Plugins/PCG/Source/PCGEditor/Private/PCGEditorGraph.cpp
--- a/Engine/Plugins/PCG/Source/PCGEditor/Private/PCGEditorGraph.cpp
+++ b/Engine/Plugins/PCG/Source/PCGEditor/Private/PCGEditorGraph.cpp
@@ -58,17 +58,18 @@ namespace PCGEditorGraphUtils
 
 void UPCGEditorGraph::InitFromNodeGraph(UPCGGraph* InPCGGraph)
 {
-	check(InPCGGraph && !PCGGraph);
-	PCGGraph = InPCGGraph;
+	check(InPCGGraph && WeakPCGGraph.IsExplicitlyNull());
+	WeakPCGGraph = InPCGGraph;
 
-	PCGGraph->OnGraphParametersChangedDelegate.AddUObject(this, &UPCGEditorGraph::OnGraphUserParametersChanged);
-	PCGGraph->OnGraphChangedDelegate.AddUObject(this, &UPCGEditorGraph::OnGraphChanged);
+	InPCGGraph->OnGraphParametersChangedDelegate.AddUObject(this, &UPCGEditorGraph::OnGraphUserParametersChanged);
+	InPCGGraph->OnGraphChangedDelegate.AddUObject(this, &UPCGEditorGraph::OnGraphChanged);
 
 	ReconstructGraph();
 }
 
 void UPCGEditorGraph::ReconstructGraph()
 {
+	UPCGGraph* PCGGraph = GetPCGGraph();
 	check(PCGGraph);
 
 	// If there are already some nodes, remove all of them.
@@ -150,16 +151,16 @@ void UPCGEditorGraph::BeginDestroy()
 {
 	Super::BeginDestroy();
 
-	OnClose();
+	SaveState();
 
-	if (PCGGraph)
+	if (UPCGGraph* PCGGraph = GetPCGGraph())
 	{
 		PCGGraph->OnGraphChangedDelegate.RemoveAll(this);
 		PCGGraph->OnGraphParametersChangedDelegate.RemoveAll(this);
 	}
 }
 
-void UPCGEditorGraph::OnClose()
+void UPCGEditorGraph::SaveState()
 {
 	ReplicateExtraNodes();
 }
@@ -184,7 +185,7 @@ void UPCGEditorGraph::CreateLinks(UPCGEditorGraphNodeBase* GraphNode, bool bCrea
 
 void UPCGEditorGraph::ReplicateExtraNodes() const
 {
-	if (PCGGraph)
+	if (UPCGGraph* PCGGraph = GetPCGGraph())
 	{
 		TArray<TObjectPtr<const UObject>> ExtraNodes;
 		TArray<FPCGGraphCommentNodeData> CommentData;
@@ -225,6 +226,14 @@ void UPCGEditorGraph::UpdateVisualizations(IPCGGraphExecutionSource* PCGSourceBe
 	}
 }
 
+FText UPCGEditorGraph::GetDisplayName() const
+{
+	FGraphDisplayInfo DisplayInfo;
+	GetSchema()->GetGraphDisplayInformation(*this, DisplayInfo);
+
+	return DisplayInfo.DisplayName;
+}
+
 const UPCGEditorGraphNodeBase* UPCGEditorGraph::GetEditorNodeFromPCGNode(const UPCGNode* InPCGNode) const
 {
 	if (ensure(InPCGNode))
@@ -409,6 +418,7 @@ void UPCGEditorGraph::OnGraphChanged(UPCGGraphInterface* InGraph, EPCGChangeType
 
 void UPCGEditorGraph::OnGraphUserParametersChanged(UPCGGraphInterface* InGraph, EPCGGraphParameterEvent ChangeType, FName ChangedPropertyName)
 {
+	UPCGGraph* PCGGraph = GetPCGGraph();
 	if (InGraph != PCGGraph)
 	{
 		return;
diff --git a/Engine/Plugins/PCG/Source/PCGEditor/Private/PCGEditorGraph.h b/Engine/Plugins/PCG/Source/PCGEditor/Private/PCGEditorGraph.h
--- a/Engine/Plugins/PCG/Source/PCGEditor/Private/PCGEditorGraph.h
+++ b/Engine/Plugins/PCG/Source/PCGEditor/Private/PCGEditorGraph.h
@@ -37,7 +37,7 @@ public:
 	void ReconstructGraph();
 
 	/** When the editor is closing */
-	void OnClose();
+	void SaveState();
 
 	/** Creates the links for a given node */
 	void CreateLinks(UPCGEditorGraphNodeBase* InGraphNode, bool bCreateInbound, bool bCreateOutbound);
@@ -45,11 +45,13 @@ public:
 	/** To be called everytime we need to replicate our extra nodes to the underlying PCGGraph */
 	void ReplicateExtraNodes() const;
 
-	UPCGGraph* GetPCGGraph() const { return PCGGraph; }
+	UPCGGraph* GetPCGGraph() const { return WeakPCGGraph.Get(); }
 
 	void SetEditor(TWeakPtr<FPCGEditor> InEditor) { PCGEditor = InEditor; }
 	TWeakPtr<FPCGEditor> GetEditor() const { return PCGEditor; }
 
+	FText GetDisplayName() const;
+
 	/** Updates the grid size visualization in the editor. */
 	void UpdateVisualizations(IPCGGraphExecutionSource* PCGSourceBeingInspected, const FPCGStack* PCGStackBeingInspected);
 
@@ -75,7 +77,7 @@ protected:
 
 private:
 	UPROPERTY()
-	TObjectPtr<UPCGGraph> PCGGraph = nullptr;
+	TWeakObjectPtr<UPCGGraph> WeakPCGGraph;
 
 	TWeakPtr<FPCGEditor> PCGEditor = nullptr;
 };
diff --git a/Engine/Plugins/PCG/Source/PCGEditor/Private/PCGEditorModule.cpp b/Engine/Plugins/PCG/Source/PCGEditor/Private/PCGEditorModule.cpp
--- a/Engine/Plugins/PCG/Source/PCGEditor/Private/PCGEditorModule.cpp
+++ b/Engine/Plugins/PCG/Source/PCGEditor/Private/PCGEditorModule.cpp
@@ -28,6 +28,7 @@
 #include "Grid/PCGPartitionActor.h"
 #include "WorldPartitionBuilder/PCGWorldPartitionBuilder.h"
 
+#include "PCGEditor.h"
 #include "PCGEditorCommands.h"
 #include "PCGEditorGraph.h"
 #include "PCGEditorMenuUtils.h"
@@ -1183,6 +1184,11 @@ void FPCGEditorModule::OnGraphPreSave(UPCGGraph* Graph, FObjectPreSaveContext Ob
 	// No need to do it on cooking
 	if (!ObjectSaveContext.IsProceduralSave())
 	{
+		if (TSharedPtr<FPCGEditor> PCGEditor = Graph->PCGEditorGraph->GetEditor().Pin())
+		{
+			PCGEditor->OnGraphPreSave(Graph);
+		}
+
 		Graph->PCGEditorGraph->ReplicateExtraNodes();
 	}
 }
diff --git a/Engine/Plugins/PCG/Source/PCGEditor/Private/PCGEditorStyle.cpp b/Engine/Plugins/PCG/Source/PCGEditor/Private/PCGEditorStyle.cpp
--- a/Engine/Plugins/PCG/Source/PCGEditor/Private/PCGEditorStyle.cpp
+++ b/Engine/Plugins/PCG/Source/PCGEditor/Private/PCGEditorStyle.cpp
@@ -140,6 +140,7 @@ FPCGEditorStyle::FPCGEditorStyle() : FSlateStyleSet("PCGEditorStyle")
 	Set("PCG.Command.StopRegen", new IMAGE_BRUSH_SVG("Style/PCG_Command_StopRegen", Icon20x20));
 	Set("PCG.Command.OpenGraphParams", new IMAGE_BRUSH_SVG("Style/PCG_Graph_OpenGraphParams", Icon20x20));
 	Set("PCG.Command.GraphSettings", new IMAGE_BRUSH_SVG("Style/PCG_Command_GraphSettings", Icon20x20));
+	Set("PCG.Command.EditGraph", new IMAGE_BRUSH_SVG("Style/PCG_Command_EditGraph", Icon20x20));
 	Set("PCG.Command.OpenDebugTreeTab", new IMAGE_BRUSH_SVG("Style/PCG_Command_OpenDebugTreeTab", Icon20x20));
 	Set("PCG.Command.RunDeterminismTest", new IMAGE_BRUSH_SVG("Style/PCG_Command_RunDeterminismTest", Icon20x20));
 
diff --git a/Engine/Plugins/PCG/Source/PCGEditor/Private/PCGEditorTabFactories.cpp b/Engine/Plugins/PCG/Source/PCGEditor/Private/PCGEditorTabFactories.cpp
--- a/Engine/Plugins/PCG/Source/PCGEditor/Private/PCGEditorTabFactories.cpp
+++ b/Engine/Plugins/PCG/Source/PCGEditor/Private/PCGEditorTabFactories.cpp
@@ -2,7 +2,12 @@
 
 #include "PCGEditorTabFactories.h"
 
+#include "PCGGraph.h"
 #include "PCGEditor.h"
+#include "PCGEditorGraph.h"
+#include "PCGEditorStyle.h"
+
+#include "Widgets/Docking/SDockTab.h"
 
 #include "Widgets/Docking/SDockTab.h"
 
@@ -60,5 +65,91 @@ void FPCGGenericTabFactory::OnTabClosed(TSharedRef<SDockTab> DockTab) const
 	OnTabClosedDelegate.ExecuteIfBound(GetIdentifier());
 }
 
+FPCGGraphEditorDocumentTabFactory::FPCGGraphEditorDocumentTabFactory(TSharedPtr<class FPCGEditor> InEditorPtr, FOnCreateGraphEditorWidget InCreateGraphEditorWidgetCallback)
+	: FDocumentTabFactoryForObjects<UPCGEditorGraph>(PCGEditorTabs::GraphEditorID, InEditorPtr)
+	, OnCreateGraphEditorWidget(InCreateGraphEditorWidgetCallback)
+{
+}
+
+TSharedPtr<FPCGEditor> FPCGGraphEditorDocumentTabFactory::GetPCGEditor() const
+{
+	return StaticCastSharedPtr<FPCGEditor>(HostingApp.Pin());
+}
+
+void FPCGGraphEditorDocumentTabFactory::OnTabActivated(TSharedPtr<SDockTab> InTab) const
+{
+	TSharedRef<SGraphEditor> GraphEditor = StaticCastSharedRef<SGraphEditor>(InTab->GetContent());
+	GetPCGEditor()->OnGraphEditorFocused(GraphEditor);
+}
+
+void FPCGGraphEditorDocumentTabFactory::OnTabBackgrounded(TSharedPtr<SDockTab> InTab) const
+{
+	TSharedRef<SGraphEditor> GraphEditor = StaticCastSharedRef<SGraphEditor>(InTab->GetContent());
+	GetPCGEditor()->OnGraphEditorBackgrounded(GraphEditor);
+}
+
+void FPCGGraphEditorDocumentTabFactory::OnTabRefreshed(TSharedPtr<SDockTab> InTab) const
+{
+	TSharedRef<SGraphEditor> GraphEditor = StaticCastSharedRef<SGraphEditor>(InTab->GetContent());
+	GraphEditor->NotifyGraphChanged();
+}
+
+void FPCGGraphEditorDocumentTabFactory::SaveState(TSharedPtr<SDockTab> InTab, TSharedPtr<FTabPayload> InPayload) const
+{
+	TSharedRef<SGraphEditor> GraphEditor = StaticCastSharedRef<SGraphEditor>(InTab->GetContent());
+
+	FVector2f ViewLocation;
+	float ZoomAmount;
+	GraphEditor->GetViewLocation(ViewLocation, ZoomAmount);
+
+	UPCGEditorGraph* EditorGraph = InPayload.IsValid() && InPayload->IsValid() ? FTabPayload_UObject::CastChecked<UPCGEditorGraph>(InPayload) : nullptr;
+	UPCGGraph* Graph = EditorGraph ? EditorGraph->GetPCGGraph() : nullptr;
+	if (!Graph)
+	{
+		return;
+	}
+	
+	UPCGGraph* OuterGraph = Graph->GetTypedOuter<UPCGGraph>();
+	UPCGGraph* MainGraph = OuterGraph ? OuterGraph : Graph;
+
+	if (MainGraph == Graph || MainGraph->ContainsEmbeddedSubgraph(Graph))
+	{
+		MainGraph->LastEditedDocuments.Add(FPCGGraphDocumentInfo(Graph, ViewLocation, ZoomAmount));
+
+		EditorGraph->SaveState();
+	}
+}
+
+FText FPCGGraphEditorDocumentTabFactory::GetTabTitle(TWeakObjectPtr<UPCGEditorGraph> InWeakEditorGraph)
+{
+	if (UPCGEditorGraph* EditorGraph = InWeakEditorGraph.Get())
+	{
+		return EditorGraph->GetDisplayName();
+	}
+
+	return LOCTEXT("GraphEditorDocumentTabInvalidTitle", "Invalid");
+}
+
+TAttribute<FText> FPCGGraphEditorDocumentTabFactory::ConstructTabNameForObject(UPCGEditorGraph* InEditorGraph) const
+{
+	return TAttribute<FText>::Create(TAttribute<FText>::FGetter::CreateStatic(&FPCGGraphEditorDocumentTabFactory::GetTabTitle, TWeakObjectPtr<UPCGEditorGraph>(InEditorGraph)));
+}
+
+TSharedRef<SWidget> FPCGGraphEditorDocumentTabFactory::CreateTabBodyForObject(const FWorkflowTabSpawnInfo& InInfo, UPCGEditorGraph* InEditorGraph) const
+{
+	return OnCreateGraphEditorWidget.Execute(InInfo.TabInfo, InEditorGraph);
+}
+
+const FSlateBrush* FPCGGraphEditorDocumentTabFactory::GetTabIconForObject(const FWorkflowTabSpawnInfo& InInfo, UPCGEditorGraph* InEditorGraph) const
+{
+	return FPCGEditorStyle::Get().GetBrush("ClassIcon.PCGGraph");
+}
+
+bool FPCGGraphEditorDocumentTabFactory::SupportsObjectType(UObject* InDocumentID) const
+{
+	const UPCGEditorGraph* EditorGraph = Cast<UPCGEditorGraph>(InDocumentID);
+
+	return EditorGraph && EditorGraph->GetPCGGraph();
+}
 
 #undef LOCTEXT_NAMESPACE
\ No newline at end of file
diff --git a/Engine/Plugins/PCG/Source/PCGEditor/Private/Schema/PCGEditorGraphSchema.cpp b/Engine/Plugins/PCG/Source/PCGEditor/Private/Schema/PCGEditorGraphSchema.cpp
--- a/Engine/Plugins/PCG/Source/PCGEditor/Private/Schema/PCGEditorGraphSchema.cpp
+++ b/Engine/Plugins/PCG/Source/PCGEditor/Private/Schema/PCGEditorGraphSchema.cpp
@@ -646,6 +646,16 @@ void UPCGEditorGraphSchema::GetSubgraphElementActions(FGraphActionMenuBuilder& A
 
 		return true;
 	});
+
+	if (ActionFilter.Graph && ActionFilter.Graph->GetPCGGraph())
+	{
+		for (UPCGGraph* EmbeddedSubgraph : ActionFilter.Graph->GetPCGGraph()->GetEmbeddedSubgraphs())
+		{
+			TSharedPtr<FPCGEditorGraphSchemaAction_NewSubgraphElement> NewSubgraphAction(new FPCGEditorGraphSchemaAction_NewSubgraphElement(EmbeddedSubgraph->Category, EmbeddedSubgraph->Title, EmbeddedSubgraph->Description, 0));
+			NewSubgraphAction->SubgraphObjectPath = EmbeddedSubgraph;
+			ActionMenuBuilder.AddAction(NewSubgraphAction);
+		}
+	}
 }
 
 void UPCGEditorGraphSchema::GetExtraElementActions(FGraphActionMenuBuilder& ActionMenuBuilder, const FPCGActionsFilter& ActionFilter) const
@@ -890,6 +900,21 @@ void UPCGEditorGraphSchema::OnPinConnectionDoubleCicked(UEdGraphPin* PinA, UEdGr
 	}
 }
 
+void UPCGEditorGraphSchema::GetGraphDisplayInformation(const UEdGraph& Graph, /*out*/ FGraphDisplayInfo& DisplayInfo) const
+{
+	if (const UPCGEditorGraph* EditorGraph = Cast<UPCGEditorGraph>(&Graph))
+	{
+		if (const UPCGGraph* PCGGraph = EditorGraph->GetPCGGraph())
+		{
+			DisplayInfo.DisplayName = PCGGraph->IsEmbeddedSubgraph() ? PCGGraph->GetDisplayName() : LOCTEXT("DefaultGraphDisplayInformation", "Graph");
+			DisplayInfo.PlainName = FText::FromString(PCGGraph->GetName());
+			return;
+		}
+	}
+
+	Super::GetGraphDisplayInformation(Graph, DisplayInfo);
+}
+
 const FSlateBrush* UPCGEditorGraphSchema::GetMetadataTypeSlateBrush(const EPCGContainerType ContainerType) const
 {
 	switch (ContainerType)
diff --git a/Engine/Plugins/PCG/Source/PCGEditor/Private/Widgets/SPCGEditorGraphDebugObjectTree.cpp b/Engine/Plugins/PCG/Source/PCGEditor/Private/Widgets/SPCGEditorGraphDebugObjectTree.cpp
--- a/Engine/Plugins/PCG/Source/PCGEditor/Private/Widgets/SPCGEditorGraphDebugObjectTree.cpp
+++ b/Engine/Plugins/PCG/Source/PCGEditor/Private/Widgets/SPCGEditorGraphDebugObjectTree.cpp
@@ -613,10 +613,9 @@ void SPCGEditorGraphDebugObjectTree::OnTick(double InCurrentTime, bool bForceRef
 			}
 			else // Select first occurrence otherwise
 			{
-				const UPCGGraph* CurrentGraph = PCGEditor.Pin()->GetPCGGraph();
 				for (const FPCGEditorGraphDebugObjectItemPtr& Item : AllGraphItems)
 				{
-					if (Item->GetPCGStack() && Item->GetPCGStack()->GetGraphForCurrentFrame() == CurrentGraph)
+					if (Item->GetPCGStack() && IsMainOrEmbeddedGraph(Item->GetPCGStack()->GetGraphForCurrentFrame()))
 					{
 						ExpandAndSelectDebugObject(Item);
 						break;
@@ -686,9 +685,8 @@ bool SPCGEditorGraphDebugObjectTree::SetDebugObjectFromStackFromAnotherEditor(co
 {
 	RefreshTree();
 
-	const UPCGGraph* CurrentGraph = PCGEditor.Pin()->GetPCGGraph();
 	// If the given stack is already valid for this editor, try to select it as-is.
-	if (InStack.GetGraphForCurrentFrame() == CurrentGraph)
+	if (IsMainOrEmbeddedGraph(InStack.GetGraphForCurrentFrame()))
 	{
 		if (FPCGEditorGraphDebugObjectItemPtr Item = GetItemFromStack(InStack))
 		{
@@ -702,7 +700,7 @@ bool SPCGEditorGraphDebugObjectTree::SetDebugObjectFromStackFromAnotherEditor(co
 	{
 		const FPCGStack* ItemStack = Item->GetPCGStack();
 
-		if (ItemStack && ItemStack->BeginsWith(InStack) && ItemStack->GetGraphForCurrentFrame() == CurrentGraph)
+		if (ItemStack && ItemStack->BeginsWith(InStack) && IsMainOrEmbeddedGraph(ItemStack->GetGraphForCurrentFrame()))
 		{
 			ExpandAndSelectDebugObject(Item);
 			return true;
@@ -815,8 +813,8 @@ bool SPCGEditorGraphDebugObjectTree::IsOpen() const
 FPCGEditorGraphDebugObjectItemPtr SPCGEditorGraphDebugObjectTree::GetFirstDebugObjectFromSelection() const
 {
 	TRACE_CPUPROFILER_EVENT_SCOPE(SPCGEditorGraphDebugObjectTree::GetFirstDebugObjectFromSelection);
-	const UPCGGraph* PCGGraph = GetPCGGraph();
-	if (!PCGGraph)
+	const UPCGGraph* MainGraph = GetMainGraph();
+	if (!MainGraph)
 	{
 		return nullptr;
 	}
@@ -864,7 +862,7 @@ FPCGEditorGraphDebugObjectItemPtr SPCGEditorGraphDebugObjectTree::GetFirstDebugO
 			// Look for graph in dynamic stacks.
 			if (PCGEditorModule)
 			{
-				TArray<FPCGStackSharedPtr> ExecutedStacks = PCGEditorModule->GetExecutedStacksPtrs(PCGComponent, PCGGraph);
+				TArray<FPCGStackSharedPtr> ExecutedStacks = PCGEditorModule->GetExecutedStacksPtrs(PCGComponent, MainGraph);
 				for (const FPCGStackSharedPtr& Stack : ExecutedStacks)
 				{
 					if (FPCGEditorGraphDebugObjectItemPtr Item = GetItemFromStack(*Stack))
@@ -889,15 +887,15 @@ void SPCGEditorGraphDebugObjectTree::AddStacksToTree(const TArray<FPCGStackShare
 	TMap<const FPCGStack, FPCGEditorGraphDebugObjectItemPtr>& InOutStackToItem)
 {
 	TRACE_CPUPROFILER_EVENT_SCOPE(SPCGEditorGraphDebugObjectTree::AddStacksToTree);
-	const UPCGGraph* GraphBeingEdited = GetPCGGraph();
-	if (!GraphBeingEdited)
+	const UPCGGraph* MainGraph = GetMainGraph();
+	if (!MainGraph)
 	{
 		return;
 	}
 
 	for (const FPCGStackSharedPtr& Stack : Stacks)
 	{
-		AddStackToTree(*Stack, GraphBeingEdited, InOutOwnerItems, InOutStackToItem);
+		AddStackToTree(*Stack, MainGraph, InOutOwnerItems, InOutStackToItem);
 	}
 }
 
@@ -906,15 +904,15 @@ void SPCGEditorGraphDebugObjectTree::AddStacksToTree(const TArray<FPCGStack>& St
 	TMap<const FPCGStack, FPCGEditorGraphDebugObjectItemPtr>& InOutStackToItem)
 {
 	TRACE_CPUPROFILER_EVENT_SCOPE(SPCGEditorGraphDebugObjectTree::AddStacksToTree);
-	const UPCGGraph* GraphBeingEdited = GetPCGGraph();
-	if (!GraphBeingEdited)
+	const UPCGGraph* MainGraph = GetMainGraph();
+	if (!MainGraph)
 	{
 		return;
 	}
 
 	for (const FPCGStack& Stack : Stacks)
 	{
-		AddStackToTree(Stack, GraphBeingEdited, InOutOwnerItems, InOutStackToItem);
+		AddStackToTree(Stack, MainGraph, InOutOwnerItems, InOutStackToItem);
 	}
 }
 
@@ -1100,7 +1098,7 @@ void SPCGEditorGraphDebugObjectTree::AddStackToTree(const FPCGStack& Stack,
 			// When we encounter a graph, we look at the frame index and/or preceding frames to determine the graph type.
 			if (const UPCGGraph* StackGraph = StackFrame.GetObject_GameThread<UPCGGraph>())
 			{
-				const bool bIsDebuggable = (GraphBeingEdited == StackGraph);
+				const bool bIsDebuggable = IsMainOrEmbeddedGraph(StackGraph);
 
 				// Top graph.
 				if (FrameIndex == TopGraphIndex && StackGraph == TopGraph)
@@ -1191,8 +1189,8 @@ void SPCGEditorGraphDebugObjectTree::RefreshTree()
 	AllGraphItems.Empty();
 	DebugObjectTreeView->RequestTreeRefresh();
 
-	UPCGGraph* PCGGraph = GetPCGGraph();
-	if (!PCGGraph)
+	UPCGGraph* MainGraph = GetMainGraph();
+	if (!MainGraph)
 	{
 		return;
 	}
@@ -1249,7 +1247,7 @@ void SPCGEditorGraphDebugObjectTree::RefreshTree()
 					bool bSomeStacksContainCurrentGraph = false;
 					if (Subsystem->GetStackContext(MapKey.Get<0>(), MapKey.Get<1>(), MapKey.Get<2>(), GraphStackContext))
 					{
-						bSomeStacksContainCurrentGraph = Algo::AnyOf(GraphStackContext.GetStacks(), [PCGGraph](const FPCGStack& Stack) { return Stack.HasObject(PCGGraph); });
+						bSomeStacksContainCurrentGraph = Algo::AnyOf(GraphStackContext.GetStacks(), [MainGraph](const FPCGStack& Stack) { return Stack.HasObject(MainGraph); });
 					}
 
 					MapValue = &(GraphToContexts.Add(MapKey, FGraphContextMapValue(std::move(GraphStackContext), bSomeStacksContainCurrentGraph)));
@@ -1271,7 +1269,7 @@ void SPCGEditorGraphDebugObjectTree::RefreshTree()
 		if(IPCGEditorModule* PCGEditorModule = IPCGEditorModule::Get())
 		{
 			TRACE_CPUPROFILER_EVENT_SCOPE(SPCGEditorGraphDebugObjectTree::RefreshTree::DynamicStacks);
-			TArray<FPCGStackSharedPtr> GraphStacks = PCGEditorModule->GetExecutedStacksPtrs(nullptr, PCGGraph, /*bOnlyWithSubgraphAsCurrentFrame=*/!bShowDownstream);
+			TArray<FPCGStackSharedPtr> GraphStacks = PCGEditorModule->GetExecutedStacksPtrs(nullptr, MainGraph, /*bOnlyWithSubgraphAsCurrentFrame=*/!bShowDownstream);
 			// TODO: do per component filtering?
 			AddStacksToTree(GraphStacks, OwnerItems, StackToItem);
 		}
@@ -1420,12 +1418,16 @@ void SPCGEditorGraphDebugObjectTree::OnEditorSelectionChanged(UObject* InObject)
 	IsSetDebugObjectFromSelectionEnabled.Reset();
 }
 
-UPCGGraph* SPCGEditorGraphDebugObjectTree::GetPCGGraph() const
+UPCGGraph* SPCGEditorGraphDebugObjectTree::GetFocusedGraph() const
 {
 	const TSharedPtr<FPCGEditor> PCGEditorPtr = PCGEditor.Pin();
-	const UPCGEditorGraph* PCGEditorGraph = PCGEditorPtr.IsValid() ? PCGEditorPtr->GetPCGEditorGraph() : nullptr;
+	return PCGEditorPtr.IsValid() ? PCGEditorPtr->GetFocusedGraph() : nullptr;
+}
 
-	return PCGEditorGraph ? PCGEditorGraph->GetPCGGraph() : nullptr;
+UPCGGraph* SPCGEditorGraphDebugObjectTree::GetMainGraph() const
+{
+	const TSharedPtr<FPCGEditor> PCGEditorPtr = PCGEditor.Pin();
+	return PCGEditorPtr.IsValid() ? PCGEditorPtr->GetMainGraph() : nullptr;
 }
 
 TSharedRef<ITableRow> SPCGEditorGraphDebugObjectTree::MakeTreeRowWidget(FPCGEditorGraphDebugObjectItemPtr InItem, const TSharedRef<STableViewBase>& InOwnerTable)
@@ -1493,18 +1495,18 @@ void SPCGEditorGraphDebugObjectTree::OnSelectionChanged(FPCGEditorGraphDebugObje
 		}
 	}
 
-	const UPCGGraph* CurrentGraph = PCGEditor.Pin()->GetPCGGraph();
+	const UPCGGraph* MainGraph = GetMainGraph();
 	bool bStackInspectedSet = false;
 
 	// Only attempt to inspect stacks that correspond to the edited graph. Other graphs need to be inspected in their own editor.
 	if (InItem)
 	{
 		// Basically - if the selection is something from "upstream", then we can clear the stack being inspected.
-		if (SelectedStack.HasObject(CurrentGraph))
+		if (SelectedStack.HasObject(MainGraph))
 		{
 			// If the last graph in the stack is the current graph, then we can just use that as the selected stack
 			TArray<FPCGStackFrame>& StackFrames = SelectedStack.GetStackFramesMutable();
-			while(!StackFrames.IsEmpty() && StackFrames.Last().GetObject_GameThread<UPCGGraph>() != CurrentGraph)
+			while(!StackFrames.IsEmpty() && !IsMainOrEmbeddedGraph(StackFrames.Last().GetObject_GameThread<UPCGGraph>()))
 			{
 				StackFrames.Pop();
 			}
@@ -1523,7 +1525,7 @@ void SPCGEditorGraphDebugObjectTree::OnSelectionChanged(FPCGEditorGraphDebugObje
 	}
 
 	// Finally mark the selected item as selected
-	if (SelectedStack.GetGraphForCurrentFrame() == CurrentGraph)
+	if (IsMainOrEmbeddedGraph(SelectedStack.GetGraphForCurrentFrame()))
 	{
 		if (FPCGEditorGraphDebugObjectItemPtr CurrentSelectedItem = GetItemFromStack(SelectedStack))
 		{
@@ -1573,6 +1575,17 @@ void SPCGEditorGraphDebugObjectTree::OnSetExpansionRecursive(FPCGEditorGraphDebu
 	}
 }
 
+bool SPCGEditorGraphDebugObjectTree::IsMainOrEmbeddedGraph(const UPCGGraph* InGraph) const
+{
+	if (!InGraph)
+	{
+		return false;
+	}
+
+	const UPCGGraph* MainGraph = GetMainGraph();
+	return InGraph == MainGraph || InGraph->GetEmbeddedParentGraph() == MainGraph;
+}
+
 void SPCGEditorGraphDebugObjectTree::ExpandAndSelectFirstLeafDebugObject(const FPCGEditorGraphDebugObjectItemPtr& InItem)
 {
 	if (!InItem.IsValid())
@@ -1586,7 +1599,6 @@ void SPCGEditorGraphDebugObjectTree::ExpandAndSelectFirstLeafDebugObject(const F
 	// Find first occurrence of tree (breadth-first) of current graph in children.
 	FPCGEditorGraphDebugObjectItemPtr ItemToSelect;
 
-	const UPCGGraph* CurrentGraph = PCGEditor.Pin()->GetPCGGraph();
 	TArray<FPCGEditorGraphDebugObjectItemPtr> ToVisit = InItem->GetChildren().Array();
 	int32 VisitIndex = 0;
 	
@@ -1596,7 +1608,7 @@ void SPCGEditorGraphDebugObjectTree::ExpandAndSelectFirstLeafDebugObject(const F
 		const FPCGStack* ItemStack = Item ? Item->GetPCGStack() : nullptr;
 		if (ItemStack)
 		{
-			if(!ItemStack->GetStackFrames().IsEmpty() && ItemStack->GetStackFrames().Last().GetObject_GameThread<UPCGGraph>() == CurrentGraph)
+			if(!ItemStack->GetStackFrames().IsEmpty() && IsMainOrEmbeddedGraph(ItemStack->GetStackFrames().Last().GetObject_GameThread<UPCGGraph>()))
 			{
 				ItemToSelect = Item;
 				break;
@@ -1694,7 +1706,7 @@ void SPCGEditorGraphDebugObjectTree::JumpToGraphInTree(const FPCGEditorGraphDebu
 
 	// If the item is about this graph, then we'll jump to this instance's caller - we'll walk back until we find a graph (e.g. parent)
 	// and the next entry should be our caller node.
-	if (Item->IsDebuggable())
+	if (Item->IsDebuggable() && !IsMainOrEmbeddedGraph(Item->GetPCGGraph()))
 	{
 		// Debuggable target graphs correspond to the currently edited graph. For this case open the parent graph and jump to the corresponding subgraph node.
 		// Search the stack for the parent graph.
@@ -1720,19 +1732,24 @@ void SPCGEditorGraphDebugObjectTree::JumpToGraphInTree(const FPCGEditorGraphDebu
 
 	if (JumpToPCGGraph)
 	{
-		GEditor->GetEditorSubsystem<UAssetEditorSubsystem>()->OpenEditorForAsset(JumpToPCGGraph);
-		IAssetEditorInstance* EditorInstance = GEditor->GetEditorSubsystem<UAssetEditorSubsystem>()->FindEditorForAsset(JumpToPCGGraph, /*bFocusIfOpen*/true);
-		FPCGEditor* OtherPCGEditor = static_cast<FPCGEditor*>(EditorInstance);
-
-		if (OtherPCGEditor)
+		if (TSharedPtr<FPCGEditor> PinnedEditor = PCGEditor.Pin())
 		{
-			// Implementation note - for selection purposes we had culled the stack to be in the local referential
-			// however in the case of the jump here we want the full stack
-			OtherPCGEditor->SetStackBeingInspectedFromAnotherEditor(Stack);
+			FPCGEditor* CurrentPCGEditor = PinnedEditor->OpenEditorForGraph(JumpToPCGGraph);
+			if (!CurrentPCGEditor)
+			{
+				return;
+			}
+
+			if (CurrentPCGEditor != PinnedEditor.Get())
+			{
+				// Implementation note - for selection purposes we had culled the stack to be in the local referential
+				// however in the case of the jump here we want the full stack
+				CurrentPCGEditor->SetStackBeingInspectedFromAnotherEditor(Stack);
+			}
 
 			if (JumpToPCGNode)
 			{
-				OtherPCGEditor->JumpToNode(JumpToPCGNode);
+				CurrentPCGEditor->JumpToNode(JumpToPCGNode);
 			}
 		}
 	}
@@ -1765,9 +1782,8 @@ void SPCGEditorGraphDebugObjectTree::FocusOnItem(const FPCGEditorGraphDebugObjec
 		return;
 	}
 
-	const UPCGGraph* CurrentGraph = PCGEditor.Pin()->GetPCGGraph();
 	const FPCGStack& Stack = *Item->GetPCGStack();
-
+	const UPCGGraph* FocusedGraph = GetFocusedGraph();
 	bool bFoundCurrentGraph = false;
 	const UPCGNode* LastNode = nullptr;
 
@@ -1777,7 +1793,7 @@ void SPCGEditorGraphDebugObjectTree::FocusOnItem(const FPCGEditorGraphDebugObjec
 		const FPCGStackFrame& StackFrame = Stack.GetStackFrames()[i];
 		const UObject* Object = StackFrame.GetObject_GameThread();
 
-		if (Object == CurrentGraph)
+		if (Object == FocusedGraph)
 		{
 			bFoundCurrentGraph = true;
 			break;
diff --git a/Engine/Plugins/PCG/Source/PCGEditor/Private/Widgets/SPCGEditorGraphDebugObjectTree.h b/Engine/Plugins/PCG/Source/PCGEditor/Private/Widgets/SPCGEditorGraphDebugObjectTree.h
--- a/Engine/Plugins/PCG/Source/PCGEditor/Private/Widgets/SPCGEditorGraphDebugObjectTree.h
+++ b/Engine/Plugins/PCG/Source/PCGEditor/Private/Widgets/SPCGEditorGraphDebugObjectTree.h
@@ -305,7 +305,9 @@ private:
 
 	void OnEditorSelectionChanged(UObject* InObject);
 
-	UPCGGraph* GetPCGGraph() const;
+	UPCGGraph* GetFocusedGraph() const;
+	UPCGGraph* GetMainGraph() const;
+	bool IsMainOrEmbeddedGraph(const UPCGGraph* InGraph) const;
 
 	TSharedRef<ITableRow> MakeTreeRowWidget(FPCGEditorGraphDebugObjectItemPtr InItem, const TSharedRef<STableViewBase>& InOwnerTable);
 	void OnGetChildren(FPCGEditorGraphDebugObjectItemPtr InItem, TArray<FPCGEditorGraphDebugObjectItemPtr>& OutChildren) const;
@@ -320,7 +322,7 @@ private:
 
 	/** Expand the given row and select the next occurrence of the current graph. */
 	void ExpandAndSelectFirstLeafDebugObject(const FPCGEditorGraphDebugObjectItemPtr& InItem);
-
+	
 	TSharedRef<SWidget> OpenFilterMenu();
 	const FSlateBrush* GetFilterBadgeIcon() const;
 
diff --git a/Engine/Plugins/PCG/Source/PCGEditor/Private/Widgets/SPCGEditorGraphEmbeddedSubgraphsView.cpp b/Engine/Plugins/PCG/Source/PCGEditor/Private/Widgets/SPCGEditorGraphEmbeddedSubgraphsView.cpp
new file mode 100644
--- /dev/null
+++ b/Engine/Plugins/PCG/Source/PCGEditor/Private/Widgets/SPCGEditorGraphEmbeddedSubgraphsView.cpp
@@ -0,0 +1,491 @@
+// Copyright Epic Games, Inc. All Rights Reserved.
+
+#include "SPCGEditorGraphEmbeddedSubgraphsView.h"
+
+#include "PCGEditor.h"
+#include "PCGEditorGraph.h"
+#include "PCGEditorStyle.h"
+#include "PCGGraph.h"
+#include "Schema/PCGEditorGraphSchema.h"
+
+#include "GraphEditorDragDropAction.h"
+#include "ScopedTransaction.h"
+#include "Misc/MessageDialog.h"
+#include "Serialization/ArchiveReplaceObjectRef.h"
+#include "Serialization/FindObjectReferencers.h"
+#include "Widgets/SBoxPanel.h"
+#include "Widgets/Input/SSearchBox.h"
+#include "Widgets/Text/SInlineEditableTextBlock.h"
+#include "Widgets/Views/SListView.h"
+
+#define LOCTEXT_NAMESPACE "SPCGEditorGraphEmbeddedSubgraphsView"
+
+namespace PCGEditorGraphEmbeddedSubgraphsView
+{
+	DECLARE_DELEGATE(FOnRenameRequest);
+
+	struct FEmbeddedSubgraphItem
+	{
+		TWeakObjectPtr<UPCGGraph> WeakEmbeddedSubgraph;
+		TWeakPtr<FPCGEditor> WeakEditor;
+		void Delete();
+		void RenameRequest() { OnRenameRequest.ExecuteIfBound(); }
+
+		FOnRenameRequest OnRenameRequest;
+	};
+
+	class SEmbeddedSubgraphRow : public STableRow<TSharedPtr<FEmbeddedSubgraphItem>>
+	{
+	public:
+		SLATE_BEGIN_ARGS(SEmbeddedSubgraphRow) {}
+			SLATE_ARGUMENT(TSharedPtr<FEmbeddedSubgraphItem>, Item)
+		SLATE_END_ARGS()
+
+		void Construct(const FArguments&, TSharedRef<STableViewBase>);
+		FText GetItemLabel() const;
+		EVisibility GetItemVisibility() const;
+		void OnTextCommitted(const FText& InNewText, ETextCommit::Type InTextCommit);
+		bool OnVerifyTextChanged(const FText& InNewText, FText& OutErrorMessage);
+		virtual FReply OnDragDetected(const FGeometry&, const FPointerEvent&) override;
+		FReply OnDeleteClicked();
+		void OnRenameRequest();
+
+	private:
+		TSharedPtr<FEmbeddedSubgraphItem> Item;
+		TSharedPtr<SInlineEditableTextBlock> TextBlock;
+	};
+
+	class FEmbeddedSubgraphDragDropAction : public FGraphEditorDragDropAction
+	{
+	public:
+		DRAG_DROP_OPERATOR_TYPE(FEmbeddedSubgraphDragDropAction, FGraphEditorDragDropAction);
+
+		TSharedPtr<FEmbeddedSubgraphItem> Item;
+
+		static TSharedRef<FEmbeddedSubgraphDragDropAction> New(TSharedPtr<FEmbeddedSubgraphItem> InItem)
+		{
+			TSharedRef<FEmbeddedSubgraphDragDropAction> Operation = MakeShareable(new FEmbeddedSubgraphDragDropAction);
+			Operation->Item = InItem;
+			Operation->Construct();
+			return Operation;
+		}
+		
+		virtual FReply DroppedOnPanel(const TSharedRef<SWidget>&, const FVector2f&, const FVector2f&, UEdGraph&) override;
+		virtual void HoverTargetChanged() override;
+		FText GetHoverText() const;
+	};
+}
+
+void PCGEditorGraphEmbeddedSubgraphsView::FEmbeddedSubgraphItem::Delete()
+{
+	UPCGGraph* EmbeddedSubgraph = WeakEmbeddedSubgraph.Get();
+	if (!EmbeddedSubgraph)
+	{
+		return;
+	}
+
+	TArray<UObject*> Referencers;
+	// Find all refs
+	for (auto Referencer : TFindObjectReferencers<UPCGGraph>({ EmbeddedSubgraph }, /*PackageToCheck =*/EmbeddedSubgraph->GetPackage(), /*bIgnoreTemplates =*/false))
+	{
+		// Skip outer references
+		if (Referencer.Key->IsIn(Referencer.Value) || Referencer.Value->IsIn(Referencer.Key))
+		{
+			continue;
+		}
+		Referencers.Add(Referencer.Value);
+	}
+
+	TMap<UPCGGraph*, UPCGGraph*> Redirects;
+	Redirects.Add(EmbeddedSubgraph, nullptr);
+
+	FScopedTransaction Transaction(LOCTEXT("DeleteEmbeddedSubgraph", "Delete Embedded Subgraph"));
+	if (!Referencers.IsEmpty())
+	{
+		EAppReturnType::Type Ret = FMessageDialog::Open(
+			EAppMsgType::YesNoCancel, LOCTEXT("DeleteWithReferencesMsg", "Embedded subgraph is referenced. Deleting the subgraph will invalidate those subgraph nodes. Are you sure?"),
+			LOCTEXT("DeleteWithReferencesTitle", "Delete embedded subgraph with references?"));
+		if (Ret != EAppReturnType::Yes)
+		{
+			Transaction.Cancel();
+			return;
+		}
+
+		for (UObject* Referencer : Referencers)
+		{
+			Referencer->Modify();
+			FArchiveReplaceObjectRef<UPCGGraph>(Referencer, Redirects);
+		}
+	}
+
+	if (TSharedPtr<FPCGEditor> PCGEditor = WeakEditor.Pin())
+	{
+		PCGEditor->CloseDocument(EmbeddedSubgraph);
+	}
+
+	EmbeddedSubgraph->GetTypedOuter<UPCGGraph>()->DeleteEmbeddedSubgraph(EmbeddedSubgraph);
+}
+
+FReply PCGEditorGraphEmbeddedSubgraphsView::FEmbeddedSubgraphDragDropAction::DroppedOnPanel(const TSharedRef<SWidget>& Panel, const FVector2f& ScreenPosition, const FVector2f& GraphPosition, UEdGraph& Graph)
+{
+	if (!Graph.GetSchema()->IsA<UPCGEditorGraphSchema>())
+	{
+		return FReply::Unhandled();
+	}
+
+	UPCGEditorGraph* EditorGraph = Cast<UPCGEditorGraph>(&Graph);
+	if (!ensure(EditorGraph))
+	{
+		return FReply::Unhandled();
+	}
+
+	FAssetData Asset(Item->WeakEmbeddedSubgraph.Get());
+	Graph.GetSchema()->DroppedAssetsOnGraph({ Asset }, GraphPosition, &Graph);
+	
+	return FReply::Handled();
+}
+
+FText PCGEditorGraphEmbeddedSubgraphsView::FEmbeddedSubgraphDragDropAction::GetHoverText() const
+{
+	return Item.IsValid() && Item->WeakEmbeddedSubgraph.IsValid() ? Item->WeakEmbeddedSubgraph->Title : LOCTEXT("InvalidHoverText", "Invalid");
+}
+
+void PCGEditorGraphEmbeddedSubgraphsView::FEmbeddedSubgraphDragDropAction::HoverTargetChanged()
+{
+	SetSimpleFeedbackMessage(FPCGEditorStyle::Get().GetBrush("ClassIcon.PCGGraph"), FLinearColor::White, GetHoverText());
+}
+
+void PCGEditorGraphEmbeddedSubgraphsView::SEmbeddedSubgraphRow::Construct(const FArguments& InArgs, TSharedRef<STableViewBase> InOwner)
+{
+	Item = InArgs._Item;
+
+	TSharedRef<SHorizontalBox> HorizontalBox = SNew(SHorizontalBox);
+
+	HorizontalBox->AddSlot()
+	.AutoWidth()
+	.HAlign(HAlign_Center)
+	.VAlign(VAlign_Center)
+	.Padding(4.f)
+	[
+		SNew(SImage)
+		.ColorAndOpacity(FSlateColor::UseForeground())
+		.Image(FPCGEditorStyle::Get().GetBrush("ClassIcon.PCGGraph"))
+	];
+
+	HorizontalBox->AddSlot()
+	.HAlign(HAlign_Left)
+	.VAlign(VAlign_Center)
+	.Padding(4.f)
+	[
+		SAssignNew(TextBlock, SInlineEditableTextBlock)
+		.Text(this, &PCGEditorGraphEmbeddedSubgraphsView::SEmbeddedSubgraphRow::GetItemLabel)
+		.HintText(LOCTEXT("EmbeddedSubgraphLabelHint", "Label"))
+		.OnTextCommitted(this, &PCGEditorGraphEmbeddedSubgraphsView::SEmbeddedSubgraphRow::OnTextCommitted)
+		.OnVerifyTextChanged(this, &PCGEditorGraphEmbeddedSubgraphsView::SEmbeddedSubgraphRow::OnVerifyTextChanged)
+		.IsSelected(this, &PCGEditorGraphEmbeddedSubgraphsView::SEmbeddedSubgraphRow::IsSelected)
+		.MultiLine(false)
+	];
+
+	HorizontalBox->AddSlot()
+	.HAlign(HAlign_Right)
+	.VAlign(VAlign_Center)
+	.Padding(4.f)
+	[
+		SNew(SButton)
+		.ButtonStyle(FAppStyle::Get(), "SimpleButton")
+		.Visibility(this, &PCGEditorGraphEmbeddedSubgraphsView::SEmbeddedSubgraphRow::GetItemVisibility)
+		.OnClicked(this, &PCGEditorGraphEmbeddedSubgraphsView::SEmbeddedSubgraphRow::OnDeleteClicked)
+		[
+			SNew(SImage)
+			.ColorAndOpacity(FSlateColor::UseForeground())
+			.Image(FAppStyle::GetBrush("Icons.Delete"))
+		]
+	];
+
+	STableRow<TSharedPtr<FEmbeddedSubgraphItem>>::Construct(
+		STableRow::FArguments()
+		[
+			HorizontalBox
+		],
+		InOwner
+	);
+
+	Item->OnRenameRequest = FOnRenameRequest::CreateSP(this, &PCGEditorGraphEmbeddedSubgraphsView::SEmbeddedSubgraphRow::OnRenameRequest);
+}
+
+FText PCGEditorGraphEmbeddedSubgraphsView::SEmbeddedSubgraphRow::GetItemLabel() const
+{
+	if (UPCGGraph* Graph = Item->WeakEmbeddedSubgraph.Get())
+	{
+		return Graph->Title;
+	}
+	else
+	{
+		return LOCTEXT("UnknownEmbeddedSubgraphLabel", "Unknown Graph");
+	}
+}
+
+EVisibility PCGEditorGraphEmbeddedSubgraphsView::SEmbeddedSubgraphRow::GetItemVisibility() const
+{
+	return IsSelected() ? EVisibility::Visible : EVisibility::Hidden;
+}
+
+FReply PCGEditorGraphEmbeddedSubgraphsView::SEmbeddedSubgraphRow::OnDeleteClicked()
+{
+	if (Item.IsValid())
+	{
+		Item->Delete();
+		return FReply::Handled();
+	}
+	else
+	{
+		return FReply::Unhandled();
+	}
+}
+
+void PCGEditorGraphEmbeddedSubgraphsView::SEmbeddedSubgraphRow::OnRenameRequest()
+{
+	if (TextBlock.IsValid())
+	{
+		TextBlock->EnterEditingMode();
+	}
+}
+
+void PCGEditorGraphEmbeddedSubgraphsView::SEmbeddedSubgraphRow::OnTextCommitted(const FText& InNewText, ETextCommit::Type InTextCommit)
+{
+	if (UPCGGraph* EmbeddedSubgraph = Item->WeakEmbeddedSubgraph.Get())
+	{
+		FScopedTransaction Transaction(LOCTEXT("RenameEmbeddedSubgraph", "Rename Embeded subgraph"));
+		EmbeddedSubgraph->Modify();
+		EmbeddedSubgraph->Title = InNewText;
+	}
+}
+
+bool PCGEditorGraphEmbeddedSubgraphsView::SEmbeddedSubgraphRow::OnVerifyTextChanged(const FText& InNewText, FText& OutErrorMessage)
+{
+	const UPCGGraph* EmbeddedSubgraph = Item->WeakEmbeddedSubgraph.Get();
+	if (!EmbeddedSubgraph)
+	{
+		return false;
+	}
+
+	if (InNewText.IsEmpty())
+	{
+		OutErrorMessage = LOCTEXT("EmbeddedSubgraphLabelInvalidError", "Invalid label");
+		return false;
+	}
+
+	if (UPCGGraph* OuterGraph = EmbeddedSubgraph->GetTypedOuter<UPCGGraph>())
+	{
+		for (UPCGGraph* CurrentEmbeddedSubgraph : OuterGraph->GetEmbeddedSubgraphs())
+		{
+			if (CurrentEmbeddedSubgraph != EmbeddedSubgraph && CurrentEmbeddedSubgraph->Title.EqualToCaseIgnored(InNewText))
+			{
+				OutErrorMessage = LOCTEXT("EmbeddedSubgraphLabelAlreadyExistsError", "Local graph with same label already exists");
+				return false;
+			}
+		}
+
+		return true;
+	}
+	else
+	{
+		return false;
+	}
+}
+
+FReply PCGEditorGraphEmbeddedSubgraphsView::SEmbeddedSubgraphRow::OnDragDetected(const FGeometry&, const FPointerEvent&)
+{
+	if (Item.IsValid())
+	{
+		if (TSharedPtr<PCGEditorGraphEmbeddedSubgraphsView::FEmbeddedSubgraphDragDropAction> DragDropOp = PCGEditorGraphEmbeddedSubgraphsView::FEmbeddedSubgraphDragDropAction::New(Item))
+		{
+			return FReply::Handled().BeginDragDrop(DragDropOp.ToSharedRef());
+		}
+	}
+	
+	return FReply::Unhandled();
+}
+
+UPCGGraph* SPCGEditorGraphEmbeddedSubgraphsView::GetEditorGraph() const
+{
+	if (TSharedPtr<FPCGEditor> PinnedEditor = PCGEditorPtr.Pin())
+	{
+		return PinnedEditor.IsValid() ? PinnedEditor->GetMainGraph() : nullptr;
+	}
+
+	return nullptr;
+}
+
+void SPCGEditorGraphEmbeddedSubgraphsView::Construct(const FArguments& InArgs, const TSharedPtr<FPCGEditor>& InPCGEditor)
+{
+	PCGEditorPtr = InPCGEditor;
+	
+	UPCGGraph* EditorGraph = GetEditorGraph();
+	if (!EditorGraph)
+	{
+		return;
+	}
+
+	TSharedRef<SVerticalBox> VerticalBox = SNew(SVerticalBox);
+
+	ChildSlot
+	[
+		VerticalBox
+	];
+
+	// Setup toolbar
+	{
+		// Setup toolbar
+		TSharedRef<SHorizontalBox> ToolbarBox = SNew(SHorizontalBox);
+
+		ToolbarBox->AddSlot()
+		.HAlign(HAlign_Fill)
+		.FillWidth(1.f)
+		[
+			SNew(SSearchBox)
+			.OnTextChanged(this, &SPCGEditorGraphEmbeddedSubgraphsView::OnFilterTextChanged)
+		];
+
+		ToolbarBox->AddSlot()
+		.VAlign(VAlign_Center)
+		.HAlign(HAlign_Right)
+		.Padding(4.f, 0.f, 0.f, 0.f)
+		.AutoWidth()
+		[
+			SNew(SButton)
+			.ButtonStyle(FAppStyle::Get(), "SimpleButton")
+			.ToolTipText(LOCTEXT("NewEmbeddedSubgraphToolTip", "New Embedded subgraph"))
+			.OnClicked(this, &SPCGEditorGraphEmbeddedSubgraphsView::NewEmbeddedSubgraphClicked)
+			[
+				SNew(SImage)
+				.ColorAndOpacity(FSlateColor::UseForeground())
+				.Image(FAppStyle::GetBrush("Icons.PlusCircle"))
+			]
+		];
+
+		VerticalBox->AddSlot()
+		.AutoHeight()
+		.Padding(8.f, 8.f, 8.f, 4.f)
+		[
+			ToolbarBox
+		];
+	}
+
+	// Setup list view
+	{
+		RefreshEmbeddedSubgraphsList();
+
+		VerticalBox->AddSlot()
+		[
+			SAssignNew(ListView, SListView<TSharedPtr<PCGEditorGraphEmbeddedSubgraphsView::FEmbeddedSubgraphItem>>)
+			.OnGenerateRow(this, &SPCGEditorGraphEmbeddedSubgraphsView::GenerateEmbeddedSubgraphRow)
+			.SelectionMode(ESelectionMode::Single)
+			.ListItemsSource(&Items)
+			.OnMouseButtonDoubleClick(this, &SPCGEditorGraphEmbeddedSubgraphsView::OnEmbeddedSubgraphItemDoubleClick)
+			.OnKeyDownHandler(this, &SPCGEditorGraphEmbeddedSubgraphsView::OnEmbeddedSubgraphItemKeyDownHandler)
+		];
+	}
+}
+
+void SPCGEditorGraphEmbeddedSubgraphsView::Tick(const FGeometry& AllottedGeometry, const double InCurrentTime, const float InDeltaTime)
+{
+	SCompoundWidget::Tick(AllottedGeometry, InCurrentTime, InDeltaTime);
+
+	if (bRefreshRequested)
+	{
+		RefreshEmbeddedSubgraphsList();
+		bRefreshRequested = false;
+	}
+}
+
+FReply SPCGEditorGraphEmbeddedSubgraphsView::NewEmbeddedSubgraphClicked()
+{
+	if (UPCGGraph* EditorGraph = GetEditorGraph())
+	{
+		{
+			FScopedTransaction Transaction(LOCTEXT("AddEmbeddedSubgraph", "New Embedded subgraph"));
+			if (UPCGGraph* NewGraph = EditorGraph->AddNewEmbeddedSubgraph())
+			{
+				if (TSharedPtr<FPCGEditor> PinnedEditor = PCGEditorPtr.Pin())
+				{
+					PinnedEditor->OpenDocument(NewGraph, FDocumentTracker::EOpenDocumentCause::OpenNewDocument);
+				}
+			}
+		}
+		RefreshEmbeddedSubgraphsList();
+	}
+
+	return FReply::Handled();
+}
+
+void SPCGEditorGraphEmbeddedSubgraphsView::OnFilterTextChanged(const FText& InText)
+{
+	FilterText = InText;
+	RequestRefresh();
+}
+
+void SPCGEditorGraphEmbeddedSubgraphsView::RefreshEmbeddedSubgraphsList()
+{
+	Items.Empty();
+
+	if (UPCGGraph* EditorGraph = GetEditorGraph())
+	{
+		FString FilterTextStr = FilterText.ToString();
+		for (const auto& EmbeddedSubgraph : EditorGraph->GetEmbeddedSubgraphs())
+		{
+			if (FilterText.IsEmpty() || EmbeddedSubgraph->Title.ToString().Contains(FilterTextStr))
+			{
+				Items.Add(MakeShared<PCGEditorGraphEmbeddedSubgraphsView::FEmbeddedSubgraphItem>(EmbeddedSubgraph, PCGEditorPtr));
+			}
+		}
+		
+	}
+	
+	if (ListView.IsValid())
+	{
+		ListView->RebuildList();
+	}
+}
+
+TSharedRef<ITableRow> SPCGEditorGraphEmbeddedSubgraphsView::GenerateEmbeddedSubgraphRow(TSharedPtr<PCGEditorGraphEmbeddedSubgraphsView::FEmbeddedSubgraphItem> InItem, const TSharedRef<STableViewBase>& InOwner)
+{
+	return SNew(PCGEditorGraphEmbeddedSubgraphsView::SEmbeddedSubgraphRow, InOwner).Item(InItem);
+}
+
+void SPCGEditorGraphEmbeddedSubgraphsView::OnEmbeddedSubgraphItemDoubleClick(TSharedPtr<PCGEditorGraphEmbeddedSubgraphsView::FEmbeddedSubgraphItem> InItem)
+{
+	if (TSharedPtr<FPCGEditor> PinnedEditor = PCGEditorPtr.Pin())
+	{
+		UPCGGraph* EmbeddedSubgraph = InItem->WeakEmbeddedSubgraph.Get();
+		
+		PinnedEditor->OpenDocument(EmbeddedSubgraph, FDocumentTracker::EOpenDocumentCause::OpenNewDocument);
+	}
+}
+
+FReply SPCGEditorGraphEmbeddedSubgraphsView::OnEmbeddedSubgraphItemKeyDownHandler(const FGeometry&, const FKeyEvent& InKeyEvent) const
+{
+	if (ListView->GetNumItemsSelected() != 1)
+	{
+		return FReply::Unhandled();
+	}
+
+	if (InKeyEvent.GetKey() == EKeys::F2)
+	{
+		TSharedPtr<PCGEditorGraphEmbeddedSubgraphsView::FEmbeddedSubgraphItem> Item = ListView->GetSelectedItems()[0];
+		Item->RenameRequest();
+		return FReply::Handled();
+	}
+
+	if (InKeyEvent.GetKey() == EKeys::Delete || InKeyEvent.GetKey() == EKeys::BackSpace)
+	{
+		TSharedPtr<PCGEditorGraphEmbeddedSubgraphsView::FEmbeddedSubgraphItem> Item = ListView->GetSelectedItems()[0];
+		Item->Delete();
+		return FReply::Handled();
+	}
+
+	return FReply::Unhandled();
+}
+
+#undef LOCTEXT_NAMESPACE
+
diff --git a/Engine/Plugins/PCG/Source/PCGEditor/Private/Widgets/SPCGEditorGraphEmbeddedSubgraphsView.h b/Engine/Plugins/PCG/Source/PCGEditor/Private/Widgets/SPCGEditorGraphEmbeddedSubgraphsView.h
new file mode 100644
--- /dev/null
+++ b/Engine/Plugins/PCG/Source/PCGEditor/Private/Widgets/SPCGEditorGraphEmbeddedSubgraphsView.h
@@ -0,0 +1,48 @@
+// Copyright Epic Games, Inc. All Rights Reserved.
+
+#pragma once
+
+#include "Templates/SharedPointer.h"
+
+#include "Widgets/SCompoundWidget.h"
+
+struct FGeometry;
+struct FKeyEvent;
+class FPCGEditor;
+class ITableRow;
+class STableViewBase;
+class UPCGGraph;
+
+template <typename ItemType> class SListView;
+
+namespace PCGEditorGraphEmbeddedSubgraphsView
+{
+	struct FEmbeddedSubgraphItem;
+}
+
+/** Displays a PCGGraph's Local graphs. */
+class SPCGEditorGraphEmbeddedSubgraphsView final : public SCompoundWidget
+{
+public:
+	SLATE_BEGIN_ARGS(SPCGEditorGraphEmbeddedSubgraphsView) {}
+	SLATE_END_ARGS()
+
+	void Construct(const FArguments& InArgs, const TSharedPtr<FPCGEditor>& InPCGEditor);
+	virtual void Tick(const FGeometry& AllottedGeometry, const double InCurrentTime, const float InDeltaTime) override;
+	void RequestRefresh() { bRefreshRequested = true; }
+
+private:
+	void OnFilterTextChanged(const FText& InText);
+	FReply NewEmbeddedSubgraphClicked();
+	TSharedRef<ITableRow> GenerateEmbeddedSubgraphRow(TSharedPtr<PCGEditorGraphEmbeddedSubgraphsView::FEmbeddedSubgraphItem> InItem, const TSharedRef<STableViewBase>& InOwner);
+	void OnEmbeddedSubgraphItemDoubleClick(TSharedPtr<PCGEditorGraphEmbeddedSubgraphsView::FEmbeddedSubgraphItem> InItem);
+	FReply OnEmbeddedSubgraphItemKeyDownHandler(const FGeometry&, const FKeyEvent& InKeyEvent) const;
+	UPCGGraph* GetEditorGraph() const;
+	void RefreshEmbeddedSubgraphsList();
+
+	TArray<TSharedPtr<PCGEditorGraphEmbeddedSubgraphsView::FEmbeddedSubgraphItem>> Items;
+	TSharedPtr<SListView<TSharedPtr<PCGEditorGraphEmbeddedSubgraphsView::FEmbeddedSubgraphItem>>> ListView;
+	TWeakPtr<FPCGEditor> PCGEditorPtr;
+	bool bRefreshRequested = false;
+	FText FilterText;
+};
diff --git a/Engine/Plugins/PCG/Source/PCGEditor/Private/Widgets/SPCGEditorGraphFind.cpp b/Engine/Plugins/PCG/Source/PCGEditor/Private/Widgets/SPCGEditorGraphFind.cpp
--- a/Engine/Plugins/PCG/Source/PCGEditor/Private/Widgets/SPCGEditorGraphFind.cpp
+++ b/Engine/Plugins/PCG/Source/PCGEditor/Private/Widgets/SPCGEditorGraphFind.cpp
@@ -72,7 +72,7 @@ FReply FPCGEditorGraphFindResult::OnClick(TWeakPtr<FPCGEditor> InPCGEditorPtr)
 	}
 
 	UEdGraphPin* ResolvedPin = Pin.Get();
-	if (ResolvedPin && ResolvedPin->GetOwningNode()->GetGraph() == Editor->GetPCGEditorGraph())
+	if (ResolvedPin && ResolvedPin->GetOwningNode()->GetGraph() == Editor->GetFocusedEditorGraph())
 	{
 		Editor->JumpToNode(ResolvedPin->GetOwningNode());
 	}
@@ -88,10 +88,14 @@ FReply FPCGEditorGraphFindResult::OnDoubleClick(TWeakPtr<FPCGEditor> InPCGEditor
 {
 	TSharedPtr<FPCGEditor> CurrentEditor = InPCGEditorPtr.Pin();
 	TSharedPtr<FPCGEditor> Editor = CurrentEditor;
+	if (!Editor)
+	{
+		return FReply::Unhandled();
+	}
 
 	if (ParentGraph)
 	{
-		GEditor->GetEditorSubsystem<UAssetEditorSubsystem>()->OpenEditorForAsset(ParentGraph->GetPCGGraph());
+		Editor->OpenEditorForGraph(ParentGraph->GetPCGGraph());
 		Editor = ParentGraph->GetEditor().Pin();
 	}
 
@@ -635,7 +639,7 @@ namespace PCGEditorGraphFind
 void SPCGEditorGraphFind::MatchTokens(const TArray<FString>& InTokens)
 {
 	TSharedPtr<FPCGEditor> PCGEditor = PCGEditorPtr.Pin();
-	UPCGEditorGraph* PCGEditorGraph = PCGEditor ? PCGEditor->GetPCGEditorGraph() : nullptr;
+	UPCGEditorGraph* PCGEditorGraph = PCGEditor ? PCGEditor->GetMainEditorGraph() : nullptr;
 	
 	if (!PCGEditor || !PCGEditorGraph)
 	{
@@ -651,7 +655,7 @@ void SPCGEditorGraphFind::MatchTokens(const TArray<FString>& InTokens)
 	GraphRoot.Graph = ThisGraph;
 	GraphRoot.Node = nullptr;
 
-	auto VisitAllNodes = [&AllGraphs](UPCGEditorGraph* EditorGraph, PCGEditorGraphFind::GraphTree& Parent, auto&& RecursiveCall)
+	auto VisitAllNodes = [&AllGraphs, PCGEditor](UPCGEditorGraph* EditorGraph, PCGEditorGraphFind::GraphTree& Parent, auto&& RecursiveCall)
 	{
 		if (!EditorGraph)
 		{
@@ -672,7 +676,7 @@ void SPCGEditorGraphFind::MatchTokens(const TArray<FString>& InTokens)
 							AllGraphs.Add(Subgraph);
 							if (PCGEditorGraphFind::GraphTree* Child = Parent.FindOrCreateChild(Node, Subgraph))
 							{
-								RecursiveCall(FPCGEditor::GetPCGEditorGraph(Subgraph), *Child, RecursiveCall);
+								RecursiveCall(PCGEditor->GetOrCreatePCGEditorGraph(Subgraph), *Child, RecursiveCall);
 							}
 						}
 					}
@@ -753,7 +757,7 @@ void SPCGEditorGraphFind::MatchTokens(const TArray<FString>& InTokens)
 					AllGraphs.Add(Subgraph);
 					// Find/Create child on current
 					Current = Current->FindOrCreateChild(SubgraphEditorNode, Subgraph);
-					CurrentEditorGraph = Current && Subgraph ? FPCGEditor::GetPCGEditorGraph(Subgraph) : nullptr;
+					CurrentEditorGraph = Current && Subgraph ? PCGEditor->GetOrCreatePCGEditorGraph(Subgraph) : nullptr;
 
 					if (Current && CurrentEditorGraph)
 					{
@@ -780,7 +784,7 @@ void SPCGEditorGraphFind::MatchTokens(const TArray<FString>& InTokens)
 
 	if(FindMode == EPCGGraphFindMode::ShowMinimumTree || FindMode == EPCGGraphFindMode::ShowFullTree)
 	{
-		auto RecurseInGraphTree = [this, &InTokens](const PCGEditorGraphFind::GraphTree& Node, TFunctionRef<FPCGEditorGraphFindResultPtr()> GetParentFunc, auto&& RecurseCall) -> void
+		auto RecurseInGraphTree = [this, &InTokens, PCGEditor](const PCGEditorGraphFind::GraphTree& Node, TFunctionRef<FPCGEditorGraphFindResultPtr()> GetParentFunc, auto&& RecurseCall) -> void
 		{
 			check(Node.Graph);
 
@@ -792,7 +796,7 @@ void SPCGEditorGraphFind::MatchTokens(const TArray<FString>& InTokens)
 				NodeString = FString::Printf(TEXT("%s (%s)"), *FName::NameToDisplayString(Node.Graph->GetName(), /*bIsBool=*/false), *Node.Node->GetNodeTitle(ENodeTitleType::ListView).ToString());
 			}
 
-			auto GetOrCreateNodeResult = [this, &Node, &NodeString, &ThisResult, &GetParentFunc]()
+			auto GetOrCreateNodeResult = [this, &Node, &NodeString, &ThisResult, &GetParentFunc, PCGEditor]()
 			{
 				FPCGEditorGraphFindResultPtr Parent = GetParentFunc();
 				check(Parent);
@@ -800,7 +804,7 @@ void SPCGEditorGraphFind::MatchTokens(const TArray<FString>& InTokens)
 				if(!ThisResult.IsValid() && Node.Node)
 				{
 					ThisResult = MakeShared<FPCGEditorGraphFindResult>(NodeString, Parent, Node.Node);
-					ThisResult->ParentGraph = FPCGEditor::GetPCGEditorGraph(Node.Graph);
+					ThisResult->ParentGraph = PCGEditor->GetOrCreatePCGEditorGraph(Node.Graph);
 					Parent->Children.Add(ThisResult);
 				}
 
@@ -814,7 +818,7 @@ void SPCGEditorGraphFind::MatchTokens(const TArray<FString>& InTokens)
 			}
 
 			// Create local elements
-			MatchTokensInternal(InTokens, FPCGEditor::GetPCGEditorGraph(Node.Graph), GetOrCreateNodeResult/*, MatchedTokens*/);
+			MatchTokensInternal(InTokens, PCGEditor->GetOrCreatePCGEditorGraph(Node.Graph), GetOrCreateNodeResult/*, MatchedTokens*/);
 
 			// Continue through the tree
 			for (const PCGEditorGraphFind::GraphTree& Child : Node.GetChildren())
@@ -837,7 +841,7 @@ void SPCGEditorGraphFind::MatchTokens(const TArray<FString>& InTokens)
 			FPCGEditorGraphFindResultPtr GraphNodePtr;
 			const FString GraphString = FName::NameToDisplayString(Graph->GetName(), /*bIsBool=*/false);
 
-			auto CreateGraphNode = [Graph, &GraphString, &GraphNodePtr, &GraphRoot]()
+			auto CreateGraphNode = [Graph, &GraphString, &GraphNodePtr, &GraphRoot, PCGEditor]()
 			{
 				if (!GraphNodePtr.IsValid())
 				{
@@ -860,7 +864,7 @@ void SPCGEditorGraphFind::MatchTokens(const TArray<FString>& InTokens)
 					}
 
 					GraphNodePtr = MakeShared<FPCGEditorGraphFindResult>(GraphString, FPCGEditorGraphFindResultPtr(), EditorNode);
-					GraphNodePtr->ParentGraph = FPCGEditor::GetPCGEditorGraph(Graph);
+					GraphNodePtr->ParentGraph = PCGEditor->GetOrCreatePCGEditorGraph(Graph);
 					GraphNodePtr->RootGraphNode = RootEditorNode;
 				}
 
@@ -873,7 +877,7 @@ void SPCGEditorGraphFind::MatchTokens(const TArray<FString>& InTokens)
 				CreateGraphNode()->bIsMatch = true;
 			}
 
-			MatchTokensInternal(InTokens, FPCGEditor::GetPCGEditorGraph(Graph), CreateGraphNode/*, Results*/);
+			MatchTokensInternal(InTokens,  PCGEditor->GetOrCreatePCGEditorGraph(Graph), CreateGraphNode/*, Results*/);
 
 			if (GraphNodePtr)
 			{
diff --git a/Engine/Plugins/PCG/Source/PCGEditor/Private/Widgets/SPCGEditorGraphLogView.cpp b/Engine/Plugins/PCG/Source/PCGEditor/Private/Widgets/SPCGEditorGraphLogView.cpp
--- a/Engine/Plugins/PCG/Source/PCGEditor/Private/Widgets/SPCGEditorGraphLogView.cpp
+++ b/Engine/Plugins/PCG/Source/PCGEditor/Private/Widgets/SPCGEditorGraphLogView.cpp
@@ -118,7 +118,7 @@ void SPCGEditorGraphLogView::Construct(const FArguments& InArgs, TSharedPtr<FPCG
 	const TSharedPtr<FPCGEditor> PCGEditor = PCGEditorPtr.Pin();
 	if (PCGEditor)
 	{
-		PCGEditorGraph = PCGEditor->GetPCGEditorGraph();
+		PCGEditorGraph = PCGEditor->GetMainEditorGraph();
 		PCGEditor->GetInspectionDataManager().OnInspectedStackChangedDelegate.AddSP(this, &SPCGEditorGraphLogView::OnDebugStackChanged);
 	}
 
diff --git a/Engine/Plugins/PCG/Source/PCGEditor/Private/Widgets/SPCGEditorGraphNodePalette.cpp b/Engine/Plugins/PCG/Source/PCGEditor/Private/Widgets/SPCGEditorGraphNodePalette.cpp
--- a/Engine/Plugins/PCG/Source/PCGEditor/Private/Widgets/SPCGEditorGraphNodePalette.cpp
+++ b/Engine/Plugins/PCG/Source/PCGEditor/Private/Widgets/SPCGEditorGraphNodePalette.cpp
@@ -132,7 +132,7 @@ void SPCGEditorGraphNodePalette::CollectAllActions(FGraphActionListBuilderBase&
 	FGraphActionMenuBuilder ActionMenuBuilder;
 	FPCGEditor* Editor = PCGEditor.Pin().Get();
 	check(Editor);
-	FPCGActionsFilter ActionFilter(Editor->GetPCGEditorGraph(), ElementType & PCGSchema->GetElementTypeFiltering(), PCGSchema->GetGraphEditorCustomization(Editor->GetPCGEditorGraph()));
+	FPCGActionsFilter ActionFilter(Editor->GetMainEditorGraph(), ElementType & PCGSchema->GetElementTypeFiltering(), PCGSchema->GetGraphEditorCustomization(Editor->GetMainEditorGraph()));
 	
 	PCGSchema->GetPaletteActions(ActionMenuBuilder, ActionFilter);
 	OutAllActions.Append(ActionMenuBuilder);
diff --git a/Engine/Plugins/PCG/Source/PCGEditor/Private/Widgets/SPCGEditorGraphParamsView.cpp b/Engine/Plugins/PCG/Source/PCGEditor/Private/Widgets/SPCGEditorGraphParamsView.cpp
--- a/Engine/Plugins/PCG/Source/PCGEditor/Private/Widgets/SPCGEditorGraphParamsView.cpp
+++ b/Engine/Plugins/PCG/Source/PCGEditor/Private/Widgets/SPCGEditorGraphParamsView.cpp
@@ -18,7 +18,7 @@ void SPCGEditorGraphUserParametersView::Construct(const FArguments& InArgs, cons
 		return;
 	}
 
-	UPCGGraph* PCGGraph = InPCGEditor.IsValid() && InPCGEditor->GetPCGEditorGraph() ? InPCGEditor->GetPCGEditorGraph()->GetPCGGraph() : nullptr;
+	UPCGGraph* PCGGraph = InPCGEditor.IsValid() ? InPCGEditor->GetFocusedGraph() : nullptr;
 
 	if (!PCGGraph)
 	{
@@ -82,7 +82,7 @@ void SPCGEditorGraphUserParametersView::Construct(const FArguments& InArgs, cons
 	}
 
 	// Note: Single Property View Header notes that it doesn't work with arrays or structs, so using filters for now.
-	const TSharedRef<IDetailsView> DetailsView = PropertyEditorModule.CreateDetailView(DetailsViewArgs);
+	DetailsView = PropertyEditorModule.CreateDetailView(DetailsViewArgs);
 
 	// Important: since this is using the graph object and that has details customization, we want to ensure we use the default layout instead.
 	// Otherwise we'll get buttons and more.
@@ -96,7 +96,15 @@ void SPCGEditorGraphUserParametersView::Construct(const FArguments& InArgs, cons
 		+ SVerticalBox::Slot()
 		.FillHeight(1.f)
 		[
-			DetailsView
+			DetailsView.ToSharedRef()
 		]
 	];
 }
+
+void SPCGEditorGraphUserParametersView::SetGraph(UPCGGraph* InGraph)
+{
+	if (DetailsView)
+	{
+		DetailsView->SetObject(InGraph);
+	}
+}
diff --git a/Engine/Plugins/PCG/Source/PCGEditor/Private/Widgets/SPCGEditorGraphParamsView.h b/Engine/Plugins/PCG/Source/PCGEditor/Private/Widgets/SPCGEditorGraphParamsView.h
--- a/Engine/Plugins/PCG/Source/PCGEditor/Private/Widgets/SPCGEditorGraphParamsView.h
+++ b/Engine/Plugins/PCG/Source/PCGEditor/Private/Widgets/SPCGEditorGraphParamsView.h
@@ -8,6 +8,8 @@
 
 struct FPropertyAndParent;
 class FPCGEditor;
+class IDetailsView;
+class UPCGGraph;
 
 /** Displays a PCGGraph's User Parameters details view. */
 class SPCGEditorGraphUserParametersView final : public SCompoundWidget
@@ -17,4 +19,7 @@ public:
 	SLATE_END_ARGS()
 
 	void Construct(const FArguments& InArgs, const TSharedPtr<FPCGEditor>& InPCGEditor);
+	void SetGraph(UPCGGraph* InGraph);
+private:
+	TSharedPtr<IDetailsView> DetailsView;
 };
\ No newline at end of file
diff --git a/Engine/Plugins/PCG/Source/PCGEditor/Private/Widgets/SPCGEditorGraphProfilingView.cpp b/Engine/Plugins/PCG/Source/PCGEditor/Private/Widgets/SPCGEditorGraphProfilingView.cpp
--- a/Engine/Plugins/PCG/Source/PCGEditor/Private/Widgets/SPCGEditorGraphProfilingView.cpp
+++ b/Engine/Plugins/PCG/Source/PCGEditor/Private/Widgets/SPCGEditorGraphProfilingView.cpp
@@ -308,7 +308,7 @@ void SPCGEditorGraphProfilingView::Construct(const FArguments& InArgs, TSharedPt
 	const TSharedPtr<FPCGEditor> PCGEditor = PCGEditorPtr.Pin();
 	if (PCGEditor)
 	{
-		PCGEditorGraph = PCGEditor->GetPCGEditorGraph();
+		PCGEditorGraph = PCGEditor->GetMainEditorGraph();
 		PCGEditor->GetInspectionDataManager().OnInspectedStackChangedDelegate.AddSP(this, &SPCGEditorGraphProfilingView::OnDebugStackChanged);
 	}
 
diff --git a/Engine/Plugins/PCG/Source/PCGEditor/Private/Widgets/SPCGEditorGraphTitleBar.cpp b/Engine/Plugins/PCG/Source/PCGEditor/Private/Widgets/SPCGEditorGraphTitleBar.cpp
new file mode 100644
--- /dev/null
+++ b/Engine/Plugins/PCG/Source/PCGEditor/Private/Widgets/SPCGEditorGraphTitleBar.cpp
@@ -0,0 +1,137 @@
+// Copyright Epic Games, Inc. All Rights Reserved.
+
+#include "SPCGEditorGraphTitleBar.h"
+
+#include "PCGEditor.h"
+#include "PCGEditorGraph.h"
+
+#include "Widgets/Layout/SSeparator.h"
+
+#define LOCTEXT_NAMESPACE "SPCGGraphEditorTitleBar"
+
+FText SPCGGraphEditorTitleBar::GetTitleForOneCrumb(TWeakObjectPtr<const UPCGEditorGraph> InWeakEditorGraph)
+{
+	if (const UPCGEditorGraph* EditorGraph = InWeakEditorGraph.Get())
+	{
+		return EditorGraph->GetDisplayName();
+	}
+
+	return LOCTEXT("PCGGraphEditorTitleBarInvalidCrumb", "Invalid");
+}
+
+void SPCGGraphEditorTitleBar::Construct(const FArguments& InArgs)
+{
+	EditorGraph = InArgs._Graph;
+	check(EditorGraph.IsValid());
+	check(InArgs._HistoryNavigationWidget.IsValid());
+
+	// Set-up shared breadcrumb defaults
+	FMargin BreadcrumbTrailPadding = FMargin(4.f, 2.f);
+	const FSlateBrush* BreadcrumbButtonImage = FAppStyle::GetBrush("BreadcrumbTrail.Delimiter");
+
+	this->ChildSlot
+	[
+		SNew(SBorder)
+		.BorderImage(FAppStyle::GetBrush(TEXT("Graph.TitleBackground")))
+		.HAlign(HAlign_Fill)
+		.AddMetaData<FTagMetaData>(FTagMetaData(TEXT("EventGraphTitleBar")))
+		[
+			SNew(SHorizontalBox)
+			+SHorizontalBox::Slot()
+			.AutoWidth()
+			[
+				SNew(SSeparator)
+				.Orientation(Orient_Vertical)
+			]
+			+ SHorizontalBox::Slot()
+			.AutoWidth()
+			[
+				InArgs._HistoryNavigationWidget.ToSharedRef()
+			]
+			+ SHorizontalBox::Slot()
+			.AutoWidth()
+			[
+				SNew(SSeparator)
+				.Orientation(Orient_Vertical)
+			]
+			// Title text/icon
+			+ SHorizontalBox::Slot()
+			[
+				SNew(SHorizontalBox)
+				+ SHorizontalBox::Slot()
+				.FillWidth(1.f)
+				[
+					SAssignNew(BreadcrumbTrailScrollBox, SScrollBox)
+					.Orientation(Orient_Horizontal)
+					.ScrollBarVisibility(EVisibility::Collapsed)
+					+ SScrollBox::Slot()
+					.Padding(0.f)
+					.VAlign(VAlign_Center)
+					[
+						SNew(SHorizontalBox)
+						+ SHorizontalBox::Slot()
+						.AutoWidth()
+						.VAlign(VAlign_Center)
+						[
+							SAssignNew(BreadcrumbTrail, SBreadcrumbTrail<SPCGGraphEditorTitleBar::FBreadcrumbItem>)
+							.ButtonStyle(FAppStyle::Get(), "GraphBreadcrumbButton")
+							.TextStyle(FAppStyle::Get(), "GraphBreadcrumbButtonText")
+							.ButtonContentPadding(BreadcrumbTrailPadding)
+							.DelimiterImage(BreadcrumbButtonImage)
+							.PersistentBreadcrumbs(true)
+							.OnCrumbClicked(this, &SPCGGraphEditorTitleBar::OnBreadcrumbClicked)
+						]
+					]
+				]
+			]
+		]
+	];
+
+	RebuildBreadcrumbTrail();
+	BreadcrumbTrailScrollBox->ScrollToEnd();
+}
+
+void SPCGGraphEditorTitleBar::OnBreadcrumbClicked(const FBreadcrumbItem& Crumb)
+{
+	const UPCGEditorGraph* CrumbGraph = Crumb.EditorGraph.Get();
+	if (!CrumbGraph)
+	{
+		return;
+	}
+	
+	TSharedPtr<FPCGEditor> Editor = CrumbGraph->GetEditor().Pin();
+	if (!Editor || !CrumbGraph->GetPCGGraph())
+	{
+		return;
+	}
+
+	Editor->OpenDocument(CrumbGraph->GetPCGGraph(), FDocumentTracker::EOpenDocumentCause::OpenNewDocument);
+}
+
+void SPCGGraphEditorTitleBar::RebuildBreadcrumbTrail()
+{
+	if (!EditorGraph.IsValid() || !EditorGraph->GetPCGGraph())
+	{
+		BreadcrumbTrail->ClearCrumbs(false);
+		return;
+	}
+
+	TArray<TWeakObjectPtr<const UPCGEditorGraph>> Stack;
+	Stack.Push(EditorGraph);
+
+	if (UPCGGraph* ParentGraph = EditorGraph->GetPCGGraph()->GetEmbeddedParentGraph(); ParentGraph && ParentGraph->GetEditorGraph())
+	{
+		Stack.Push(ParentGraph->GetEditorGraph());
+	}
+
+	BreadcrumbTrail->ClearCrumbs(false);
+
+	while (Stack.Num() > 0)
+	{
+		TWeakObjectPtr<const UPCGEditorGraph> WeakGraph(Stack.Pop());
+		TAttribute<FText> CrumbText = TAttribute<FText>::Create(TAttribute<FText>::FGetter::CreateStatic(&SPCGGraphEditorTitleBar::GetTitleForOneCrumb, WeakGraph));
+		BreadcrumbTrail->PushCrumb(CrumbText, { WeakGraph });
+	}
+}
+
+#undef LOCTEXT_NAMESPACE
diff --git a/Engine/Plugins/PCG/Source/PCGEditor/Private/Widgets/SPCGEditorGraphTitleBar.h b/Engine/Plugins/PCG/Source/PCGEditor/Private/Widgets/SPCGEditorGraphTitleBar.h
new file mode 100644
--- /dev/null
+++ b/Engine/Plugins/PCG/Source/PCGEditor/Private/Widgets/SPCGEditorGraphTitleBar.h
@@ -0,0 +1,46 @@
+// Copyright Epic Games, Inc. All Rights Reserved.
+
+#pragma once
+
+#include "Internationalization/Text.h"
+#include "Widgets/Navigation/SBreadcrumbTrail.h"
+#include "Widgets/SCompoundWidget.h"
+
+class FPCGEditor;
+class UPCGEditorGraph;
+
+class SPCGGraphEditorTitleBar : public SCompoundWidget
+{
+public:
+	SLATE_BEGIN_ARGS(SPCGGraphEditorTitleBar)
+		: _Graph(nullptr)
+		, _Editor()
+	{
+	}
+		SLATE_ARGUMENT(UPCGEditorGraph*, Graph)
+		SLATE_ARGUMENT(TWeakPtr<FPCGEditor>, Editor)
+		SLATE_ARGUMENT(TSharedPtr<SWidget>, HistoryNavigationWidget)
+	SLATE_END_ARGS()
+
+	void Construct(const FArguments& InArgs);
+
+protected:
+
+	struct FBreadcrumbItem
+	{
+		TWeakObjectPtr<const UPCGEditorGraph> EditorGraph;
+	};
+
+	void RebuildBreadcrumbTrail();
+	void OnBreadcrumbClicked(const FBreadcrumbItem& Crumb);
+
+	static FText GetTitleForOneCrumb(TWeakObjectPtr<const UPCGEditorGraph> InWeakEditorGraph);
+
+	/** Edited graph */
+	TWeakObjectPtr<UPCGEditorGraph> EditorGraph = nullptr;
+
+	TSharedPtr<SScrollBox> BreadcrumbTrailScrollBox;
+
+	/** Breadcrumb trail widget */
+	TSharedPtr<SBreadcrumbTrail<FBreadcrumbItem>> BreadcrumbTrail;
+};
diff --git a/Engine/Plugins/PCG/Source/PCGEditor/Public/PCGEditor.h b/Engine/Plugins/PCG/Source/PCGEditor/Public/PCGEditor.h
--- a/Engine/Plugins/PCG/Source/PCGEditor/Public/PCGEditor.h
+++ b/Engine/Plugins/PCG/Source/PCGEditor/Public/PCGEditor.h
@@ -27,6 +27,7 @@ class SPCGEditorGraphDebugObjectTree;
 class SPCGEditorGraphDetailsView;
 class SPCGEditorGraphDeterminismListView;
 class SPCGEditorGraphFind;
+class SPCGEditorGraphEmbeddedSubgraphsView;
 class SPCGEditorGraphLogView;
 class SPCGEditorGraphNodePalette;
 class SPCGEditorGraphProfilingView;
@@ -41,6 +42,7 @@ class UPCGEditorGraphNodeBase;
 class UPCGGraph;
 class UPCGSubsystem;
 struct FPCGCompilerDiagnostics;
+struct FAssetOpenArgs;
 class FWorkspaceItem;
 
 DECLARE_MULTICAST_DELEGATE_OneParam(FOnInspectedNodeChanged, UPCGEditorGraphNodeBase*);
@@ -78,13 +80,13 @@ enum class EPCGToolbarButtons
 	CancelExecution,
 	OpenDebugObjectTreeTab,
 	GraphParams,
-	GraphSettings
+	GraphSettings,
+	Graph,
 };
 
 namespace PCGEditorTabs
 {
 	constexpr FLazyName GraphEditorID = TEXT("GraphEditor");
-	constexpr FLazyName PCGEditorDocumentID = TEXT("PCGEditorDocumentID");
 	constexpr FLazyName PropertyDetailsID[] = {
 		TEXT("PropertyDetails"),
 		TEXT("PropertyDetails2"),
@@ -108,6 +110,7 @@ namespace PCGEditorTabs
 		TEXT("Viewport2"),
 		TEXT("Viewport3"),
 		TEXT("Viewport4") };
+	constexpr FLazyName EmbeddedSubgraphsID = TEXT("EmbeddedSubgraphs");
 
 	constexpr int32 NumPropertyDetailTabs = sizeof(PropertyDetailsID) / sizeof(FLazyName);
 	constexpr int32 NumAttributeTabs = sizeof(AttributesID) / sizeof(FLazyName);
@@ -118,18 +121,40 @@ class FPCGEditor : public FWorkflowCentricApplication, public FGCObject, public
 {
 public:
 	/** Edits the specified PCGGraph */
+	PCGEDITOR_API void InitializePCGEditorGraph(UPCGGraph* InPCGGraph);
+	PCGEDITOR_API UPCGEditorGraph* GetOrCreatePCGEditorGraph(UPCGGraph* InPCGGraph);
 	PCGEDITOR_API void Initialize(const EToolkitMode::Type InMode, const TSharedPtr<class IToolkitHost>& InToolkitHost, UPCGGraph* InPCGGraph, UObject* InObjectToEdit = nullptr);
 
-	/** Get the PCG editor graph being edited. */
-	PCGEDITOR_API UPCGEditorGraph* GetPCGEditorGraph();
+
+	/** Open a document tab */
+	PCGEDITOR_API void NavigateTab(FDocumentTracker::EOpenDocumentCause InCause);
+	PCGEDITOR_API TSharedPtr<SDockTab> OpenDocument(UPCGGraph* InPCGGraph, FDocumentTracker::EOpenDocumentCause InCause);
+	PCGEDITOR_API void CloseDocument(UPCGGraph* InPCGGraph);
+	PCGEDITOR_API FPCGEditor* OpenEditorForGraph(UPCGGraph* InPCGGraph);
+	PCGEDITOR_API static void OpenAssets(const FAssetOpenArgs& OpenArgs);
+
+	UE_DEPRECATED(5.8, "Use GetMainEditorGraph/GetFocusedEditorGraph instead")
+	UPCGEditorGraph* GetPCGEditorGraph() { return GetMainEditorGraph(); }
+
+	/** Get the root PCG editor graph being edited. */
+	PCGEDITOR_API UPCGEditorGraph* GetMainEditorGraph() const;
+
+	/** Get the focused PCG editor graph being edited. (root or embedded subgraph) */
+	PCGEDITOR_API UPCGEditorGraph* GetFocusedEditorGraph() const;
 
 	/** Gets/Creates the PCG graph editor for a given PCG graph */
 	PCGEDITOR_API static UPCGEditorGraph* GetPCGEditorGraph(UPCGGraph* InGraph);
 	PCGEDITOR_API static UPCGEditorGraph* GetPCGEditorGraph(const UPCGNode* InNode);
 	PCGEDITOR_API static UPCGEditorGraph* GetPCGEditorGraph(const UPCGSettings* InSettings);
 
-	/** Get the PCG graph being edited. */
-	const UPCGGraph* GetPCGGraph() { return PCGGraphBeingEdited; }
+	UE_DEPRECATED(5.8, "Use GetMainGraph/GetFocusedGraph instead")
+	const UPCGGraph* GetPCGGraph() { return GetMainGraph(); }
+
+	/** Get the root PCG graph being edited */
+	PCGEDITOR_API UPCGGraph* GetMainGraph() const;
+
+	/** Get the focused PCG graph being edited. (root or embedded subgraph) */
+	PCGEDITOR_API UPCGGraph* GetFocusedGraph() const;
 
 	/** Sets the execution stack that want to inspect. */
 	PCGEDITOR_API void SetStackBeingInspected(const FPCGStack& FullStack);
@@ -236,7 +261,6 @@ PRAGMA_ENABLE_DEPRECATION_WARNINGS
 
 	UPCGDefaultExecutionSource* GetDefaultExecutionSource() const { return PCGDefaultExecutionSource; }
 
-	PCGEDITOR_API TSharedPtr<FWorkflowTabFactory> CreateGraphEditorTabFactory();
 	PCGEDITOR_API TSharedPtr<FWorkflowTabFactory> CreatePaletteTabFactory();
 	PCGEDITOR_API TSharedPtr<FWorkflowTabFactory> CreateDebugObjectTreeTabFactory();
 	PCGEDITOR_API TSharedPtr<FWorkflowTabFactory> CreateFindTabFactory();
@@ -248,7 +272,15 @@ PRAGMA_ENABLE_DEPRECATION_WARNINGS
 	PCGEDITOR_API TSharedPtr<FWorkflowTabFactory> CreateGraphParamsTabFactory();
 	PCGEDITOR_API TSharedPtr<FWorkflowTabFactory> CreatePropertyDetailsTabFactory(int32 Index, TSharedRef<FWorkspaceItem> Group);
 	PCGEDITOR_API TSharedPtr<FWorkflowTabFactory> CreateViewportTabFactory(int32 Index, TSharedRef<FWorkspaceItem> Group);
+	PCGEDITOR_API TSharedPtr<FWorkflowTabFactory> CreateEmbeddedSubgraphsTabFactory();
+
+	/** Called when graph editor focus is changed */
+	PCGEDITOR_API void OnGraphEditorFocused(const TSharedRef<SGraphEditor>& InGraphEditor);
 
+	/** Called when the graph editor tab is backgrounded */
+	PCGEDITOR_API void OnGraphEditorBackgrounded(const TSharedRef<SGraphEditor>& InGraphEditor);
+
+	void OnGraphPreSave(const UPCGGraph* InGraph);
 protected:
 	PCGEDITOR_API virtual TAttribute<FGraphAppearanceInfo> GetAppearanceInfo() const;
 
@@ -382,6 +414,11 @@ private:
 	/** Whether the PCG object being edited is opened in details view or not */
 	bool IsEditGraphSettingsToggled() const;
 
+	/** Open main graph editor */
+	void OnEditGraph_Clicked();
+	/** Whether the main graph editor is opened */
+	bool IsEditGraphTabClosed() const;
+
 	/** Open panel to view and edit the graph parameters. */
 	void OnToggleGraphParamsPanel() const;
 	/** Is the graph params panel open. */
@@ -432,9 +469,9 @@ private:
 	bool CanConvertToStandaloneNodes() const;
 
 	/** Collapse the currently selected nodes in a subgraph */
-	void OnCollapseNodesInSubgraph();
+	void OnCollapseNodesInSubgraph(bool bCollapseToEmbeddedSubgraph);
 	/** Whether we can collapse nodes in a subgraph */
-	bool CanCollapseNodesInSubgraph() const;
+	bool CanCollapseNodesInSubgraph(bool bCollapseToEmbeddedSubgraph) const;
 
 	/** User is attempting to add a dynamic source pin to a node */
 	void OnAddDynamicInputPin();
@@ -477,10 +514,9 @@ private:
 	void OnCreateComment();
 
 	/** Get or Create new Graph editor widget */
-	TSharedRef<SGraphEditor> GetOrCreateGraphEditorWidget();
+	TSharedPtr<SGraphEditor> GetFocusedGraphEditorWidget() const;
 	void CreateGraphEditorCommands();
-	TSharedRef<SGraphEditor> CreateGraphEditorWidget();
-	TSharedRef<SGraphEditor> CreateGraphEditorWidget(TSharedPtr<FTabInfo> InTabInfo, UEdGraph* InGraph);
+	TSharedRef<SGraphEditor> CreateGraphEditorWidget(TSharedPtr<FTabInfo> InTabInfo, UPCGEditorGraph* InGraph);
 	
 	/** Get or Create palette widget */
 	TSharedRef<SPCGEditorGraphNodePalette> GetOrCreatePaletteWidget();
@@ -515,6 +551,9 @@ private:
 	/** Get or Create a new viewport tab widget */
 	TSharedRef<SPCGEditorViewport> GetOrCreateViewportWidget(int32 Index);
 
+	/** Get or Create a new embedded subgraphs tab widget */
+	TSharedRef<SPCGEditorGraphEmbeddedSubgraphsView> GetOrCreateEmbeddedSubgraphsWidget();
+
 	/** Called when the component inspected is generated/cleaned */
 	void OnSourceGenerated(IPCGGraphExecutionSource* InSource);
 
@@ -554,6 +593,8 @@ private:
 
 	void UpdateDefaultExecutionSource();
 
+	void CloseInvalidDocuments();
+
 	PCGEDITOR_API virtual TSharedRef<FTabManager::FLayout> GetDefaultLayout() const;
 	virtual void RegisterExtraTabFactories(FWorkflowAllowedTabSet& TabSet) {}
 
@@ -569,12 +610,14 @@ private:
 	TSharedPtr<SPCGEditorNodeSource> NodeSourceWidget;
 	TSharedPtr<SPCGEditorGraphUserParametersView> UserParamsWidget;
 	TSharedPtr<SPCGEditorViewport> ViewportWidgets[PCGEditorTabs::NumViewportTabs];
+	TSharedPtr<SPCGEditorGraphEmbeddedSubgraphsView> EmbeddedSubgraphsWidget;
 
 	TSharedPtr<FUICommandList> GraphEditorCommands;
 
-	TObjectPtr<UPCGGraph> PCGGraphBeingEdited = nullptr;
 	TObjectPtr<UPCGDefaultExecutionSource> PCGDefaultExecutionSource = nullptr;
-	UPCGEditorGraph* PCGEditorGraph = nullptr;
+
+	TObjectPtr<UPCGGraph> MainGraph = nullptr;
+	TWeakPtr<SGraphEditor> FocusedGraphEditor = nullptr;
 
 	FPCGEditorInspectionDataManager InspectionDataManager;
 
@@ -584,4 +627,7 @@ private:
 
 	/** Document manager for workflow tabs */
 	TSharedPtr<FDocumentTracker> DocumentManager;
+
+	/** Factory that spawns graph editors; used to look up all tabs spawned by it. */
+	TWeakPtr<FDocumentTabFactory> DocumentTabFactory;
 };
diff --git a/Engine/Plugins/PCG/Source/PCGEditor/Public/PCGEditorCommands.h b/Engine/Plugins/PCG/Source/PCGEditor/Public/PCGEditorCommands.h
--- a/Engine/Plugins/PCG/Source/PCGEditor/Public/PCGEditorCommands.h
+++ b/Engine/Plugins/PCG/Source/PCGEditor/Public/PCGEditorCommands.h
@@ -20,6 +20,7 @@ public:
 	// ~End TCommands<> interface
 
 	TSharedPtr<FUICommandInfo> CollapseNodes;
+	TSharedPtr<FUICommandInfo> CollapseNodesToEmbeddedSubgraph;
 	TSharedPtr<FUICommandInfo> ExportNodes;
 	TSharedPtr<FUICommandInfo> ConvertToStandaloneNodes;
 	TSharedPtr<FUICommandInfo> Find;
@@ -30,6 +31,7 @@ public:
 	TSharedPtr<FUICommandInfo> RunDeterminismNodeTest;
 	TSharedPtr<FUICommandInfo> RunDeterminismGraphTest;
 	TSharedPtr<FUICommandInfo> EditGraphSettings;
+	TSharedPtr<FUICommandInfo> EditGraph;
 	TSharedPtr<FUICommandInfo> ToggleGraphParams;
 	TSharedPtr<FUICommandInfo> CancelExecution;
 	TSharedPtr<FUICommandInfo> ToggleEnabled;
diff --git a/Engine/Plugins/PCG/Source/PCGEditor/Public/PCGEditorTabFactories.h b/Engine/Plugins/PCG/Source/PCGEditor/Public/PCGEditorTabFactories.h
--- a/Engine/Plugins/PCG/Source/PCGEditor/Public/PCGEditorTabFactories.h
+++ b/Engine/Plugins/PCG/Source/PCGEditor/Public/PCGEditorTabFactories.h
@@ -4,13 +4,16 @@
 
 #include "Delegates/DelegateCombinations.h"
 #include "WorkflowOrientedApp/WorkflowTabFactory.h"
+#include "WorkflowOrientedApp/WorkflowUObjectDocuments.h"
 
 enum ETabRole : uint8;
 class FPCGEditor;
 struct FSlateIcon;
 class FText;
 class FWorkspaceItem;
+class SGraphEditor;
 class SWidget;
+class UPCGEditorGraph;
 
 DECLARE_DELEGATE_OneParam(FPCGGenericTabClosed, FName);
 
@@ -49,4 +52,29 @@ private:
 
 	TFunction<TSharedRef<SWidget>()> CreateTabBodyFunc;
 	FPCGGenericTabClosed OnTabClosedDelegate;
-};
\ No newline at end of file
+};
+struct FPCGGraphEditorDocumentTabFactory : public FDocumentTabFactoryForObjects<UPCGEditorGraph>
+{
+public:
+	DECLARE_DELEGATE_RetVal_TwoParams(TSharedRef<SGraphEditor>, FOnCreateGraphEditorWidget, TSharedPtr<FTabInfo>, UPCGEditorGraph*);
+public:
+	FPCGGraphEditorDocumentTabFactory(TSharedPtr<FPCGEditor> InEditorPtr, FOnCreateGraphEditorWidget InCreateGraphEditorWidgetCallback);
+
+	virtual void OnTabActivated(TSharedPtr<SDockTab> InTab) const override;
+	virtual void OnTabBackgrounded(TSharedPtr<SDockTab> InTab) const override;
+	virtual void OnTabRefreshed(TSharedPtr<SDockTab> InTab) const override;
+	virtual void SaveState(TSharedPtr<SDockTab> InTab, TSharedPtr<FTabPayload> InPayload) const override;
+	virtual bool SupportsObjectType(UObject* InDocumentID) const override;
+
+protected:
+	static FText GetTabTitle(TWeakObjectPtr<UPCGEditorGraph> InWeakEditorGraph);
+
+	TSharedPtr<FPCGEditor> GetPCGEditor() const;
+
+	virtual TAttribute<FText> ConstructTabNameForObject(UPCGEditorGraph* InEditorGraph) const override;
+	virtual TSharedRef<SWidget> CreateTabBodyForObject(const FWorkflowTabSpawnInfo& InInfo, UPCGEditorGraph* InEditorGraph) const override;
+	virtual const FSlateBrush* GetTabIconForObject(const FWorkflowTabSpawnInfo& InInfo, UPCGEditorGraph* InEditorGraph) const override;
+
+protected:
+	FOnCreateGraphEditorWidget OnCreateGraphEditorWidget;
+};
diff --git a/Engine/Plugins/PCG/Source/PCGEditor/Public/Schema/PCGEditorGraphSchema.h b/Engine/Plugins/PCG/Source/PCGEditor/Public/Schema/PCGEditorGraphSchema.h
--- a/Engine/Plugins/PCG/Source/PCGEditor/Public/Schema/PCGEditorGraphSchema.h
+++ b/Engine/Plugins/PCG/Source/PCGEditor/Public/Schema/PCGEditorGraphSchema.h
@@ -53,6 +53,7 @@ public:
 	PCGEDITOR_API virtual void DroppedAssetsOnGraph(const TArray<FAssetData>& Assets, const FVector2f& GraphPosition, UEdGraph* Graph) const override;
 	PCGEDITOR_API virtual void GetAssetsGraphHoverMessage(const TArray<FAssetData>& Assets, const UEdGraph* HoverGraph, FString& OutTooltipText, bool& OutOkIcon) const override;
 	PCGEDITOR_API virtual void OnPinConnectionDoubleCicked(UEdGraphPin* PinA, UEdGraphPin* PinB, const FVector2f& GraphPosition) const override;
+	PCGEDITOR_API virtual void GetGraphDisplayInformation(const UEdGraph& Graph, /*out*/ FGraphDisplayInfo& DisplayInfo) const override;
 	//~ End EdGraphSchema Interface
 
 	PCGEDITOR_API virtual const FSlateBrush* GetMetadataTypeSlateBrush(EPCGContainerType ContainerType) const;
