Added a UE::CUEnum concept for testing if a type is a UEnum.  This requires testing for a legal StaticEnum declaration, which requires includes for StaticEnum<E>() calls which were previously missing.

#Synced-CL 50349841

diff --git a/Engine/Plugins/Animation/ControlRig/Source/ControlRig/Private/Rigs/RigHierarchyController.cpp b/Engine/Plugins/Animation/ControlRig/Source/ControlRig/Private/Rigs/RigHierarchyController.cpp
--- a/Engine/Plugins/Animation/ControlRig/Source/ControlRig/Private/Rigs/RigHierarchyController.cpp
+++ b/Engine/Plugins/Animation/ControlRig/Source/ControlRig/Private/Rigs/RigHierarchyController.cpp
@@ -7,6 +7,7 @@
 #include "ModularRig.h"
 #include "HelperUtil.h"
 #include "Engine/SkeletalMesh.h"
+#include "Math/Axis.h"
 #include "Misc/ScopeLock.h"
 #include "Units/Execution/RigUnit_PrepareForExecution.h"
 
diff --git a/Engine/Plugins/PCG/Source/PCG/Private/Compute/PCGComputeCommon.cpp b/Engine/Plugins/PCG/Source/PCG/Private/Compute/PCGComputeCommon.cpp
--- a/Engine/Plugins/PCG/Source/PCG/Private/Compute/PCGComputeCommon.cpp
+++ b/Engine/Plugins/PCG/Source/PCG/Private/Compute/PCGComputeCommon.cpp
@@ -27,6 +27,7 @@
 #include "Editor/IPCGEditorModule.h"
 
 #include "DynamicRHI.h"
+#include "PixelFormat.h"
 #include "RHIStats.h"
 #include "RenderResource.h"
 
diff --git a/Engine/Plugins/PCG/Source/PCG/Public/PCGGraph.h b/Engine/Plugins/PCG/Source/PCG/Public/PCGGraph.h
--- a/Engine/Plugins/PCG/Source/PCG/Public/PCGGraph.h
+++ b/Engine/Plugins/PCG/Source/PCG/Public/PCGGraph.h
@@ -5,6 +5,7 @@
 #include "PCGCommon.h"
 #include "PCGNode.h"
 #include "PCGSettings.h"
+#include "Concepts/UEnum.h"
 #include "Data/PCGToolData.h"
 #include "Editor/PCGGraphCustomization.h"
 #include "Graph/PCGStackContext.h"
@@ -239,7 +240,7 @@ public:
 		const FInstancedPropertyBag* UserParameters = GetUserParametersStruct();
 		check(UserParameters);
 
-		if constexpr (std::is_enum_v<T> && StaticEnum<T>())
+		if constexpr (UE::CUEnum<T>)
 		{
 			return FPCGGraphParameterExtension::GetGraphParameter<T>(*UserParameters, PropertyName, StaticEnum<T>());
 		}
@@ -258,7 +259,7 @@ public:
 		check(UserParameters);
 
 		EPropertyBagResult Result;
-		if constexpr (std::is_enum_v<T> && StaticEnum<T>())
+		if constexpr (UE::CUEnum<T>)
 		{
 			Result = FPCGGraphParameterExtension::SetGraphParameter(*UserParameters, PropertyName, Value, StaticEnum<T>());
 		}
diff --git a/Engine/Plugins/RenderGraphInsights/Source/RenderGraphInsights/Private/RenderGraphTrack.cpp b/Engine/Plugins/RenderGraphInsights/Source/RenderGraphInsights/Private/RenderGraphTrack.cpp
--- a/Engine/Plugins/RenderGraphInsights/Source/RenderGraphInsights/Private/RenderGraphTrack.cpp
+++ b/Engine/Plugins/RenderGraphInsights/Source/RenderGraphInsights/Private/RenderGraphTrack.cpp
@@ -3,6 +3,7 @@
 #include "RenderGraphTrack.h"
 
 #include "Framework/MultiBox/MultiBoxBuilder.h"
+#include "PixelFormat.h"
 #include "RenderGraphTimingViewSession.h"
 #include "RenderGraphTrackDrawHelper.h"
 #include "Widgets/Input/SSearchBox.h"
diff --git a/Engine/Source/Developer/OutputLog/Private/SOutputLog.cpp b/Engine/Source/Developer/OutputLog/Private/SOutputLog.cpp
--- a/Engine/Source/Developer/OutputLog/Private/SOutputLog.cpp
+++ b/Engine/Source/Developer/OutputLog/Private/SOutputLog.cpp
@@ -9,6 +9,7 @@
 #include "Framework/Text/TextLayout.h"
 #include "Misc/ConfigCacheIni.h"
 #include "Misc/OutputDeviceHelper.h"
+#include "Misc/OutputDevice.h"
 #include "Misc/ScopeLock.h"
 #include "SlateOptMacros.h"
 #include "Textures/SlateIcon.h"
diff --git a/Engine/Source/Runtime/Core/Public/Concepts/UEnum.h b/Engine/Source/Runtime/Core/Public/Concepts/UEnum.h
new file mode 100644
--- /dev/null
+++ b/Engine/Source/Runtime/Core/Public/Concepts/UEnum.h
@@ -0,0 +1,22 @@
+// Copyright Epic Games, Inc. All Rights Reserved.
+
+#pragma once
+
+#include "Concepts/CompleteType.h"
+#include "Concepts/SameAs.h"
+#include "UObject/CoreReflectedTypeAccessors.h"
+
+#include <type_traits>
+
+namespace UE
+{
+	/**
+	* Concept which describes a UEnum (i.e. an enum marked up with UENUM(), not a UObject derived from UEnum).
+	*/
+	template <typename T>
+	concept CUEnum = UE::CCompleteType<T> &&
+		requires (T* Ptr)
+		{
+			{ StaticEnum<std::remove_cv_t<T>>() } -> UE::CSameAs<UEnum*>;
+		};
+}
diff --git a/Engine/Source/Runtime/Core/Public/UObject/CoreReflectedTypeAccessors.h b/Engine/Source/Runtime/Core/Public/UObject/CoreReflectedTypeAccessors.h
--- a/Engine/Source/Runtime/Core/Public/UObject/CoreReflectedTypeAccessors.h
+++ b/Engine/Source/Runtime/Core/Public/UObject/CoreReflectedTypeAccessors.h
@@ -8,4 +8,4 @@ class UEnum;
 C++ templated Static(Enum) retrieval function prototypes.
 -----------------------------------------------------------------------------*/
 
-template<typename EnumType> UEnum* StaticEnum();
+template<typename EnumType> UEnum* StaticEnum() = delete;
