[PCG] First attempt at a generic attribute in Metadata, supporting any type that is supported by the reflection system.

* Created a new FPCGMetadataAttributeDesc that mirror FPropertyBagPropertyDesc, will be the equivalent of the EPCGMetadataTypes.
* New FPCGMetadataAttributeGeneric that is not yet used anywhere. It will rely on the FProperty system to be able to allocate and copy values around
* Supports all containers (Array/Set/Map)
* Use wrappers for Set/Map, that can be manipulated liked TSet and TMap
* Lots of Unit Tests that can also help to understand how generic attribute are used and their limitations.
* There are also benchmarks to compare performance. Performance is on-par with "basic types"
* Temporary functions in the Domain just to be able to create generic attributes.
* Lots of todos, wanted to focus more on functionality and performance.

#rnx

#Synced-CL 50157473

diff --git a/Engine/Plugins/PCG/Source/PCG/Private/Metadata/PCGMetadataAttribute.cpp b/Engine/Plugins/PCG/Source/PCG/Private/Metadata/PCGMetadataAttribute.cpp
--- a/Engine/Plugins/PCG/Source/PCG/Private/Metadata/PCGMetadataAttribute.cpp
+++ b/Engine/Plugins/PCG/Source/PCG/Private/Metadata/PCGMetadataAttribute.cpp
@@ -121,6 +121,11 @@ void FPCGMetadataAttributeBase::SetValuesFromValueKeys(const TArrayView<const TT
 }
 
 void FPCGMetadataAttributeBase::SetValuesFromValueKeys(const TArrayView<const PCGMetadataEntryKey>& EntryKeys, const TArrayView<const PCGMetadataValueKey>& ValueKeys, bool bResetValueOnDefaultValueKey)
+{
+	SetValuesFromValueKeys(PCGValueRangeHelpers::MakeConstValueRange(EntryKeys), ValueKeys, bResetValueOnDefaultValueKey);
+}
+
+void FPCGMetadataAttributeBase::SetValuesFromValueKeys(TConstPCGValueRange<PCGMetadataEntryKey> EntryKeys, const TArrayView<const PCGMetadataValueKey>& ValueKeys, bool bResetValueOnDefaultValueKey)
 {
 	if (EntryKeys.IsEmpty() || EntryKeys.Num() != ValueKeys.Num())
 	{
diff --git a/Engine/Plugins/PCG/Source/PCG/Private/Metadata/PCGMetadataAttributeGeneric.cpp b/Engine/Plugins/PCG/Source/PCG/Private/Metadata/PCGMetadataAttributeGeneric.cpp
new file mode 100644
--- /dev/null
+++ b/Engine/Plugins/PCG/Source/PCG/Private/Metadata/PCGMetadataAttributeGeneric.cpp
@@ -0,0 +1,1416 @@
+// Copyright Epic Games, Inc. All Rights Reserved.
+
+#include "Metadata/PCGMetadataAttributeGeneric.h"
+
+#include "Metadata/PCGMetadata.h"
+#include "Metadata/PCGMetadataDomain.h"
+#include "Metadata/Accessors/PCGAttributeAccessorHelpers.h"
+
+#include "UObject/TextProperty.h"
+
+namespace PCG::Private
+{
+	class FNullArchive : public FArchive
+	{
+		//override FObjectPtr to do nothing
+		virtual FArchive& operator<<(struct FObjectPtr& Value) override
+		{
+			return *this;
+		}
+	};
+	
+	template <bool bWrite>
+	struct FOptionalLock
+	{
+		FOptionalLock(PCG::FSharedLock& InLock, bool bInLock)
+			: Lock(InLock)
+			, bLock(bInLock)
+		{
+			if (bLock)
+			{
+				if constexpr (bWrite)
+				{
+					Lock.Lock();
+				}
+				else
+				{
+					Lock.LockShared();
+				}
+			}
+		}
+		
+		~FOptionalLock()
+		{
+			if (bLock)
+			{
+				if constexpr (bWrite)
+				{
+					Lock.Unlock();
+				}
+				else
+				{
+					Lock.UnlockShared();
+				}
+			}
+		}
+		
+		PCG::FSharedLock& Lock;
+		bool bLock;
+	};
+
+	/**
+	 * Equivalent of the method in PropertyBag to convert the descriptor to a property.
+	 * 
+	 */
+	FProperty* CreatePropertyFromDesc(const FPCGMetadataAttributeDesc& Desc, const FFieldVariant PropertyScope)
+	{
+		// Handle array and nested containers properties
+		if (Desc.ContainerTypes.Num() > 0)
+		{
+			FProperty* Prop = nullptr; // the first created container will fill the return value, nested ones will fill the inner
+
+			// support for nested containers, i.e. : TArray<TArray<float>>
+			FFieldVariant PropertyOwner = PropertyScope;
+			FProperty** ValuePropertyPtr = &Prop;
+
+			// Create the container list
+			for (EPCGMetadataAttributeContainerTypes BagContainerType : Desc.ContainerTypes)
+			{
+				switch(BagContainerType)
+				{
+				case EPCGMetadataAttributeContainerTypes::Array:
+					{
+						// create an array property as a container for the tail
+						FArrayProperty* ArrayProperty = new FArrayProperty(PropertyOwner, Desc.Name, RF_Public);
+						*ValuePropertyPtr = ArrayProperty;
+						ValuePropertyPtr = &ArrayProperty->Inner;
+						PropertyOwner = ArrayProperty;
+						break;
+					}
+				case EPCGMetadataAttributeContainerTypes::Set:
+					{
+						// create a set property as a container for the tail
+						FSetProperty* SetProperty = new FSetProperty(PropertyOwner, Desc.Name, RF_Public);
+						*ValuePropertyPtr = SetProperty;
+						ValuePropertyPtr = &SetProperty->ElementProp;
+						PropertyOwner = SetProperty;
+						break;
+					}
+				case EPCGMetadataAttributeContainerTypes::Map:
+					{
+						// create a map property as a container for the tail
+						FMapProperty* MapProperty = new FMapProperty(PropertyOwner, Desc.Name, RF_Public);
+						*ValuePropertyPtr = MapProperty;
+						ValuePropertyPtr = &MapProperty->ValueProp;
+						PropertyOwner = MapProperty;
+
+						// create the key property (container types are not supported for keys)
+						FPCGMetadataAttributeDesc KeyDesc;
+						KeyDesc.ValueType = Desc.KeyType;
+						KeyDesc.ValueTypeObject = Desc.KeyTypeObject;
+						KeyDesc.Name = Desc.Name;
+						FProperty* KeyProp = CreatePropertyFromDesc(KeyDesc, PropertyOwner);
+						if (!KeyProp)
+						{
+							delete MapProperty;
+							return nullptr;
+						}
+						
+						MapProperty->KeyProp = KeyProp;
+						break;
+					}
+				default:
+					ensureMsgf(false, TEXT("Unsupported container type %s"), *UEnum::GetValueAsString(BagContainerType));
+					break;
+				}
+			}
+
+			// finally create the tail type
+			FPCGMetadataAttributeDesc InnerDesc = Desc;
+			InnerDesc.ContainerTypes.Reset();
+			*ValuePropertyPtr = CreatePropertyFromDesc(InnerDesc, PropertyOwner);
+			
+			if (*ValuePropertyPtr == nullptr)
+			{
+				delete Prop;
+				return nullptr;
+			}
+			else
+			{
+				return Prop;
+			}
+		}
+		
+		const UScriptStruct* ScriptStruct = nullptr;
+		
+		switch (Desc.ValueType)
+		{
+		case EPCGMetadataTypes::Boolean:
+			{
+				FBoolProperty* Prop = new FBoolProperty(PropertyScope, Desc.Name, RF_Public);
+				Prop->SetBoolSize(sizeof(bool), true); // Enable native access (init the whole byte, rather than just first bit)
+				Prop->SetPropertyFlags(CPF_HasGetValueTypeHash);
+				return Prop;
+			}
+		case EPCGMetadataTypes::Byte:
+			{
+				FByteProperty* Prop = new FByteProperty(PropertyScope, Desc.Name, RF_Public);
+				Prop->SetPropertyFlags(CPF_HasGetValueTypeHash);
+				return Prop;
+			}
+		case EPCGMetadataTypes::Integer32:
+			{
+				FIntProperty* Prop = new FIntProperty(PropertyScope, Desc.Name, RF_Public);
+				Prop->SetPropertyFlags(CPF_HasGetValueTypeHash);
+				return Prop;
+			}
+		case EPCGMetadataTypes::Integer64:
+			{
+				FInt64Property* Prop = new FInt64Property(PropertyScope, Desc.Name, RF_Public);
+				Prop->SetPropertyFlags(CPF_HasGetValueTypeHash);
+				return Prop;
+			}
+		case EPCGMetadataTypes::Float:
+			{
+				FFloatProperty* Prop = new FFloatProperty(PropertyScope, Desc.Name, RF_Public);
+				Prop->SetPropertyFlags(CPF_HasGetValueTypeHash);
+				return Prop;
+			}
+		case EPCGMetadataTypes::Double:
+			{
+				FDoubleProperty* Prop = new FDoubleProperty(PropertyScope, Desc.Name, RF_Public);
+				Prop->SetPropertyFlags(CPF_HasGetValueTypeHash);
+				return Prop;
+			}
+		case EPCGMetadataTypes::Name:
+			{
+				FNameProperty* Prop = new FNameProperty(PropertyScope, Desc.Name, RF_Public);
+				Prop->SetPropertyFlags(CPF_HasGetValueTypeHash);
+				return Prop;
+			}
+		case EPCGMetadataTypes::String:
+			{
+				FStrProperty* Prop = new FStrProperty(PropertyScope, Desc.Name, RF_Public);
+				Prop->SetPropertyFlags(CPF_HasGetValueTypeHash);
+				return Prop;
+			}
+		case EPCGMetadataTypes::Text:
+			{
+				FTextProperty* Prop = new FTextProperty(PropertyScope, Desc.Name, RF_Public);
+				return Prop;
+			}
+		case EPCGMetadataTypes::Enum:
+			if (const UEnum* Enum = Cast<UEnum>(Desc.ValueTypeObject))
+			{
+				FEnumProperty* Prop = new FEnumProperty(PropertyScope, Desc.Name, RF_Public);
+				Prop->SetPropertyFlags(CPF_HasGetValueTypeHash);
+				FNumericProperty* UnderlyingProp = new FByteProperty(Prop, "UnderlyingType", RF_Public); // HACK: Hardwire to byte property for now for BP compatibility
+				UnderlyingProp->SetPropertyFlags(CPF_HasGetValueTypeHash | CPF_IsPlainOldData);
+				Prop->SetEnum(const_cast<UEnum*>(Enum));
+				Prop->AddCppProperty(UnderlyingProp);
+				return Prop;
+			}
+			break;
+		case EPCGMetadataTypes::SoftObjectPath:
+			ScriptStruct = TBaseStructure<FSoftObjectPath>::Get();
+			break;
+		case EPCGMetadataTypes::SoftClassPath:
+			ScriptStruct = TBaseStructure<FSoftClassPath>::Get();
+			break;
+		case EPCGMetadataTypes::Vector2:
+			ScriptStruct = TBaseStructure<FVector2D>::Get();
+			break;
+		case EPCGMetadataTypes::Vector:
+			ScriptStruct = TBaseStructure<FVector>::Get();
+			break;
+		case EPCGMetadataTypes::Vector4:
+			ScriptStruct = TBaseStructure<FVector4>::Get();
+			break;
+		case EPCGMetadataTypes::Quaternion:
+			ScriptStruct = TBaseStructure<FQuat>::Get();
+			break;
+		case EPCGMetadataTypes::Rotator:
+			ScriptStruct = TBaseStructure<FRotator>::Get();
+			break;
+		case EPCGMetadataTypes::Transform:
+			ScriptStruct = TBaseStructure<FTransform>::Get();
+			break;
+		case EPCGMetadataTypes::Struct:
+			ScriptStruct = Cast<UScriptStruct>(Desc.ValueTypeObject);
+			break;
+			
+			// @todo_pcg: Support class and objects
+// 		case EPCGMetadataTypes::Class:
+// 			if (const UClass* Class = Cast<UClass>(Desc.ValueTypeObject))
+// 			{
+// 				FClassProperty* Prop = new FClassProperty(PropertyScope, Desc.Name, RF_Public);
+// #if WITH_EDITORONLY_DATA
+// 				Prop->SetMetaClass(Desc.MetaClass ? Desc.MetaClass.Get() : const_cast<UClass*>(Class));
+// #else
+// 				Prop->SetMetaClass(const_cast<UClass*>(Class));
+// #endif
+// 				Prop->PropertyClass = UClass::StaticClass();
+// 				Prop->SetPropertyFlags(CPF_HasGetValueTypeHash);
+// 				return Prop;
+// 			}
+// 			break;
+// 		case EPCGMetadataTypes::SoftClass:
+// 			if (const UClass* Class = Cast<UClass>(Desc.ValueTypeObject))
+// 			{
+// 				FSoftClassProperty* Prop = new FSoftClassProperty(PropertyScope, Desc.Name, RF_Public);
+// #if WITH_EDITORONLY_DATA
+// 				Prop->SetMetaClass(Desc.MetaClass ? Desc.MetaClass.Get() : const_cast<UClass*>(Class));
+// #else
+// 				Prop->SetMetaClass(const_cast<UClass*>(Class));
+// #endif
+// 				Prop->PropertyClass = UClass::StaticClass();
+// 				Prop->SetPropertyFlags(CPF_HasGetValueTypeHash);
+// 				return Prop;
+// 			}
+// 			break;
+// 		case EPCGMetadataTypes::Object:
+// 			if (const UClass* Class = Cast<UClass>(Desc.ValueTypeObject))
+// 			{
+// 				FObjectProperty* Prop = new FObjectProperty(PropertyScope, Desc.Name, RF_Public);
+// 				if (Class->HasAnyClassFlags(CLASS_DefaultToInstanced))
+// 				{
+// 					Prop->SetPropertyFlags(CPF_InstancedReference);
+// 				}
+// 				Prop->SetPropertyClass(const_cast<UClass*>(Class));
+// 				Prop->SetPropertyFlags(CPF_HasGetValueTypeHash | CPF_TObjectPtr);
+// 				return Prop;
+// 			}
+// 			break;
+// 		case EPCGMetadataTypes::SoftObject:
+// 			if (const UClass* Class = Cast<UClass>(Desc.ValueTypeObject))
+// 			{
+// 				FSoftObjectProperty* Prop = new FSoftObjectProperty(PropertyScope, Desc.Name, RF_Public);
+// 				if (Class->HasAnyClassFlags(CLASS_DefaultToInstanced))
+// 				{
+// 					Prop->SetPropertyFlags(CPF_InstancedReference);
+// 				}
+// 				Prop->SetPropertyClass(const_cast<UClass*>(Class));
+// 				Prop->SetPropertyFlags(CPF_HasGetValueTypeHash);
+// 				return Prop;
+// 			}
+// 			break;
+		default:
+			ensureMsgf(false, TEXT("Unhandled type %s"), *UEnum::GetValueAsString(Desc.ValueType));
+		}
+		
+		if (ScriptStruct)
+		{
+			FStructProperty* Prop = new FStructProperty(PropertyScope, Desc.Name, RF_Public);
+			Prop->Struct = const_cast<UScriptStruct*>(ScriptStruct);
+
+			if (ScriptStruct->GetCppStructOps() && ScriptStruct->GetCppStructOps()->HasGetTypeHash())
+			{
+				Prop->SetPropertyFlags(CPF_HasGetValueTypeHash);
+			}
+				
+			if (ScriptStruct->StructFlags & STRUCT_HasInstancedReference)
+			{
+				Prop->SetPropertyFlags(CPF_ContainsInstancedReference);
+			}
+				
+			return Prop;
+		}
+
+		return nullptr;
+	}
+
+	/**
+	 * Taken from FArrayProperty::CopyValuesInternal in order to leverage MemCpy when we have plain old data
+	 * with the optimization twist for Object copies (also present in CopyValuesInternal)
+	 * @todo_pcg: Refactorize it in Core to re-use it. Maybe add it to FScriptArrayHelper?
+	 */
+	void CopyArray(const FArrayProperty* InProperty, void* DestPtr, const void* SrcPtr, int32 Count)
+	{
+		const FProperty* Inner = InProperty ? InProperty->Inner : nullptr;
+
+		if (!Inner || Count <= 0 || !DestPtr || !SrcPtr)
+		{
+			return;
+		}
+		
+		size_t Size = Inner->GetElementSize();
+		const uint8* SrcData = (uint8*)SrcPtr;
+		uint8* DestData = (uint8*)DestPtr;
+		if (!(Inner->PropertyFlags & CPF_IsPlainOldData))
+		{
+			// If we are in the closed body of a transaction, we have a special case for `FObjectProperty` that vastly
+			// improves performance when copying arrays around (in the order of 4x faster). We **do not want** to do
+			// this optimization with anything that derives from `FObjectProperty` though, so we **do not use** `IsA`
+			// as that would all catch classes that inherited from `FObjectProperty`!
+			if (AutoRTFM::IsClosed() && (Inner->GetClass() == FObjectProperty::StaticClass()))
+			{
+				#if UE_WITH_REMOTE_OBJECT_HANDLE
+				// CopyCompleteValue will resolve handles in the RawPtr <- TObjectPtr case.
+				// Since we can't abort while in an open scope, use a speculation scope to defer any
+				// migrations until this scope ends.
+				UE::RemoteExecutor::FSpeculationExecutionScope AvoidMigrationsInTheOpen;
+				#endif
+
+				// Since we've already written to each location that we are going to be copying over below, we know that
+				// all of the record-writes that the transactional system would have inserted are no-ops. Also, each copy
+				// of a `TObjectPtr` will go into the open to do mark the object as reachable, so by doing a single open
+				// here we can avoid doing N calls to record-write and N calls to open for each object in the array. Big
+				// saving on large arrays!
+				UE_AUTORTFM_OPEN_NO_VALIDATION
+				{
+					// This optimization assumes that `FObjectProperty` is just backed by a `TObjectPtr` which when copying
+					// does not do any heap allocations, just pulls over the pointer to the underlying `UObject` and then
+					// marks it as reachable for incremental reachability analysis. Lets check that this is still the case
+					// here so that if anyone ever changes that in future this line will shout at them!
+					static_assert(std::is_same_v<decltype(FObjectProperty::GetDefaultPropertyValue()), TObjectPtr<UObject>>);
+
+					for (int32 i = 0; i < Count; i++)
+					{
+						Inner->CopyCompleteValue(DestData + i * Size, SrcData + i * Size);
+					}
+				};
+			}
+			else
+			{
+				for (int32 i = 0; i < Count; i++)
+				{
+					Inner->CopyCompleteValue(DestData + i * Size, SrcData + i * Size);
+				}
+			}
+		}
+		else
+		{			
+			FMemory::Memcpy(DestData, SrcData, Count * Size);
+		}
+	}
+}
+
+/**
+ * Create a new wrapper using the Attribute desc. This is expected to be stored in a SharedPtr to avoid having 
+ * to always recreate it.
+ */
+FPCGAttributeProperty::FPCGAttributeProperty(const FPCGMetadataAttributeDesc& InAttributeDesc)
+{
+	// Use an array property to be able to leverage FScriptArrayHelper. It will contain the property
+	// that matches with the Descriptor.
+	// Note: CreatePropertyFromDesc will allocate a raw pointer, and Property will become its owner. When Property
+	// is destroyed, all allocated properties will also be deleted. It's all done internally.
+	Property = MakeUnique<FArrayProperty>(FFieldVariant{nullptr}, NAME_None, RF_Public);
+	FProperty* NewProperty = PCG::Private::CreatePropertyFromDesc(InAttributeDesc, Property.Get());
+
+	if (!NewProperty)
+	{
+		// ERROR
+		Property.Reset();
+		return;
+	}
+	
+	Property->AddCppProperty(NewProperty);
+	// We need to link the property so it fills the Property Flags and other members correctly.
+	// We don't need an archive (that's how it is done in the FProperty::StaticLink), but FNullArchive is private so we copied it.
+	PCG::Private::FNullArchive ArDummy;
+	Property->Link(ArDummy);
+	
+	// Cache info if the data is plain old data.
+	bIsPlainOldData = !!(NewProperty->PropertyFlags & CPF_IsPlainOldData);
+	
+	// Compress strings and soft paths, if they are not in a container.
+	bCompressData = InAttributeDesc.ContainerTypes.IsEmpty() &&  
+		(InAttributeDesc.ValueType == EPCGMetadataTypes::String 
+		|| InAttributeDesc.ValueType == EPCGMetadataTypes::Text
+		|| InAttributeDesc.ValueType == EPCGMetadataTypes::SoftObjectPath
+		|| InAttributeDesc.ValueType == EPCGMetadataTypes::SoftClassPath
+		|| InAttributeDesc.ValueTypeObject == TBaseStructure<FSoftObjectPath>()
+		|| InAttributeDesc.ValueTypeObject == TBaseStructure<FSoftClassPath>());
+
+	InnerElementCachedSize = NewProperty->GetElementSize();
+}
+
+FPCGAttributeProperty::~FPCGAttributeProperty() = default;
+
+void FPCGAttributeProperty::Copy(void* DestPtr, const void* SrcPtr, int32 Count)
+{
+	PCG::Private::CopyArray(Property.Get(), DestPtr, SrcPtr, Count);
+}
+
+const void* FPCGAttributeProperty::GetPtrInArray(const void* InArrayStart, int32 Index) const
+{
+	check(InnerElementCachedSize > 0);
+	return static_cast<const uint8*>(InArrayStart) + Index * InnerElementCachedSize;
+}
+
+FPCGMetadataAttributeGeneric::FPCGMetadataAttributeGeneric(const FPCGMetadataAttributeDesc& InAttributeDesc, FPCGMetadataDomain* InMetadata, FName InName, bool bInAllowsInterpolation)
+	: FPCGMetadataAttributeBase(InMetadata, InName, nullptr, bInAllowsInterpolation)
+	, CachedDesc(InAttributeDesc)
+	, UnderlyingProperty(MakeShared<FPCGAttributeProperty>(CachedDesc))
+{
+	// @todo_pcg support multi-containers
+	if (!ensure(CachedDesc.ContainerTypes.Num() <= 1))
+	{
+		UnderlyingProperty.Reset();
+		CachedDesc = FPCGMetadataAttributeDesc{};
+		return;
+	}
+	
+	Init();
+}
+
+FPCGMetadataAttributeGeneric::FPCGMetadataAttributeGeneric(TNotNull<const FPCGMetadataAttributeGeneric*> InParent, FPCGMetadataDomain* InMetadata, FName InName, bool bInAllowsInterpolation)
+	: FPCGMetadataAttributeBase(InMetadata, InName, InParent, bInAllowsInterpolation)
+	, CachedDesc(InParent->CachedDesc)
+	, UnderlyingProperty(InParent->UnderlyingProperty)
+{
+	// @todo_pcg: Should we init there so we always have a default value like before?
+	Init();
+
+	ValueKeyOffset = InParent->GetValueKeyOffsetForChild();
+	
+	if (UnderlyingProperty && ensure(!Values.IsEmpty()) && ensure(!InParent->Values.IsEmpty()))
+	{
+		// Copy the default value from the parent.
+		UnderlyingProperty->Copy(Values.GetData(), InParent->Values.GetData(), 1);
+	}
+}
+
+PCGMetadataValueKey FPCGMetadataAttributeGeneric::GetValueKeyOffsetForChild() const
+{
+	PCG::TSharedScopeLock ReadLock(ValueLock);
+	return Values.Num() + ValueKeyOffset;
+}
+
+void FPCGMetadataAttributeGeneric::AddReferencedObjects(FReferenceCollector& Collector)
+{
+	// @todo_pcg: We should register all the stored objects
+	// We'll have to see the cost of this, and if we don't hold object references for too long.
+}
+
+TArray<PCGMetadataValueKey> FPCGMetadataAttributeGeneric::AddValues_Internal(FInValues InValues, int32 Count, int32 ExistingIndex, bool bLockless)
+{
+	TArray<PCGMetadataValueKey> ValueKeys;
+
+	int32 FirstValueIndex = INDEX_NONE;
+
+	TOptional<FScriptArrayHelper> Helper = MakeArrayHelper();
+	if (!Helper.IsSet())
+	{
+		return ValueKeys;
+	}
+	
+	{
+		PCG::Private::FOptionalLock</*bWrite=*/true> WriteLock(ValueLock, !bLockless);
+		if (ExistingIndex == INDEX_NONE)
+		{
+			check(!bLockless);
+			FirstValueIndex = AllocateValues_Unsafe(Count);
+		}
+		else
+		{
+			// Expected to be already allocated.
+			check(Helper->Num() >= ExistingIndex + Count);
+			FirstValueIndex = ExistingIndex;
+		}
+		
+		check(FirstValueIndex != INDEX_NONE);
+		
+		Visit([this, &Helper, FirstValueIndex, Count](auto&& Value)
+			{
+				using T = std::decay_t<decltype(Value)>;
+				if constexpr(std::is_same_v<T, FInValuesByValue>)
+				{
+					const void* SrcPtr = Value.InValues;
+					void* DestPtr = Helper->GetRawPtr(FirstValueIndex);
+					UnderlyingProperty->Copy(DestPtr, SrcPtr, Count);
+				}
+				else if constexpr(std::is_same_v<T, FInValuesByValueSubset>)
+				{
+					for (int32 i = 0; i < Value.Indices.Num(); ++i)
+					{
+						const void* SrcPtr = UnderlyingProperty->GetPtrInArray(Value.InValues,  Value.Indices[i]);
+						void* DestPtr = Helper->GetRawPtr(FirstValueIndex + i);
+						UnderlyingProperty->Copy(DestPtr, SrcPtr, 1);
+					}
+				}
+				else if constexpr (std::is_same_v<T, FInValuesAsSet>)
+				{
+					const FSetProperty* InnerProperty = CastFieldChecked<FSetProperty>(UnderlyingProperty->Property->Inner);
+					for (int32 i = 0; i < Count; ++i)
+					{
+						FScriptSetHelper OtherHelper(InnerProperty, Helper->GetRawPtr(FirstValueIndex + i));
+						const int32 Num = Value.InValues[i].Num();
+				
+						for (int j = 0; j < Num; ++j)
+						{
+							OtherHelper.AddElement(Value.InValues[i][j]);
+						}
+					}
+				}
+				else if constexpr (std::is_same_v<T, FInValuesAsMap>)
+				{
+					const FMapProperty* InnerProperty = CastFieldChecked<FMapProperty>(UnderlyingProperty->Property->Inner);
+					for (int32 i = 0; i < Count; ++i)
+					{
+						FScriptMapHelper OtherHelper(InnerProperty, Helper->GetRawPtr(FirstValueIndex + i));
+						const int32 Num = Value.InValues[i].Num();
+				
+						for (int j = 0; j < Num; ++j)
+						{
+							OtherHelper.AddPair(Value.InValues[i][j].Key, Value.InValues[i][j].Value);
+						}
+					}
+				}
+				else if constexpr (std::is_same_v<T, FInValuesAsArray>)
+				{
+					const FArrayProperty* InnerProperty = CastFieldChecked<FArrayProperty>(UnderlyingProperty->Property->Inner);
+					for (int32 i = 0; i < Count; ++i)
+					{
+						FScriptArrayHelper OtherHelper(InnerProperty, Helper->GetRawPtr(FirstValueIndex + i));
+						const int32 Num = Value.InValues[i].template Get<1>();
+						if (InnerProperty->HasAnyPropertyFlags(CPF_IsPlainOldData))
+						{
+							OtherHelper.AddUninitializedValues(Num);
+						}
+						else
+						{
+							OtherHelper.AddValues(Num);
+						}
+				
+						PCG::Private::CopyArray(InnerProperty, OtherHelper.GetRawPtr(), Value.InValues[i].template Get<0>(), Num);
+					}
+				}
+				else
+				{
+					static_assert(!std::is_same_v<T, T>, "Missing variant case for AddValues_Internal");
+				}
+			}, InValues);
+	}
+
+	ValueKeys.SetNum(Count);
+	for (int ValueIndex = 0; ValueIndex < Count; ++ValueIndex)
+	{
+		ValueKeys[ValueIndex] = ValueKeyOffset + FirstValueIndex + ValueIndex;
+	}
+
+	return ValueKeys;
+}
+
+TArray<PCGMetadataValueKey> FPCGMetadataAttributeGeneric::AddCompressedValues(FInValuesByValue InUniqueValues,
+	TArray<PCGMetadataValueKey>& FoundValueKeys)
+{
+	const bool bHasDuplicateValues = (InUniqueValues.Count != FoundValueKeys.Num());
+
+	TArray<PCGMetadataValueKey> FoundUniqueValueKeys;
+	TArray<PCGMetadataValueKey>& FoundKeys = (bHasDuplicateValues ? FoundUniqueValueKeys : FoundValueKeys);
+
+	// Implementation note: when we don't have any duplicate values, the previously-set values in FoundValueKeys will be wiped out - this is intended
+	const bool bAtLeastOneValueNotFound = !FindValues(FoundKeys, InUniqueValues);
+
+	if (bAtLeastOneValueNotFound)
+	{
+		FInValuesByValueSubset ValueSubset{};
+		ValueSubset.InValues = InUniqueValues.InValues;
+		ValueSubset.Indices.Reserve(InUniqueValues.Count);
+
+		for (int ValueIndex = 0; ValueIndex < InUniqueValues.Count; ++ValueIndex)
+		{
+			if (FoundKeys[ValueIndex] == PCGNotFoundValueKey)
+			{
+				ValueSubset.Indices.Add(ValueIndex);
+			}
+		}
+
+		const int32 Count = ValueSubset.Indices.Num();
+		TArray<PCGMetadataValueKey> NewValueKeys = AddValues_Internal(FInValues{TInPlaceType<FInValuesByValueSubset>{}, MoveTemp(ValueSubset)}, Count);
+
+		int32 NewlyAddedValueKeyIndex = 0;
+		for (int ValueIndex = 0; ValueIndex < InUniqueValues.Count; ++ValueIndex)
+		{
+			if (FoundKeys[ValueIndex] == PCGNotFoundValueKey)
+			{
+				FoundKeys[ValueIndex] = NewValueKeys[NewlyAddedValueKeyIndex++];
+			}
+		}
+	}
+
+	// Remap to full array if needed
+	if (bHasDuplicateValues)
+	{
+		for (PCGMetadataValueKey& ValueToRemap : FoundValueKeys)
+		{
+			ValueToRemap = FoundUniqueValueKeys[ValueToRemap];
+		}
+	}
+
+	return FoundValueKeys;
+}
+
+bool FPCGMetadataAttributeGeneric::FindValues(TArray<PCGMetadataValueKey>& OutValueKeys, const FInValuesByValue& InValues)
+{
+	OutValueKeys.Init(PCGNotFoundValueKey, InValues.Count);
+
+	int ValueKeysSet = 0;
+	FindValuesInternal(InValues, OutValueKeys, ValueKeysSet, /*bIsRoot=*/true);
+
+	return (ValueKeysSet == InValues.Count);
+}
+
+void FPCGMetadataAttributeGeneric::FindValuesInternal(const FInValuesByValue& InValues, TArray<PCGMetadataValueKey>& ValueKeys, int& ValueKeysSet, bool bIsRoot) const
+{
+	check(InValues.Count == ValueKeys.Num());
+	
+	if (!ensure(UnderlyingProperty))
+	{
+		return;
+	}
+
+	auto Compare = [this, &InValues](int32 InValueIndex, PCGMetadataValueKey ValueKey)
+	{	
+		const void* Value1 = GetReadAddressForValueKey_Unsafe(ValueKey);
+		const void* Value2 = UnderlyingProperty->GetPtrInArray(InValues.InValues, InValueIndex);
+	
+		if (Value1 == nullptr || Value2 == nullptr)
+		{
+			return false;
+		}
+
+		return UnderlyingProperty->Property->Inner->Identical(Value1, Value2);
+	};
+
+	if (bIsRoot)
+	{
+		for (int ValueIndex = 0; ValueIndex < InValues.Count; ++ValueIndex)
+		{
+			if (Compare(ValueIndex, PCGDefaultValueKey))
+			{
+				ValueKeys[ValueIndex] = PCGDefaultValueKey;
+				++ValueKeysSet;
+			}
+		}
+	}
+
+	if (ValueKeysSet != InValues.Count)
+	{
+		const FPCGMetadataAttributeBase* ThisParent = GetParent();
+		if (ThisParent && ThisParent->IsGeneric())
+		{
+			static_cast<const FPCGMetadataAttributeGeneric*>(ThisParent)->FindValuesInternal(InValues, ValueKeys, ValueKeysSet, /*bRoot=*/false);
+		}
+	}
+
+	if (ValueKeysSet != InValues.Count)
+	{
+		PCG::Private::FOptionalLock</*bWrite=*/false> ReadLock(ValueLock, /*bLock=*/bIsRoot);
+		for (int ValueIndex = 0; ValueIndex < InValues.Count; ++ValueIndex)
+		{
+			if (ValueKeys[ValueIndex] != PCGNotFoundValueKey)
+			{
+				continue;
+			}
+
+			for (PCGMetadataValueKey ValueKey = 1; ValueKey < Values.Num(); ++ValueKey)
+			{
+				const PCGMetadataValueKey OffsetKey = ValueKey + ValueKeyOffset;
+				if (Compare(ValueIndex, OffsetKey))
+				{
+					ValueKeys[ValueIndex] = OffsetKey;
+					++ValueKeysSet;
+					break;
+				}
+			}
+		}
+	}
+}
+
+void FPCGMetadataAttributeGeneric::Init()
+{
+	if (!ensure(UnderlyingProperty) || !ensure(Values.IsEmpty()) || !UnderlyingProperty->Property)
+	{
+		return;
+	}
+
+	// Initialize the array with a single element containing the default value
+	FScriptArrayHelper Helper(UnderlyingProperty->Property.Get(), &Values);
+	Helper.AddValue();
+	
+	auto GetType = [](const FProperty* InProperty) -> int16
+	{
+		int16 InType = static_cast<int16>(PCGAttributeAccessorHelpers::GetMetadataTypeForProperty(InProperty));
+		if (InType == static_cast<int16>(EPCGMetadataTypes::Unknown))
+		{
+			// @todo_pcg TO BE FIXED
+			// For now mimic `Unknown + sizeof(T)` of the Traits
+			InType = static_cast<int16>(EPCGMetadataTypes::Unknown) + InProperty->GetElementSize();
+		}
+		
+		return InType;
+	};
+	
+	const FProperty* Property = UnderlyingProperty->Property->Inner;
+	if (!CachedDesc.ContainerTypes.IsEmpty())
+	{
+		if (const FArrayProperty* ArrayProperty = CastField<FArrayProperty>(Property))
+		{
+			TypeId = GetType(ArrayProperty->Inner);
+		}
+		else if (const FSetProperty* SetProperty = CastField<FSetProperty>(Property))
+		{
+			TypeId = GetType(SetProperty->ElementProp);
+		}
+		else if (const FMapProperty* MapProperty = CastField<FMapProperty>(Property))
+		{
+			TypeId = GetType(MapProperty->ValueProp);
+		}
+		else
+		{
+			checkNoEntry();
+		}
+	}
+	else
+	{
+		TypeId = GetType(Property);
+	}
+}
+
+void FPCGMetadataAttributeGeneric::Flatten()
+{
+	// @todo_pcg: TODO
+}
+
+void FPCGMetadataAttributeGeneric::FlattenAndCompress(const TArrayView<const PCGMetadataEntryKey>& InEntryKeysToKeep)
+{
+	// @todo_pcg: TODO
+}
+
+void FPCGMetadataAttributeGeneric::Reset()
+{
+	ValueKeyOffset = 0;
+	Parent = nullptr;
+
+	{
+		PCG::TUniqueScopeLock WriteLock(EntryMapLock);
+		EntryToValueKeyMap.Empty();
+	}
+
+	if (TOptional<FScriptArrayHelper> Helper = MakeArrayHelper(); Helper.IsSet())
+	{
+		PCG::TUniqueScopeLock WriteLock(ValueLock);
+		// Keep the default value
+		Helper->Resize(1);
+	}
+}
+
+FPCGMetadataAttributeBase* FPCGMetadataAttributeGeneric::Copy(FName NewName, FPCGMetadataDomain* InMetadata, bool bKeepParent, bool bCopyEntries, bool bCopyValues) const
+{
+	// If we copy an attribute where we don't want to keep the parent, while copying entries and/or values, we'll lose data.
+	// In that case, we will copy all the data from this attribute and all its ancestors.
+
+	// We can't keep the parent if we don't have the same root.
+	checkSlow(!bKeepParent || PCGMetadataHelpers::HasSameRoot(Metadata, InMetadata));
+
+	// Validate that the new name is valid
+	if (!IsValidName(NewName))
+	{
+		UE_LOG(LogPCG, Error, TEXT("Try to create a new attribute with an invalid name: %s"), *NewName.ToString());
+		return nullptr;
+	}
+
+	// This copies to a new attribute.
+	FPCGMetadataAttributeGeneric* AttributeCopy = nullptr;
+	if (bKeepParent)
+	{
+		AttributeCopy = new FPCGMetadataAttributeGeneric(this, InMetadata, NewName, bAllowsInterpolation);
+	}
+	else
+	{
+		AttributeCopy = new FPCGMetadataAttributeGeneric(CachedDesc, InMetadata, NewName, bAllowsInterpolation);
+		
+		if (ensure(AttributeCopy->UnderlyingProperty) && ensure(!Values.IsEmpty()) && ensure(!AttributeCopy->Values.IsEmpty()))
+		{
+			// Copy this default value to the attribute copy.
+			AttributeCopy->UnderlyingProperty->Copy(AttributeCopy->Values.GetData(), Values.GetData(), 1);
+		}
+	}
+
+	// Gather the chain of parents if we don't keep the parent and we want to copy entries/values.
+	// We always have at least one item, "this".
+	TArray<const FPCGMetadataAttributeGeneric*, TInlineAllocator<2>> Parents = { this };
+	if (!bKeepParent && (bCopyEntries || bCopyValues))
+	{
+		const FPCGMetadataDomain* CurrentMetadata = Metadata;
+		const FPCGMetadataAttributeGeneric* Current = this;
+
+		const FPCGMetadataDomain* ParentMetadata = PCGMetadataHelpers::GetParentMetadata(CurrentMetadata);
+		while (ParentMetadata && Current->Parent)
+		{
+			CurrentMetadata = ParentMetadata;
+			Current = static_cast<const FPCGMetadataAttributeGeneric*>(Current->Parent);
+			Parents.Add(Current);
+
+			ParentMetadata = PCGMetadataHelpers::GetParentMetadata(CurrentMetadata);
+		}
+	}
+
+	if (bCopyEntries)
+	{
+		// We go backwards, since we need to preserve order (root -> this)
+		// Latest entry in our Parents array is the root.
+		for (int32 i = Parents.Num() - 1; i >= 0; --i)
+		{
+			const FPCGMetadataAttributeGeneric* Current = Parents[i];
+
+			// Only need to lock if Current == this, since the parents are assumed to be const.
+			PCG::Private::FOptionalLock</*bWrite=*/false> ReadLock(Current->EntryMapLock, /*bLock=*/Current == this);
+			AttributeCopy->EntryToValueKeyMap.Append(Current->EntryToValueKeyMap);
+		}
+	}
+
+	if (bCopyValues)
+	{
+		// We go backwards, since we need to preserve order (root -> this)
+		// Latest entry in our Parents array is the root.
+		if (TOptional<FScriptArrayHelper> Helper = AttributeCopy->MakeArrayHelper(); Helper.IsSet())
+		{
+			for (int32 i = Parents.Num() - 1; i >= 0; --i)
+			{
+				const FPCGMetadataAttributeGeneric* Current = Parents[i];
+
+				// Only need to lock if Current == this, since the parents are assumed to be const.
+				PCG::Private::FOptionalLock</*bWrite=*/false> ReadLock(Current->ValueLock, /*bLock=*/Current == this);
+
+				// Do not copy the first value, as it is the default value, already copied/initialized by the ctor
+				TOptional<FScriptArrayHelper> ParentHelper = Current->MakeArrayHelper();
+				if (!ParentHelper.IsSet())
+				{
+					continue;
+				}
+				
+				const int32 NumValuesToCopy = ParentHelper->Num() - 1;
+				if (NumValuesToCopy <= 0 || AttributeCopy->AllocateValues_Unsafe(NumValuesToCopy) == INDEX_NONE)
+				{
+					continue;	
+				}
+
+				UnderlyingProperty->Copy(Helper->GetRawPtr(1), ParentHelper->GetRawPtr(1), NumValuesToCopy);
+
+				// The expected value key offset is the one for this attribute (i == 0), and only if we
+				// keep the parent. Otherwise we don't have any parent, so offset should be kept at 0.
+				if (i == 0 && bKeepParent)
+				{
+					AttributeCopy->ValueKeyOffset = Current->ValueKeyOffset;
+				}
+			}
+		}
+	}
+
+	return AttributeCopy;
+}
+
+FPCGMetadataAttributeBase* FPCGMetadataAttributeGeneric::Copy(FName NewName, UPCGMetadata* InMetadata, bool bKeepParent, bool bCopyEntries,
+	bool bCopyValues) const
+{
+	return Copy(NewName, InMetadata ? InMetadata->GetDefaultMetadataDomain() : nullptr, bKeepParent, bCopyEntries, bCopyValues);
+}
+
+FPCGMetadataAttributeBase* FPCGMetadataAttributeGeneric::CopyToAnotherType(int16 Type) const
+{
+	return nullptr;
+}
+
+void FPCGMetadataAttributeGeneric::SetValue(PCGMetadataEntryKey ItemKey, const FPCGMetadataAttributeBase* InAttribute, PCGMetadataEntryKey InEntryKey)
+{
+	if (!InAttribute || !InAttribute->IsGeneric())
+	{
+		return;
+	}
+
+	const FPCGMetadataAttributeGeneric* InAttributeGeneric = static_cast<const FPCGMetadataAttributeGeneric*>(InAttribute);
+
+	TOptional<FScriptArrayHelper> ThisHelper = MakeArrayHelper();
+	if (!ThisHelper.IsSet())
+	{
+		return;
+	}
+
+	int32 ValueKey;
+	
+	{
+		PCG::TUniqueScopeLock WriteLock(ValueLock);
+
+		const void* SrcPtr = InAttributeGeneric->GetReadAddressFromEntryKey_Unsafe(InEntryKey);
+		if (!SrcPtr)
+		{
+			return;	
+		}
+		
+		const int32 NewValueIndex = ThisHelper->AddValue();
+		ValueKey = ValueKeyOffset + NewValueIndex;
+		void* DestPtr = ThisHelper->GetRawPtr(NewValueIndex);
+		UnderlyingProperty->Copy(DestPtr, SrcPtr);
+	}
+
+	{
+		PCG::TUniqueScopeLock WriteLock(EntryMapLock);
+		EntryToValueKeyMap.Emplace(ItemKey, ValueKey);
+	}
+}
+
+void FPCGMetadataAttributeGeneric::SetZeroValue(PCGMetadataEntryKey ItemKey)
+{
+	TOptional<FScriptArrayHelper> ThisHelper = MakeArrayHelper();
+	if (!ThisHelper.IsSet())
+	{
+		return;
+	}
+
+	int32 ValueKey;
+	
+	{
+		PCG::TUniqueScopeLock WriteLock(ValueLock);
+		const int32 NewValueIndex = ThisHelper->AddValue();
+		ValueKey = ValueKeyOffset + NewValueIndex;
+	}
+
+	{
+		PCG::TUniqueScopeLock WriteLock(EntryMapLock);
+		EntryToValueKeyMap.Emplace(ItemKey, ValueKey);
+	}
+}
+
+void FPCGMetadataAttributeGeneric::AccumulateValue(PCGMetadataEntryKey ItemKey, const FPCGMetadataAttributeBase* InAttributeA,
+	PCGMetadataEntryKey InEntryKeyA, float Weight)
+{
+	// Not supported, will need to be specialized in the templated versions.
+}
+
+void FPCGMetadataAttributeGeneric::SetWeightedValue(PCGMetadataEntryKey ItemKey, const FPCGMetadataAttributeBase* InAttribute,
+	const TArrayView<const TPair<PCGMetadataEntryKey, float>>& InWeightedKeys)
+{
+	// Not supported, will need to be specialized in the templated versions.
+}
+
+void FPCGMetadataAttributeGeneric::SetValue(PCGMetadataEntryKey ItemKey, const FPCGMetadataAttributeBase* InAttributeA,
+	PCGMetadataEntryKey InEntryKeyA, const FPCGMetadataAttributeBase* InAttributeB, PCGMetadataEntryKey InEntryKeyB, EPCGMetadataOp Op)
+{
+	// Not supported, will need to be specialized in the templated versions.
+}
+
+bool FPCGMetadataAttributeGeneric::IsEqualToDefaultValue(PCGMetadataValueKey ValueKey) const
+{
+	if (ValueKey == PCGDefaultValueKey)
+	{
+		return true;
+	}
+	
+	if (ensure(UnderlyingProperty) && UnderlyingProperty->bCompressData)
+	{
+		return false;
+	}
+	else
+	{
+		return AreValuesEqual(ValueKey, PCGDefaultValueKey);
+	}
+}
+
+void FPCGMetadataAttributeGeneric::SetDefaultValueToFirstEntry()
+{
+	if (Values.Num() == 0 || Values.Num() > 2 || !ensure(UnderlyingProperty))
+	{
+		return;
+	}
+
+	{
+		PCG::TUniqueScopeLock WriteLock(ValueLock);
+
+		// First value key of this attribute.
+		const void* SrcPtr = GetReadAddressForValueKey_Unsafe(ValueKeyOffset + 1);
+		if (!SrcPtr)
+		{
+			return;	
+		}
+		
+		// Default value is at index 0
+		void* DestPtr = Values.GetData();
+		
+		UnderlyingProperty->Copy(DestPtr, SrcPtr);
+	}
+}
+
+bool FPCGMetadataAttributeGeneric::UsesValueKeys() const
+{
+	return ensure(UnderlyingProperty) ? UnderlyingProperty->bCompressData : false;
+}
+
+bool FPCGMetadataAttributeGeneric::AreValuesEqualForEntryKeys(PCGMetadataEntryKey EntryKey1, PCGMetadataEntryKey EntryKey2) const
+{
+	return AreValuesEqual(GetValueKey(EntryKey1), GetValueKey(EntryKey2));
+}
+
+bool FPCGMetadataAttributeGeneric::AreValuesEqual(PCGMetadataValueKey ValueKey1, PCGMetadataValueKey ValueKey2) const
+{
+	if (!ensure(UnderlyingProperty))
+	{
+		return false;
+	}
+	
+	if (ValueKey1 == ValueKey2)
+	{
+		return true;
+	}
+	
+	{
+		PCG::TSharedScopeLock ReadLock(ValueLock);
+		
+		const void* Value1 = GetReadAddressForValueKey_Unsafe(ValueKey1);
+		const void* Value2 = GetReadAddressForValueKey_Unsafe(ValueKey2);
+		
+		if (Value1 == nullptr || Value2 == nullptr)
+		{
+			return false;
+		}
+		
+		return UnderlyingProperty->Property->Inner->Identical(Value1, Value2);
+	}
+}
+
+/**
+ * In this function, we want to gather all the values are that given by the values keys, but only those who are at this level (aka. ValueKey >= ValueKeyOffset).
+ * UnretrievedValues will hold all the values that didn't managed to get because they are stored higher in the attribute hierarchy. The algorithm stops when
+ * all the values are retrieved.
+ * Extraction of the values will depend on the FOutValues variant, since we need to distinguished between those 5 possibilities:
+ * - Output by Values (copy)
+ * - Output by Pointer
+ * - Output as Array
+ * - Output as Set
+ * - Output as Map
+ * 
+ * On top of that, in order to reduce as much as we can virtual function calls, we also track for the different unretrived values if they are contiguous in memory, for input AND output.
+ * Values are contiguous in memory if their index in the output array are contiguous AND the values keys associated with them are also contiguous.
+ * For example, if we have indices {0, 1, 2, 5} (because we still haven't retrived values for those indices), and we have value keys {1, 2, 8, 9}, we have a contiguous
+ * range for the first two (their indices AND values keys are contiguous), but not for the rest.
+ * So we can do 3 extractions (in this example [Index:ValueKey]): [0:1 ,1:2], [2:8] and [5:9] and optimize the first one to reduce the number of virtual calls.
+ */
+void FPCGMetadataAttributeGeneric::GetValues_Internal(const TArrayView<const PCGMetadataValueKey> ValueKeys, FOutValues OutValues, TBitArray<>& UnretrievedValues) const
+{
+	bool bFoundAllKeys = true;
+	TConstSetBitIterator<> It(UnretrievedValues);
+	if (!It)
+	{
+		return;
+	}
+
+	const FPCGMetadataAttributeGeneric* ThisParent = static_cast<const FPCGMetadataAttributeGeneric*>(GetParent());
+	TOptional<FScriptArrayHelper> Helper = MakeArrayHelper();
+	if (!Helper)
+	{
+		return;
+	}
+
+	{
+		PCG::TSharedScopeLock ReadLock(ValueLock);
+		
+		struct FContinuousRange
+		{
+			void Reset()
+			{
+				StartIndex = INDEX_NONE;
+				StartValueKey = INDEX_NONE;
+				Count = 0;
+			}
+			
+			void Init(int32 InStartIndex, int64 InStartValueKey)
+			{
+				StartIndex = InStartIndex;
+				StartValueKey = InStartValueKey;
+				Count = 1;
+			}
+			
+			int32 StartIndex = INDEX_NONE;
+			int64 StartValueKey = INDEX_NONE;
+			int32 Count = 0;
+		};
+		
+		FContinuousRange ContinuousRange{};
+		
+		auto RetrieveValues = [this, &OutValues, &Helper](FContinuousRange& Range)
+		{
+			const uint8* SrcPtr = Helper->GetRawPtr(Range.StartValueKey);
+
+			// Since the values are stored in an array, they are all contiguous in memory so we can just do pointer arithmetic using the
+			// element size.
+			auto ApplyOnLoop = [this, &Range, SrcPtr](auto Func)
+			{
+				for (int32 i = 0; i < Range.Count; ++i)
+				{
+					Func(static_cast<const uint8*>(UnderlyingProperty->GetPtrInArray(SrcPtr, i)), Range.StartIndex + i);
+				}
+			};
+
+			Visit([this, &ApplyOnLoop, SrcPtr, &Range](auto&& OutValue)
+				{
+					using T = std::decay_t<decltype(OutValue)>;
+					if constexpr (std::is_same_v<T, FOutValuesAsArray>)
+					{
+						ApplyOnLoop([this, &OutValue](const uint8* SrcPtr, int32 Index)
+						{
+							FScriptArrayHelper OtherHelper(static_cast<FArrayProperty*>(UnderlyingProperty->Property->Inner), SrcPtr);
+							OutValue.OutValues[Index].template Get<0>() = OtherHelper.GetRawPtr();
+							OutValue.OutValues[Index].template Get<1>() = OtherHelper.Num();
+						});
+					}
+					else if constexpr (std::is_same_v<T, FOutValuesAsSet>)
+					{
+						ApplyOnLoop([this, &OutValue](const uint8* SrcPtr, int32 Index)
+						{
+							new (OutValue.OutValues[Index]) FScriptSetHelper(static_cast<FSetProperty*>(UnderlyingProperty->Property->Inner), SrcPtr);
+						});
+					}
+					else if constexpr (std::is_same_v<T, FOutValuesAsMap>)
+					{
+						ApplyOnLoop([this, &OutValue](const uint8* SrcPtr, int32 Index)
+						{
+							new (OutValue.OutValues[Index]) FScriptMapHelper(static_cast<FMapProperty*>(UnderlyingProperty->Property->Inner), SrcPtr);
+						});
+					}
+					else if constexpr (std::is_same_v<T, FOutValuesByPtr>)
+					{
+						ApplyOnLoop([this, &OutValue](const uint8* SrcPtr, int32 Index)
+						{
+							OutValue.OutValues[Index] = SrcPtr;
+						});
+					}
+					else if constexpr (std::is_same_v<T, FOutValuesByValue>)
+					{
+						void* DestPtr = const_cast<void*>(UnderlyingProperty->GetPtrInArray(OutValue.OutValues, Range.StartIndex));
+						UnderlyingProperty->Copy(DestPtr, SrcPtr, Range.Count);
+					}
+					else
+					{
+						static_assert(!std::is_same_v<T, T>, "Missing variant case for GetValues_Internal");
+					}
+			}, OutValues);
+		};
+		
+		for (; It; ++It)
+		{
+			const int32 Index = It.GetIndex();
+			const PCGMetadataValueKey ValueKey = ValueKeys[Index];
+			
+			int32 ValueKeyIndex = INDEX_NONE;
+			
+			if (ValueKey == PCGDefaultValueKey)
+			{
+				ValueKeyIndex = 0;
+			}
+			else if (ValueKey >= ValueKeyOffset)
+			{
+				int32 ValueIndex = ValueKey - ValueKeyOffset;
+				ValueKeyIndex = ValueIndex < Values.Num() ? ValueIndex : 0;
+			}
+			else if (!ThisParent)
+			{
+				ValueKeyIndex = 0;
+			}
+			else
+			{
+				bFoundAllKeys = false;
+			}
+			
+			if (ValueKeyIndex <= 0)
+			{
+				// Reached the end
+				if (ContinuousRange.StartIndex != INDEX_NONE)
+				{
+					RetrieveValues(ContinuousRange);
+					ContinuousRange.Reset();
+				}
+				
+				// Then copy the default value if the value key is 0
+				if (ValueKeyIndex == 0)
+				{
+					ContinuousRange.Init(Index, 0);
+					RetrieveValues(ContinuousRange);
+					ContinuousRange.Reset();
+				}
+			}
+			else
+			{
+				// Check if we are still in the same continuous range. If so, increment and continue.
+				if (ContinuousRange.StartIndex != INDEX_NONE && ContinuousRange.StartIndex + ContinuousRange.Count == Index && ContinuousRange.StartValueKey + ContinuousRange.Count == ValueKeyIndex)
+				{
+					ContinuousRange.Count++;
+					continue;
+				}
+				else if (ContinuousRange.StartIndex != INDEX_NONE)
+				{
+					// End of the range
+					RetrieveValues(ContinuousRange);
+					ContinuousRange.Reset();
+				}
+				
+				ContinuousRange.Init(Index, ValueKeyIndex);
+			}
+			
+			// If the value key index is not none, the value was (or will be) retrieved.
+			if (ValueKeyIndex != INDEX_NONE)
+			{
+				UnretrievedValues[Index] = false;
+			}
+		}
+		
+		// One last time
+		if (ContinuousRange.StartIndex != INDEX_NONE)
+		{
+			// End of the range
+			RetrieveValues(ContinuousRange);
+			ContinuousRange.Reset();
+		}
+	}
+
+	ensure(ThisParent || bFoundAllKeys);
+
+	if (ThisParent && !bFoundAllKeys)
+	{
+		ThisParent->GetValues_Internal(ValueKeys, OutValues, UnretrievedValues);
+	}
+}
+
+FInstancedPropertyBag FPCGMetadataAttributeGeneric::BuildStructForDebug(PCGMetadataEntryKey EntryKey) const
+{
+	if (!ensure(UnderlyingProperty))
+	{
+		return {};
+	}
+	
+	FInstancedPropertyBag Result;
+	
+	FPropertyBagPropertyDesc Desc{Name, UnderlyingProperty->Property->Inner};
+	
+	const UPropertyBag* PropertyBag = UPropertyBag::GetOrCreateFromDescs(MakeConstArrayView(&Desc, 1));
+	Result.InitializeFromBagStruct(PropertyBag);
+	
+	const FProperty* ResultProperty = PropertyBag->GetPropertyDescs()[0].CachedProperty;
+	
+	void* DestPtr = ResultProperty->ContainerPtrToValuePtr<void>(Result.GetMutableValue().GetMemory());
+	
+	{
+		PCG::TSharedScopeLock ReadLock(ValueLock);
+		if (const void* SrcPtr = GetReadAddressFromEntryKey_Unsafe(EntryKey))
+		{
+			ResultProperty->CopyCompleteValue(DestPtr, SrcPtr);
+		}
+	}
+
+	return Result;
+}
+
+TOptional<FScriptArrayHelper> FPCGMetadataAttributeGeneric::MakeArrayHelper() const
+{
+	if (!ensure(UnderlyingProperty) || !UnderlyingProperty->Property)
+	{
+		return {};
+	}
+	
+	return FScriptArrayHelper(UnderlyingProperty->Property.Get(), &Values);
+}
+
+const void* FPCGMetadataAttributeGeneric::GetReadAddressForValueKey_Unsafe(const PCGMetadataValueKey InValueKey) const
+{
+	check(!Values.IsEmpty());
+
+	if (InValueKey == PCGDefaultValueKey)
+	{
+		return Values.GetData();
+	}
+
+	if (InValueKey < ValueKeyOffset)
+	{
+		return ensure(Parent && Parent->IsGeneric()) ? static_cast<const FPCGMetadataAttributeGeneric*>(Parent)->GetReadAddressForValueKey_Unsafe(InValueKey) : nullptr;
+	}
+
+	const int32 ValueIndex = InValueKey - ValueKeyOffset;
+	
+	TOptional<FScriptArrayHelper> Helper = MakeArrayHelper();
+	if (!Helper || ValueIndex >= Values.Num())
+	{
+		return nullptr;
+	}
+
+	return Helper->GetRawPtr(ValueIndex);
+}
+
+const void* FPCGMetadataAttributeGeneric::GetReadAddressFromEntryKey_Unsafe(const PCGMetadataEntryKey InEntryKey) const
+{
+	return GetReadAddressForValueKey_Unsafe(GetValueKey(InEntryKey));
+}
+
+void FPCGMetadataAttributeGeneric::SetValueFromProperty(PCGMetadataEntryKey ItemKey, const void* SrcPtr, const FProperty* Property)
+{
+	if (!SrcPtr || !Property || ItemKey == PCGInvalidEntryKey)
+	{
+		return;
+	}
+	
+	if (TOptional<FScriptArrayHelper> Helper = MakeArrayHelper(); Helper.IsSet())
+	{
+		if (!UnderlyingProperty->Property->Inner->SameType(Property))
+		{
+			return;
+		}
+		
+		SetValuesFromValueKeys({ItemKey}, AddValues_Internal(FInValues{TInPlaceType<FInValuesByValue>{}, SrcPtr, 1}, 1));
+	}
+}
+
+int32 FPCGMetadataAttributeGeneric::PreallocateValues(TArrayView<PCGMetadataEntryKey*> EntryKeys, bool bLockless)
+{
+	int32 StartIndex = INDEX_NONE;
+	
+	if (!ensure(UnderlyingProperty))
+	{
+		return StartIndex;
+	}
+
+	if (!UnderlyingProperty->bCompressData)
+	{
+		PCG::Private::FOptionalLock</*bWrite=*/true> WriteValueLock(ValueLock, !bLockless);
+		StartIndex = AllocateValues_Unsafe(EntryKeys.Num());
+	}
+
+	{
+		PCG::Private::FOptionalLock</*bWrite=*/true> WriteEntryMapLock(EntryMapLock, !bLockless);
+		EntryToValueKeyMap.Reserve(EntryToValueKeyMap.Num() + EntryKeys.Num());
+
+		if (!UnderlyingProperty->bCompressData && StartIndex != INDEX_NONE)
+		{
+			for (int32 i = 0; i < EntryKeys.Num(); ++i)
+			{
+				const PCGMetadataValueKey ValueKey = StartIndex + i + ValueKeyOffset;
+				EntryToValueKeyMap.Emplace(*EntryKeys[i], ValueKey);
+			}
+		}
+	}
+
+	return StartIndex;
+}
+
+int32 FPCGMetadataAttributeGeneric::AllocateValues_Unsafe(int32 Count)
+{
+	if (TOptional<FScriptArrayHelper> Helper = MakeArrayHelper(); Helper.IsSet())
+	{
+		const int32 StartIndex = Helper->Num();
+
+		if (UnderlyingProperty->bIsPlainOldData)
+		{
+			Helper->AddUninitializedValues(Count);
+		}
+		else
+		{
+			Helper->AddValues(Count);
+		}
+
+		return StartIndex;
+	}
+	else
+	{
+		return INDEX_NONE;
+	}
+}
diff --git a/Engine/Plugins/PCG/Source/PCG/Private/Metadata/PCGMetadataCommon.cpp b/Engine/Plugins/PCG/Source/PCG/Private/Metadata/PCGMetadataCommon.cpp
--- a/Engine/Plugins/PCG/Source/PCG/Private/Metadata/PCGMetadataCommon.cpp
+++ b/Engine/Plugins/PCG/Source/PCG/Private/Metadata/PCGMetadataCommon.cpp
@@ -11,3 +11,69 @@ namespace PCGMetadataDomainID
 	const FPCGMetadataDomainID Data{EPCGMetadataDomainFlag::Data, -1, TEXT("Data")};
 	const FPCGMetadataDomainID Invalid{EPCGMetadataDomainFlag::Invalid, -1, TEXT("Invalid")};
 }
+
+bool FPCGMetadataAttributeDesc::operator==(const FPCGMetadataAttributeDesc& Other) const
+{
+	return Name == Other.Name &&
+		ValueType == Other.ValueType &&
+		ContainerTypes == Other.ContainerTypes &&
+		ValueTypeObject == Other.ValueTypeObject &&
+		KeyType == Other.KeyType &&
+		KeyTypeObject == Other.KeyTypeObject;
+}
+
+int16 FPCGMetadataAttributeDesc::AsLegacyTypeId() const
+{
+	if (static_cast<uint8>(ValueType) < static_cast<uint8>(EPCGMetadataTypes::EndLegacyTypes))
+	{
+		return static_cast<int16>(ValueType);
+	}
+	else if (ValueType == EPCGMetadataTypes::Struct)
+	{
+		if (ValueTypeObject == TBaseStructure<FVector2D>::Get())
+		{
+			return static_cast<int16>(EPCGMetadataTypes::Vector2);
+		}
+		else if (ValueTypeObject == TBaseStructure<FVector>::Get())
+		{
+			return static_cast<int16>(EPCGMetadataTypes::Vector);
+		}
+		else if (ValueTypeObject == TBaseStructure<FVector4>::Get())
+		{
+			return static_cast<int16>(EPCGMetadataTypes::Vector4);
+		}
+		else if (ValueTypeObject == TBaseStructure<FRotator>::Get())
+		{
+			return static_cast<int16>(EPCGMetadataTypes::Rotator);
+		}
+		else if (ValueTypeObject == TBaseStructure<FQuat>::Get())
+		{
+			return static_cast<int16>(EPCGMetadataTypes::Quaternion);
+		}
+		else if (ValueTypeObject == TBaseStructure<FTransform>::Get())
+		{
+			return static_cast<int16>(EPCGMetadataTypes::Transform);
+		}
+		else if (ValueTypeObject == TBaseStructure<FSoftObjectPath>::Get())
+		{
+			return static_cast<int16>(EPCGMetadataTypes::SoftObjectPath);
+		}
+		else if (ValueTypeObject == TBaseStructure<FSoftClassPath>::Get())
+		{
+			return static_cast<int16>(EPCGMetadataTypes::SoftClassPath);
+		}
+	}
+
+	return  static_cast<int16>(EPCGMetadataTypes::Unknown); 
+}
+
+int32 GetTypeHash(const FPCGMetadataAttributeDesc& Desc)
+{
+	int32 Hash = GetTypeHash(Desc.Name);
+	Hash = HashCombine(Hash, GetTypeHash(Desc.ValueType));
+	Hash = HashCombine(Hash, GetTypeHash(Desc.ContainerTypes));
+	Hash = HashCombine(Hash, GetTypeHash(Desc.ValueTypeObject));
+	Hash = HashCombine(Hash, GetTypeHash(Desc.KeyType));
+	Hash = HashCombine(Hash, GetTypeHash(Desc.KeyTypeObject));
+	return Hash;
+}
diff --git a/Engine/Plugins/PCG/Source/PCG/Private/Metadata/PCGMetadataDomain.cpp b/Engine/Plugins/PCG/Source/PCG/Private/Metadata/PCGMetadataDomain.cpp
--- a/Engine/Plugins/PCG/Source/PCG/Private/Metadata/PCGMetadataDomain.cpp
+++ b/Engine/Plugins/PCG/Source/PCG/Private/Metadata/PCGMetadataDomain.cpp
@@ -201,6 +201,81 @@ void FPCGMetadataDomain::Serialize(FArchive& InArchive)
 	}
 }
 
+FPCGMetadataAttributeGeneric* FPCGMetadataDomain::CreateGenericAttribute(const FPCGMetadataAttributeDesc& Desc, bool bAllowsInterpolation, bool bOverrideParent)
+{
+	const FName AttributeName = Desc.Name;
+	
+	if (!FPCGMetadataAttributeBase::IsValidName(AttributeName))
+	{
+		UE_LOG(LogPCG, Error, TEXT("Attribute name '%s' is invalid"), *AttributeName.ToString());
+		return nullptr;
+	}
+
+	const FPCGMetadataAttributeBase* ParentAttribute = nullptr;
+
+	if (bOverrideParent && Parent)
+	{
+		ParentAttribute = Parent->GetConstAttribute(AttributeName);
+	}
+
+	FPCGMetadataAttributeGeneric* NewAttribute = nullptr;
+
+	PRAGMA_DISABLE_DEPRECATION_WARNINGS
+	int16 TypeId = Desc.AsLegacyTypeId();
+	PRAGMA_ENABLE_DEPRECATION_WARNINGS
+
+	if (!ParentAttribute || (ParentAttribute->GetTypeId() != TypeId || !ParentAttribute->IsGeneric()))
+	{
+		// Can't parent if the types do not match, or is not generic
+		NewAttribute = new FPCGMetadataAttributeGeneric(Desc, this, AttributeName, bAllowsInterpolation);
+	}
+	else
+	{
+		NewAttribute = new FPCGMetadataAttributeGeneric(static_cast<const FPCGMetadataAttributeGeneric*>(ParentAttribute), this, AttributeName, bAllowsInterpolation);
+	}
+	
+	{
+		PCG::TUniqueScopeLock WriteLock(AttributeLock);
+
+		if (FPCGMetadataAttributeBase** ExistingAttribute = Attributes.Find(AttributeName))
+		{
+			delete NewAttribute;
+			if ((*ExistingAttribute)->GetTypeId() != TypeId || !(*ExistingAttribute)->IsGeneric())
+			{
+				UE_LOG(LogPCG, Error, TEXT("Attribute %s already exists but is not the right type. Abort."), *AttributeName.ToString());
+				return nullptr;
+			}
+			else
+			{
+				UE_LOG(LogPCG, Warning, TEXT("Attribute %s already exists"), *AttributeName.ToString());
+				NewAttribute = static_cast<FPCGMetadataAttributeGeneric*>(*ExistingAttribute);
+			}
+		}
+		else
+		{
+			NewAttribute->AttributeId = NextAttributeId++;
+			AddAttributeInternal(AttributeName, NewAttribute);
+
+			// Also when creating an attribute, notify the PCG Data owner that the latest attribute manipulated is this one.
+			SetLastCachedSelectorOnOwner(AttributeName);
+		}
+	}
+
+	return NewAttribute;
+}
+
+FPCGMetadataAttributeGeneric* FPCGMetadataDomain::GetMutableGenericAttribute(FName AttributeName)
+{
+	FPCGMetadataAttributeBase* Attribute = GetMutableAttribute(AttributeName);
+	return Attribute && Attribute->IsGeneric() ? static_cast<FPCGMetadataAttributeGeneric*>(Attribute) : nullptr;
+}
+
+const FPCGMetadataAttributeGeneric* FPCGMetadataDomain::GetConstGenericAttribute(FName AttributeName) const
+{
+	const FPCGMetadataAttributeBase* Attribute = GetConstAttribute(AttributeName);
+	return Attribute && Attribute->IsGeneric() ? static_cast<const FPCGMetadataAttributeGeneric*>(Attribute) : nullptr;
+}
+
 void FPCGMetadataDomain::AddToCrc(FArchiveCrc32& Ar, const UPCGData* Data, bool bFullDataCrc) const
 {
 	check(Data);
@@ -246,8 +321,9 @@ void FPCGMetadataDomain::AddToCrc(FArchiveCrc32& Ar, const UPCGData* Data, bool
 	{
 		Ar << const_cast<FName&>(Attribute->Name);
 
+		// @todo_pcg: Revisit CRC computation for generic attributes.
 		TUniquePtr<const IPCGAttributeAccessor> InputAccessor = PCGAttributeAccessorHelpers::CreateConstAccessor(Attribute, this);
-		if (!ensure(InputAccessor.IsValid()))
+		if (!InputAccessor.IsValid())
 		{
 			continue;
 		}
@@ -588,33 +664,36 @@ bool FPCGMetadataDomain::FlattenAndCompress(const TArrayView<const PCGMetadataEn
 	TRACE_CPUPROFILER_EVENT_SCOPE(FPCGSubMetadata::FlattenAndCompress);
 
 	// No keys or no parents, nothing to do
-	if (Attributes.IsEmpty())
+	int32 NumEntryToKeep = InEntryKeysToKeep.Num();
+	if (!Attributes.IsEmpty())
 	{
-		return false;
-	}
-
-	TArrayView<const PCGMetadataEntryKey> EntryKeysToKeep = InEntryKeysToKeep;
+		TArrayView<const PCGMetadataEntryKey> EntryKeysToKeep = InEntryKeysToKeep;
 		
-	if (InEntryKeysToKeep.Num() > 1 && !SupportsMultiEntries())
-	{
-		EntryKeysToKeep = InEntryKeysToKeep.Left(1);
-		UE_LOG(LogPCG, Warning, TEXT("Tried to flatten and compress a domain '%ls', which doesn't support multiple entries, with %d entries to keep. Will only keep the first one."), *DomainID.DebugName.ToString(), InEntryKeysToKeep.Num());
-	}
+		if (InEntryKeysToKeep.Num() > 1 && !SupportsMultiEntries())
+		{
+			EntryKeysToKeep = InEntryKeysToKeep.Left(1);
+			UE_LOG(LogPCG, Warning, TEXT("Tried to flatten and compress a domain '%ls', which doesn't support multiple entries, with %d entries to keep. Will only keep the first one."), *DomainID.DebugName.ToString(), InEntryKeysToKeep.Num());
+		}
 
-	{
-		PCG::TUniqueScopeLock WriteLock(AttributeLock);
-		for (auto& AttributePair : Attributes)
 		{
-			FPCGMetadataAttributeBase* Attribute = AttributePair.Value;
-			check(Attribute);
+			PCG::TUniqueScopeLock WriteLock(AttributeLock);
+			for (auto& AttributePair : Attributes)
+			{
+				FPCGMetadataAttributeBase* Attribute = AttributePair.Value;
+				check(Attribute);
 
-			Attribute->FlattenAndCompress(EntryKeysToKeep);
+				Attribute->FlattenAndCompress(EntryKeysToKeep);
+			}
 		}
 	}
+	else
+	{
+		NumEntryToKeep = 0;
+	}
 
 	Parent = nullptr;
 	ParentKeys.Reset();
-	ParentKeys.Init(PCGInvalidEntryKey, EntryKeysToKeep.Num());
+	ParentKeys.Init(PCGInvalidEntryKey, NumEntryToKeep);
 	ItemKeyOffset = 0;
 
 	return true;
diff --git a/Engine/Plugins/PCG/Source/PCG/Private/Tests/Metadata/PCGMetadataFlattenTest.cpp b/Engine/Plugins/PCG/Source/PCG/Private/Tests/Metadata/PCGMetadataFlattenTest.cpp
deleted file mode 100644
diff --git a/Engine/Plugins/PCG/Source/PCG/Public/Metadata/PCGMetadataAttribute.h b/Engine/Plugins/PCG/Source/PCG/Public/Metadata/PCGMetadataAttribute.h
--- a/Engine/Plugins/PCG/Source/PCG/Public/Metadata/PCGMetadataAttribute.h
+++ b/Engine/Plugins/PCG/Source/PCG/Public/Metadata/PCGMetadataAttribute.h
@@ -92,6 +92,7 @@ public:
 
 	/** Two arrays version of bulk setter to lock in write only once. Both arrays must be the same size. */
 	UE_API void SetValuesFromValueKeys(const TArrayView<const PCGMetadataEntryKey>& EntryKeys, const TArrayView<const PCGMetadataValueKey>& ValueKeys, bool bResetValueOnDefaultValueKey = true);
+	UE_API void SetValuesFromValueKeys(TConstPCGValueRange<PCGMetadataEntryKey> EntryKeys, const TArrayView<const PCGMetadataValueKey>& ValueKeys, bool bResetValueOnDefaultValueKey = true);
 	UE_API void SetValuesFromValueKeys(const TArrayView<const PCGMetadataEntryKey* const>& EntryKeys, const TArrayView<const PCGMetadataValueKey>& ValueKeys, bool bResetValueOnDefaultValueKey = true);
 
 	bool AllowsInterpolation() const { return bAllowsInterpolation; }
@@ -111,6 +112,9 @@ public:
 	/** Replaces any invalid characters in name with underscores. Returns true if Name was changed. */
 	static UE_API bool SanitizeName(FString& InOutName);
 
+	// @todo_pcg: To be removed
+	virtual bool IsGeneric() const { return false; }
+
 private:
 	// Unsafe version, needs to be write lock protected.
 	UE_API void SetValueFromValueKey_Unsafe(PCGMetadataEntryKey EntryKey, PCGMetadataValueKey ValueKey, bool bResetValueOnDefaultValueKey, bool bAllowInvalidEntries = false);
diff --git a/Engine/Plugins/PCG/Source/PCG/Public/Metadata/PCGMetadataAttributeGeneric.h b/Engine/Plugins/PCG/Source/PCG/Public/Metadata/PCGMetadataAttributeGeneric.h
new file mode 100644
--- /dev/null
+++ b/Engine/Plugins/PCG/Source/PCG/Public/Metadata/PCGMetadataAttributeGeneric.h
@@ -0,0 +1,556 @@
+// Copyright Epic Games, Inc. All Rights Reserved.
+
+#pragma once
+
+#include "PCGMetadataContainerTypes.h"
+#include "Metadata/PCGMetadataCommon.h"
+#include "PCGModule.h"
+
+#include "Helpers/PCGMetadataHelpers.h"
+#include "Metadata/PCGMetadataAttribute.h"
+#include "Metadata/PCGMetadataAttributeTraits.h"
+
+#include "Containers/StaticArray.h"
+#include "Misc/TVariant.h"
+#include "StructUtils/PropertyBag.h"
+
+class FPCGMetadataAttributeGeneric;
+
+namespace PCG::Private
+{
+	void CopyArray(const FArrayProperty* InProperty, void* DestPtr, const void* SrcPtr, int32 Count);
+	
+	template <typename T>
+	struct TIsStaticArray : std::false_type {};
+	
+	template <typename T, uint32 ...Extra>
+    struct TIsStaticArray<TStaticArray<T, Extra...>> : std::true_type {};
+	
+	template <typename Container>
+	concept CIsArrayLike = 
+		TIsTArray_V<Container> 
+		|| TIsTArrayView_V<Container>
+		|| TIsStaticArray<Container>::value;
+		
+}
+
+struct FPCGAttributeProperty
+{
+public:
+	PCG_API FPCGAttributeProperty(const FPCGMetadataAttributeDesc& InAttributeDesc);
+	PCG_API ~FPCGAttributeProperty();
+
+	PCG_API void Copy(void* DestPtr, const void* SrcPtr, int32 Count = 1);
+	PCG_API void AllocateValues(int32 Count = 1);
+
+private:
+	friend FPCGMetadataAttributeGeneric;
+
+	/** Unsafe method that will apply pointer arithmetic to get to the element in the array at the specified index. */
+	const void* GetPtrInArray(const void* InArrayStart, int32 Index) const;
+
+	TUniquePtr<FArrayProperty> Property;
+	size_t InnerElementCachedSize = 0;
+	bool bCompressData = false;
+	bool bIsPlainOldData = false;
+};
+
+class FPCGMetadataAttributeGeneric : public FPCGMetadataAttributeBase
+{
+public:
+	// New attribute with no parent
+	PCG_API FPCGMetadataAttributeGeneric(const FPCGMetadataAttributeDesc& InAttributeDesc, FPCGMetadataDomain* InMetadata, FName InName, bool bInAllowsInterpolation);
+
+	// New attribute with an existing parent
+	PCG_API FPCGMetadataAttributeGeneric(TNotNull<const FPCGMetadataAttributeGeneric*> InParent, FPCGMetadataDomain* InMetadata, FName InName, bool bInAllowsInterpolation);
+	
+	PCG_API void AddReferencedObjects(FReferenceCollector& Collector);
+
+	virtual bool IsGeneric() const override { return true; }
+	
+	const FPCGMetadataAttributeDesc& GetAttributeDesc() const { return CachedDesc; }
+
+	PCG_API virtual PCGMetadataValueKey GetValueKeyOffsetForChild() const override;
+
+	/**
+	 * GetValues section
+	 * Templated code to support all container types.
+	 * If the provided output container does not match the underlying type, it will fail.
+	 * It is also mandatory that the number of value keys and output container num match.
+	 * 
+	 * Supported types are:
+	 * - Single types that are valid for metadata
+	 * - Pointers to single types that are valid for metadata
+	 * - Typed ConstArrayViews for arrays
+	 * - PCG::TScriptSetWrapper<T> for set, wrapper that can be manipulated like a set
+	 * - PCG::TScriptMapWrapper<T> for map, wrapper that can be manipulated like a map
+	 */
+	template <typename T> requires (!std::is_pointer_v<T>)
+	void GetValues(TArrayView<const PCGMetadataValueKey> ValueKeys, TArrayView<T> OutValues) const
+	{
+		if (!PCG::Private::IsEquivalentDesc<T>(CachedDesc))
+		{
+			// @todo_pcg Better logging
+			UE_LOG(LogPCG, Error, TEXT("Type mismatch on GetValues"));
+			return;
+		}
+		
+		if (!ensure(ValueKeys.Num() == OutValues.Num()))
+		{
+			return;
+		}
+		
+		GetValues_Internal(ValueKeys, FOutValues{TInPlaceType<FOutValuesByValue>{}, OutValues.GetData()});
+	}
+
+	template <typename T>
+	void GetValues(TArrayView<const PCGMetadataValueKey> ValueKeys, TArrayView<const T*> OutValues) const
+	{
+		if (!PCG::Private::IsEquivalentDesc<T>(CachedDesc))
+		{
+			// @todo_pcg Better logging
+			UE_LOG(LogPCG, Error, TEXT("Type mismatch on GetValues"));
+			return;
+		}
+		
+		if (!ensure(ValueKeys.Num() == OutValues.Num()))
+		{
+			return;
+		}
+		
+		// We need to type erase the pointers, but to be fully C++ compliant about aliasing rules we can't reinterpret_cast. So make a copy of the pointers.
+		TArray<const void*, TInlineAllocator<256>> Temp;
+		Temp.SetNumUninitialized(OutValues.Num());
+		GetValues_Internal(ValueKeys, FOutValues{ TInPlaceType<FOutValuesByPtr>{}, Temp });
+		for (int32 i = 0; i < OutValues.Num(); ++i)
+		{
+			OutValues[i] = static_cast<const T*>(Temp[i]);
+		}
+	}
+	
+	template <typename T> 
+	void GetValues(TArrayView<const PCGMetadataValueKey> ValueKeys, TArrayView<TConstArrayView<T>> OutValues) const
+	{
+		if (!PCG::Private::IsEquivalentDesc<TArray<T>>(CachedDesc))
+		{
+			// @todo_pcg Better logging
+			UE_LOG(LogPCG, Error, TEXT("Type mismatch on GetValues"));
+			return;
+		}
+		
+		if (!ensure(ValueKeys.Num() == OutValues.Num()))
+		{
+			return;
+		}
+		
+		TArray<TTuple<const void*, int32>, TInlineAllocator<256>> Temp;
+		Temp.SetNumUninitialized(OutValues.Num());
+		GetValues_Internal(ValueKeys, FOutValues{TInPlaceType<FOutValuesAsArray>{}, Temp});
+		for (int32 i = 0; i < OutValues.Num(); ++i)
+		{
+			OutValues[i] = MakeConstArrayView<T>(static_cast<const T*>(Temp[i].Get<0>()), Temp[i].Get<1>());
+		}
+	}
+	
+	// To support sets
+	template <typename T> 
+	void GetValues(TArrayView<const PCGMetadataValueKey> ValueKeys, TArrayView<PCG::TScriptSetWrapper<T>> OutValues) const
+	{
+		if (!PCG::Private::IsEquivalentDesc<TSet<T>>(CachedDesc))
+		{
+			// @todo_pcg Better logging
+			UE_LOG(LogPCG, Error, TEXT("Type mismatch on GetValues"));
+			return;
+		}
+		
+		if (!ensure(ValueKeys.Num() == OutValues.Num()))
+		{
+			return;
+		}
+		
+		TArray<FScriptSetHelper*, TInlineAllocator<256>> Temp;
+		Temp.Reserve(OutValues.Num());
+		for (PCG::TScriptSetWrapper<T>& Value : OutValues)
+		{
+			Temp.Add(&Value.Helper);
+		}
+		GetValues_Internal(ValueKeys, FOutValues{TInPlaceType<FOutValuesAsSet>{}, Temp});
+	}
+	
+	// To support maps
+	template <typename KeyType, typename ValueType> 
+	void GetValues(TArrayView<const PCGMetadataValueKey> ValueKeys, TArrayView<PCG::TScriptMapWrapper<KeyType, ValueType>> OutValues) const
+	{
+		if (!PCG::Private::IsEquivalentDesc<TMap<KeyType, ValueType>>(CachedDesc))
+		{
+			// @todo_pcg Better logging
+			UE_LOG(LogPCG, Error, TEXT("Type mismatch on GetValues"));
+			return;
+		}
+		
+		if (!ensure(ValueKeys.Num() == OutValues.Num()))
+		{
+			return;
+		}
+		
+		TArray<FScriptMapHelper*, TInlineAllocator<256>> Temp;
+		Temp.Reserve(OutValues.Num());
+		for (PCG::TScriptMapWrapper<KeyType, ValueType>& Value : OutValues)
+		{
+			Temp.Add(&Value.Helper);
+		}
+		GetValues_Internal(ValueKeys, FOutValues{TInPlaceType<FOutValuesAsMap>{}, Temp});
+	}
+	
+	////////////// End GetValues Section
+	
+	// Fill the OutValues using the array of entry keys.
+	// cf. GetValues for more info on supported types.
+	template <typename T>
+	void GetValuesFromItemKeys(TConstArrayView<PCGMetadataEntryKey> EntryKeys, TArrayView<T> OutValues) const
+	{
+		GetValuesFromItemKeys(PCGValueRangeHelpers::MakeConstValueRange(EntryKeys), OutValues);
+	}
+	
+	// Fill the OutValues using the range of entry keys.
+	// cf. GetValues for more info on supported types.
+	template <typename T>
+	void GetValuesFromItemKeys(TConstPCGValueRange<PCGMetadataEntryKey> EntryKeys, TArrayView<T> OutValues) const
+	{
+		if (!ensure(EntryKeys.Num() == OutValues.Num()))
+		{
+			return;
+		}
+
+		TArray<PCGMetadataValueKey> ValueKeys;
+		GetValueKeys(EntryKeys, ValueKeys);
+		GetValues(ValueKeys, OutValues);
+	}
+
+	// Override the default value for this attribute.
+	// Will fail if the type mismatch.
+	template <typename T>
+	void SetDefaultValue(const T& Value)
+	{
+		if (TOptional<FScriptArrayHelper> Helper = MakeArrayHelper(); Helper.IsSet())
+		{
+			check(Helper->Num() > 0);
+			AddValues(MakeArrayView(&Value, 1), /*ExistingIndex=*/0);
+		}
+	}
+
+	template <typename T>
+	void SetValues(TArrayView<const PCGMetadataEntryKey * const> ItemKeys, TArrayView<const T> InValues)
+	{
+		SetValuesFromValueKeys(ItemKeys, AddValues(InValues));
+	}
+
+	template <typename T>
+	void SetValues(TArrayView<const PCGMetadataEntryKey> ItemKeys, TArrayView<const T> InValues)
+	{
+		SetValuesFromValueKeys(ItemKeys, AddValues(InValues));
+	}
+	
+	template <typename T>
+	void SetValues(TConstPCGValueRange<PCGMetadataEntryKey> ItemKeys, TArrayView<const T> InValues)
+	{
+		SetValuesFromValueKeys(ItemKeys, AddValues(InValues));
+	}
+	
+	template <typename T>
+	void SetValues(TPCGValueRange<PCGMetadataEntryKey> ItemKeys, TArrayView<const T> InValues)
+	{
+		SetValuesFromValueKeys(PCGValueRangeHelpers::MakeConstValueRange(ItemKeys), AddValues(InValues));
+	}
+
+	// Special version to write on pre-allocated data if we do not compress, in order to be able to write in concurrency.
+	template <typename T>
+	void SetValues_TryLockless(TArrayView<PCGMetadataEntryKey*> ItemKeys, const TArrayView<const T>& InValues, int32 StartIndex)
+	{
+		if (InValues.IsEmpty())
+		{
+			return;
+		}
+
+		if (UnderlyingProperty->bCompressData)
+		{
+			SetValues(ItemKeys, InValues);
+		}
+		else
+		{
+			check(StartIndex != INDEX_NONE && Values.IsValidIndex(StartIndex + InValues.Num() - 1));
+			
+			// Since values are already preallocated, we can take advantage of the AddValues with existing index.
+			AddValues(InValues, StartIndex, /*bLockless=*/true);
+		}
+	}
+
+	// Copy values from SrcPtr to the attribute. Property must match the attribute property, and the SrcPtr is expected to be
+	// reinterpreted as the property type.
+	PCG_API void SetValueFromProperty(PCGMetadataEntryKey ItemKey, const void* SrcPtr, const FProperty* Property);
+	
+	int32 PreallocateValues(TArrayView<PCGMetadataEntryKey*> EntryKeys, bool bLockless);
+
+	// Allocate the values (without initialization if it can) and return the first index of that new memory space.
+	// Return INDEX_NONE if it didn't work.
+	int32 AllocateValues_Unsafe(int32 Count);
+	
+	void Prepare(int32 Count)
+	{
+		// @todo_pcg To be implemented
+	}
+
+	// @todo_pcg: Temporary for debug. Build a struct that will contain the value at the entry key.
+	PCG_API FInstancedPropertyBag BuildStructForDebug(PCGMetadataEntryKey EntryKey) const;
+
+protected:
+	// Used to make an array helper of the local FScriptArray.
+	// Optional, since the attribute can be ill-formed, so the user should always check if it is set before
+	// using it.
+	TOptional<FScriptArrayHelper> MakeArrayHelper() const;
+	
+	// Utility function to get the pointer to a value using a value key. Return nullptr if invalid.
+	const void* GetReadAddressForValueKey_Unsafe(const PCGMetadataValueKey InValueKey) const;
+
+	// Utility function to get the pointer to a value using an entry key. Return nullptr if invalid.
+	const void* GetReadAddressFromEntryKey_Unsafe(const PCGMetadataEntryKey InEntryKey) const;
+	
+	/**
+	 * GetValues implementation section
+	 * GetValues_Internal is working on a Variant that will type erase the underlying output container
+	 * in order to reduce template code size.
+	 */
+	struct FOutValuesByValue
+	{
+		void* OutValues;
+	};
+	
+	struct FOutValuesByPtr
+	{
+		TArrayView<const void*> OutValues;
+	};
+	
+	struct FOutValuesAsArray
+	{
+		TArrayView<TTuple<const void*, int32>> OutValues;
+	};
+	
+	struct FOutValuesAsSet
+	{
+		TArrayView<FScriptSetHelper*> OutValues;
+	};
+	
+	struct FOutValuesAsMap
+	{
+		TArrayView<FScriptMapHelper*> OutValues;
+	};
+	
+	using FOutValues = TVariant<FOutValuesByValue, FOutValuesByPtr, FOutValuesAsArray, FOutValuesAsSet, FOutValuesAsMap>;
+	
+	void GetValues_Internal(const TArrayView<const PCGMetadataValueKey> ValueKeys, FOutValues OutValues) const
+	{
+		// Bitset with all unretrieved values. If we have any unretrieved value, we will ask the parent for those.
+		TBitArray<> UnretrievedValues(true, ValueKeys.Num());
+		
+		return GetValues_Internal(ValueKeys, MoveTemp(OutValues), UnretrievedValues);
+	}
+
+	/**
+	 * AddValues section
+	 * Values will be dispatched to a type erased variant to reduce template code size
+	 * Variant contains all supported containers: Values/Array/Set/Map
+	 * 
+	 * Note that to support setting the default value (and potentially write on existing values if it
+	 * was preallocated), all AddValues have an extra parameter `ExistingIndex` that will indicate if allocation
+	 * is necessary or not. This is a very niche feature, and it is expected that everything is already allocated
+	 * or it will crash. There is also an option for lockless, if we expect to be able to write without any lock.
+	 */
+	struct FInValuesByValue
+	{
+		const void* InValues;
+		int32 Count;
+	};
+
+	struct FInValuesByValueSubset
+	{
+		const void* InValues;
+		TArray<int32> Indices;
+	};
+
+	struct FInValuesAsArray
+	{
+		TArrayView<TTuple<const void*, int32>> InValues;
+	};
+
+	struct FInValuesAsSet
+	{
+		TArrayView<TArray<const void*>> InValues;
+	};
+
+	struct FInValuesAsMap
+	{
+		TArrayView<TArray<TPair<const void*, const void*>>> InValues;
+	};
+
+	using FInValues = TVariant<FInValuesByValue, FInValuesAsArray, FInValuesAsSet, FInValuesAsMap, FInValuesByValueSubset>;
+
+	/**
+	 * Basic version that will accept all single values (as an array view).
+	 */
+	template <typename T>
+	TArray<PCGMetadataValueKey> AddValues(TArrayView<const T> InValues, int32 ExistingIndex = INDEX_NONE, bool bLockless = false)
+	{
+		if (!PCG::Private::IsEquivalentDesc<T>(CachedDesc))
+		{
+			// @todo_pcg Better logging
+			UE_LOG(LogPCG, Error, TEXT("Type mismatch on AddValues"));
+			return {};
+		}
+
+		// For compressed data, we need to analyze which values actually need to be added. If existing index is not none, ignore it.
+		// Need the concept for Equals, since `AddUnique` require the operator '==` to be defined.
+		if constexpr (UE::Core::Private::CWithConstUEOpEquals<T, T>)
+		{
+			if (UnderlyingProperty->bCompressData && ExistingIndex == INDEX_NONE)
+			{
+				// Since we're getting raw values here, we might have duplicates
+				// so we should aim to remove duplicates here so we preserve our 'compress data' idea, otherwise it will break other foundational blocks (e.g. partition)
+				TArray<T, TInlineAllocator<256>> UniqueValues;
+
+				// Initially, fill with mapping to unique values so we can remap them at the end if needed
+				TArray<PCGMetadataValueKey> FoundValueKeys;
+				FoundValueKeys.Reserve(InValues.Num());
+
+				for (int ValueIndex = 0; ValueIndex < InValues.Num(); ++ValueIndex)
+				{
+					FoundValueKeys.Emplace(UniqueValues.AddUnique(InValues[ValueIndex]));
+				}
+
+				return AddCompressedValues(FInValuesByValue{UniqueValues.GetData(), UniqueValues.Num()}, FoundValueKeys);
+			}
+		}
+
+		return AddValues_Internal(FInValues{TInPlaceType<FInValuesByValue>{}, InValues.GetData(), InValues.Num()}, InValues.Num(), ExistingIndex, bLockless);
+	}
+	
+	/**
+	 * Version that accept anything that is array like, like `TArray<T>`, `TArrayView<T>`, `TStaticArray<T, 5>`, `TArray<T, TInlineAllocator<5>>`, etc...
+	 */
+	template <typename ArrayContainer> requires (PCG::Private::CIsArrayLike<ArrayContainer>)
+	TArray<PCGMetadataValueKey> AddValues(TConstArrayView<ArrayContainer> InValues, int32 ExistingIndex = INDEX_NONE, bool bLockless = false)
+	{
+		using T = typename ArrayContainer::ElementType;
+		if (!PCG::Private::IsEquivalentDesc<TArray<T>>(CachedDesc))
+		{
+			// @todo_pcg Better logging
+			UE_LOG(LogPCG, Error, TEXT("Type mismatch on AddValues"));
+			return {};
+		}
+		
+		TArray<TTuple<const void*, int32>, TInlineAllocator<256>> InValuesTuple;
+		InValuesTuple.Reserve(InValues.Num());
+		Algo::Transform(InValues, InValuesTuple, [](const auto& Item) { return TTuple<const void*, int32>{Item.GetData(), Item.Num()}; });
+		return AddValues_Internal(FInValues{TInPlaceType<FInValuesAsArray>{}, InValuesTuple}, InValues.Num(), ExistingIndex, bLockless);
+	}
+
+	/**
+	 * Version that accepts Sets
+	 */
+	template <typename T, typename ...Extra>
+	TArray<PCGMetadataValueKey> AddValues(TConstArrayView<TSet<T, Extra...>> InValues, int32 ExistingIndex = INDEX_NONE, bool bLockless = false)
+	{
+		if (!PCG::Private::IsEquivalentDesc<TSet<T>>(CachedDesc))
+		{
+			// @todo_pcg Better logging
+			UE_LOG(LogPCG, Error, TEXT("Type mismatch on AddValues"));
+			return {};
+		}
+		
+		TArray<TArray<const void*>, TInlineAllocator<256>> InSetValues;
+		InSetValues.Reserve(InValues.Num());
+		Algo::Transform(InValues, InSetValues, [](const auto& Item) -> TArray<const void*>
+			{
+				TArray<const void*> SetValues;
+				SetValues.Reserve(Item.Num());
+				Algo::Transform(Item, SetValues, [](const T& Value) -> const T* { return &Value;});
+				return SetValues;
+			});
+		return AddValues_Internal(FInValues{TInPlaceType<FInValuesAsSet>{}, InSetValues}, InValues.Num(), ExistingIndex, bLockless);
+	}
+	
+	/**
+	 * Version that accepts Maps
+	 */
+	template <typename KeyType, typename ValueType, typename ...Extra>
+	TArray<PCGMetadataValueKey> AddValues(TConstArrayView<TMap<KeyType, ValueType, Extra...>> InValues, int32 ExistingIndex = INDEX_NONE, bool bLockless = false)
+	{
+		if (!PCG::Private::IsEquivalentDesc<TMap<KeyType, ValueType>>(CachedDesc))
+		{
+			// @todo_pcg Better logging
+			UE_LOG(LogPCG, Error, TEXT("Type mismatch on AddValues"));
+			return {};
+		}
+		
+		TArray<TArray<TPair<const void*, const void*>>, TInlineAllocator<256>> InMapValues;
+		InMapValues.Reserve(InValues.Num());
+		Algo::Transform(InValues, InMapValues, [](const auto& Item) -> TArray<TPair<const void*, const void*>>
+			{
+				TArray<TPair<const void*, const void*>> MapValues;
+				MapValues.Reserve(Item.Num());
+				Algo::Transform(Item, MapValues, [](const auto& It) -> TPair<const void*, const void*> { return {&It.Key, &It.Value};});
+				return MapValues;
+			});
+		return AddValues_Internal(FInValues{TInPlaceType<FInValuesAsMap>{}, InMapValues}, InValues.Num(), ExistingIndex, bLockless);
+	}
+	
+	// Internal implementation that will do the dispatch using the typed erased variant.
+	TArray<PCGMetadataValueKey> AddValues_Internal(FInValues InValues, int32 Count, int32 ExistingIndex = INDEX_NONE, bool bLockless = false);
+
+	// Internal implementation for compressed data
+	TArray<PCGMetadataValueKey> AddCompressedValues(FInValuesByValue InUniqueValues, TArray<PCGMetadataValueKey>& FoundValueKeys);
+
+	/** 
+	 * Find all the value keys associated to InValues for compressed data. Return true if all the values were found.
+	 */
+	bool FindValues(TArray<PCGMetadataValueKey>& OutValueKeys, const FInValuesByValue& InValues);
+
+private:
+	void Init();
+	
+	void GetValues_Internal(const TArrayView<const PCGMetadataValueKey> ValueKeys, FOutValues OutValues, TBitArray<>& UnretrievedValues) const;
+	void FindValuesInternal(const FInValuesByValue& InValues, TArray<PCGMetadataValueKey>& ValueKeys, int& ValueKeysSet, bool bIsRoot) const;
+
+public:
+	virtual void Flatten() override;
+	virtual void FlattenAndCompress(const TArrayView<const PCGMetadataEntryKey>& InEntryKeysToKeep) override;
+	virtual void Reset() override;
+	virtual FPCGMetadataAttributeBase* Copy(FName NewName, FPCGMetadataDomain* InMetadata, bool bKeepParent, bool bCopyEntries = true,
+		bool bCopyValues = true) const override;
+
+	UE_DEPRECATED(5.6, "Use the version with FPCGMetadataDomain")
+	virtual FPCGMetadataAttributeBase* Copy(FName NewName, UPCGMetadata* InMetadata, bool bKeepParent, bool bCopyEntries = true,
+		bool bCopyValues = true) const override;
+
+	virtual FPCGMetadataAttributeBase* CopyToAnotherType(int16 Type) const override;
+	virtual void SetValue(PCGMetadataEntryKey ItemKey, const FPCGMetadataAttributeBase* InAttribute, PCGMetadataEntryKey InEntryKey) override;
+	virtual void SetZeroValue(PCGMetadataEntryKey ItemKey) override;
+	virtual void AccumulateValue(PCGMetadataEntryKey ItemKey, const FPCGMetadataAttributeBase* InAttributeA, PCGMetadataEntryKey InEntryKeyA,
+		float Weight) override;
+	virtual void SetWeightedValue(PCGMetadataEntryKey ItemKey, const FPCGMetadataAttributeBase* InAttribute,
+		const TArrayView<const TPair<PCGMetadataEntryKey, float>>& InWeightedKeys) override;
+	virtual void SetValue(PCGMetadataEntryKey ItemKey, const FPCGMetadataAttributeBase* InAttributeA, PCGMetadataEntryKey InEntryKeyA,
+		const FPCGMetadataAttributeBase* InAttributeB, PCGMetadataEntryKey InEntryKeyB, EPCGMetadataOp Op) override;
+	virtual bool IsEqualToDefaultValue(PCGMetadataValueKey ValueKey) const override;
+	virtual void SetDefaultValueToFirstEntry() override;
+	virtual bool UsesValueKeys() const override;
+	virtual bool AreValuesEqualForEntryKeys(PCGMetadataEntryKey EntryKey1, PCGMetadataEntryKey EntryKey2) const override;
+	virtual bool AreValuesEqual(PCGMetadataValueKey ValueKey1, PCGMetadataValueKey ValueKey2) const override;
+
+private:
+	FPCGMetadataAttributeDesc CachedDesc;
+	TSharedPtr<FPCGAttributeProperty> UnderlyingProperty;
+	FScriptArray Values;
+	
+	mutable PCG::FSharedLock ValueLock;
+	PCGMetadataValueKey ValueKeyOffset = 0;
+};
diff --git a/Engine/Plugins/PCG/Source/PCG/Public/Metadata/PCGMetadataAttributeTpl.h b/Engine/Plugins/PCG/Source/PCG/Public/Metadata/PCGMetadataAttributeTpl.h
--- a/Engine/Plugins/PCG/Source/PCG/Public/Metadata/PCGMetadataAttributeTpl.h
+++ b/Engine/Plugins/PCG/Source/PCG/Public/Metadata/PCGMetadataAttributeTpl.h
@@ -600,6 +600,16 @@ public:
 	{
 		SetValuesFromValueKeys(ItemKeys, AddValues(InValues));
 	}
+	
+	void SetValues(TConstPCGValueRange<PCGMetadataEntryKey> ItemKeys, const TArrayView<const T>& InValues)
+	{
+		SetValuesFromValueKeys(ItemKeys, AddValues(InValues));
+	}
+	
+	void SetValues(TPCGValueRange<PCGMetadataEntryKey> ItemKeys, const TArrayView<const T>& InValues)
+	{
+		SetValuesFromValueKeys(PCGValueRangeHelpers::MakeConstValueRange(ItemKeys), AddValues(InValues));
+	}
 
 	template<typename U>
 	void SetValue(PCGMetadataEntryKey ItemKey, const U& InValue)
@@ -708,7 +718,12 @@ public:
 	* if we ever have to go check the parent attribute, as we need to modify the entry keys for that.
 	* If you don't care if the Entry Keys are modified, use the non-const version of the EntryKeys.
 	*/
-	void GetValuesFromItemKeys(const TArrayView<const PCGMetadataEntryKey> EntryKeys, TArrayView<T> OutValues) const
+	void GetValuesFromItemKeys(TConstArrayView<PCGMetadataEntryKey> EntryKeys, TArrayView<T> OutValues) const
+	{
+		GetValuesFromItemKeys(PCGValueRangeHelpers::MakeConstValueRange(EntryKeys), OutValues);
+	}
+	
+	void GetValuesFromItemKeys(TConstPCGValueRange<PCGMetadataEntryKey> EntryKeys, TArrayView<T> OutValues) const
 	{
 		if (!ensure(EntryKeys.Num() == OutValues.Num()))
 		{
diff --git a/Engine/Plugins/PCG/Source/PCG/Public/Metadata/PCGMetadataAttributeTraits.h b/Engine/Plugins/PCG/Source/PCG/Public/Metadata/PCGMetadataAttributeTraits.h
--- a/Engine/Plugins/PCG/Source/PCG/Public/Metadata/PCGMetadataAttributeTraits.h
+++ b/Engine/Plugins/PCG/Source/PCG/Public/Metadata/PCGMetadataAttributeTraits.h
@@ -2,37 +2,13 @@
 
 #pragma once
 
+#include "Metadata/PCGMetadataCommon.h"
+
 #include "Math/Transform.h" // IWYU pragma: keep
 #include "Misc/DefaultValueHelper.h"
 #include "UObject/Class.h"
 #include "UObject/SoftObjectPath.h"
 
-#include "PCGMetadataAttributeTraits.generated.h"
-
-UENUM(BlueprintType)
-enum class EPCGMetadataTypes : uint8
-{
-	Float = 0,
-	Double,
-	Integer32,
-	Integer64,
-	Vector2,
-	Vector,
-	Vector4,
-	Quaternion,
-	Transform,
-	String,
-	Boolean,
-	Rotator,
-	Name,
-	SoftObjectPath,
-	SoftClassPath,
-
-	Count UMETA(Hidden),
-
-	Unknown = 255 UMETA(Hidden),
-};
-
 // Convenient macro to avoid duplicating a lot of code with all our supported types.
 #define PCG_FOREACH_SUPPORTEDTYPES(MACRO) \
 	MACRO(int32)           \
@@ -55,13 +31,89 @@ namespace PCG
 {
 	namespace Private
 	{
-		template<typename T>
+		// Utility structure to be able to extract the types of Key and Values for a map.
+		template <typename Container>
+		struct TMapTypeExtractor
+		{
+		};
+		
+		template <typename InKeyType, typename InValueType>
+		struct TMapTypeExtractor<TMap<InKeyType, InValueType>>
+		{
+			using KeyType = InKeyType;
+			using ValueType = InValueType;
+		};
+		
+		// Utility struct to know when a type is a basic type, specialized in the macro.
+		template <typename T> struct TIsBasicType { enum { Value = false }; };
+		
+		// Concept to catch all "basic types" that are also structs.
+		template <typename T>
+		concept CIsBasicStruct = requires (T)
+		{
+			TBaseStructure<T>::Get();
+		};
+		
+		// Old way of converting types to the type enum, will be progressively replaced by MakeAttributeDesc/IsSameType.
+		template <typename T>
 		struct MetadataTypes
 		{
 			enum { Id = static_cast<uint16>(EPCGMetadataTypes::Unknown) + sizeof(T) };
 		};
 
-#define PCGMetadataGenerateDataTypes(Type, TypeEnum) template<> struct MetadataTypes<Type>{ enum { Id = static_cast<uint16>(EPCGMetadataTypes::TypeEnum)}; }
+		/**
+		 * Return the attribute descriptor associated with the type T, with the Attribute name provided.
+		 * This version does not compile and is a pitfall for all unsupported types.
+		 */
+		template <typename T>
+		FPCGMetadataAttributeDesc MakeAttributeDesc(FName AttributeName)
+		{
+			static_assert(!std::is_same_v<T,T>, "This type is not supported for metadata."); 
+			return FPCGMetadataAttributeDesc{};
+		}
+		
+		/**
+		 * Return true if the provided Desc is matching with type T.
+		 * This version does not compile and is a pitfall for all unsupported types.
+		 */
+		template <typename T>
+		bool IsEquivalentDesc(const FPCGMetadataAttributeDesc& Desc)
+		{
+			static_assert(!std::is_same_v<T,T>, "This type is not supported for metadata."); 
+			return false;
+		}
+		
+		/**
+		 * Return true if the provided type and type object is matching with type T.
+		 * This version does not compile for non-basic types and is a pitfall for all unsupported types.
+		 */
+		template <typename T>
+		bool IsSameType(const EPCGMetadataTypes InType, const UObject* InTypeObject)
+		{
+			if constexpr (!TIsBasicType<T>::Value)
+			{
+				static_assert(!std::is_same_v<T,T>, "This type is not supported for metadata."); 
+				return false;
+			}
+			else if constexpr (CIsBasicStruct<T>) 
+			{
+				// Basic structures can be represented with their own EPCGMetadataTypes (like Vector) or as a struct.
+				return static_cast<uint16>(InType) == MetadataTypes<T>::Id || (InType == EPCGMetadataTypes::Struct && InTypeObject == TBaseStructure<T>::Get());
+			}
+			else
+			{
+				return static_cast<uint16>(InType) == MetadataTypes<T>::Id;
+			}
+		}
+
+		/**
+		 * Macro to specialize all known "basic" types.
+		 */
+#define PCGMetadataGenerateDataTypes(Type, TypeEnum) \
+	template<> struct MetadataTypes<Type>{ enum { Id = static_cast<uint16>(EPCGMetadataTypes::TypeEnum)}; }; \
+	template<> inline FPCGMetadataAttributeDesc MakeAttributeDesc<Type>(FName AttributeName) { return FPCGMetadataAttributeDesc{.Name=AttributeName, .ValueType=EPCGMetadataTypes::TypeEnum}; } \
+	template<> struct TIsBasicType<Type> { enum { Value = true }; }; \
+	template<> inline bool IsEquivalentDesc<Type>(const FPCGMetadataAttributeDesc& Desc) { return IsSameType<Type>(Desc.ValueType, Desc.ValueTypeObject) && Desc.ContainerTypes.IsEmpty(); }
 
 		PCGMetadataGenerateDataTypes(float, Float);
 		PCGMetadataGenerateDataTypes(double, Double);
@@ -78,9 +130,131 @@ namespace PCG
 		PCGMetadataGenerateDataTypes(FName, Name);
 		PCGMetadataGenerateDataTypes(FSoftObjectPath, SoftObjectPath);
 		PCGMetadataGenerateDataTypes(FSoftClassPath, SoftClassPath);
+		PCGMetadataGenerateDataTypes(FText, Text);
+		PCGMetadataGenerateDataTypes(uint8, Byte);
 
 #undef PCGMetadataGenerateDataTypes
 
+		/**
+		 * Specialization section of MakeAttributeDesc
+		 * Supported generic types:
+		 *   - UStruct
+		 *   - UEnum
+		 *   - TArray<T> if T is a supported type. Does not support nested containers for now.
+		 *   - TSet<T> if T is a supported type. Does not support nested containers for now.
+		 *   - TMap<T, U> if T and U are supported types. Does not support nested containers for now.
+		 */
+		template<typename T> requires (TModels_V<CStaticStructProvider, T>)
+		FPCGMetadataAttributeDesc MakeAttributeDesc(FName AttributeName)
+		{
+			return FPCGMetadataAttributeDesc{
+				.Name = AttributeName,
+				.ValueType = EPCGMetadataTypes::Struct,
+				.ValueTypeObject = T::StaticStruct()
+			};
+		}
+		
+		template<typename T> requires (TIsUEnumClass<T>::Value == true)
+		FPCGMetadataAttributeDesc MakeAttributeDesc(FName AttributeName)
+		{
+			return FPCGMetadataAttributeDesc{
+				.Name = AttributeName,
+				.ValueType = EPCGMetadataTypes::Enum,
+				.ValueTypeObject = StaticEnum<T>()
+			};
+		}
+		
+		template <typename Container> requires TIsTArray_V<Container>
+		FPCGMetadataAttributeDesc MakeAttributeDesc(FName AttributeName)
+		{
+			using T = typename Container::ElementType;
+			FPCGMetadataAttributeDesc Result = MakeAttributeDesc<T>(AttributeName);
+			Result.ContainerTypes.Add(EPCGMetadataAttributeContainerTypes::Array);
+			return Result;
+		}
+		
+		template <typename Container> requires (TIsTSet<Container>::Value == true)
+		FPCGMetadataAttributeDesc MakeAttributeDesc(FName AttributeName)
+		{
+			using T = typename Container::ElementType;
+			FPCGMetadataAttributeDesc Result = MakeAttributeDesc<T>(AttributeName);
+			Result.ContainerTypes.Add(EPCGMetadataAttributeContainerTypes::Set);
+			return Result;
+		}
+		
+		template<typename Container> requires (TIsTMap<Container>::Value == true)
+		FPCGMetadataAttributeDesc MakeAttributeDesc(FName AttributeName)
+		{
+			using KeyType = typename TMapTypeExtractor<Container>::KeyType;
+			using ValueType = typename TMapTypeExtractor<Container>::ValueType;
+			
+			FPCGMetadataAttributeDesc Result = MakeAttributeDesc<ValueType>(AttributeName);
+			FPCGMetadataAttributeDesc Key = MakeAttributeDesc<KeyType>(AttributeName);
+			
+			Result.ContainerTypes.Add(EPCGMetadataAttributeContainerTypes::Map);
+			Result.KeyType = Key.ValueType;
+			Result.KeyTypeObject = Key.ValueTypeObject;
+			return Result;
+		}
+		
+		//////////////////////////////
+		
+		template<typename T> requires (TModels_V<CStaticStructProvider, T>)
+		bool IsSameType(const EPCGMetadataTypes InType, const UObject* InTypeObject)
+		{
+			return InType == EPCGMetadataTypes::Struct &&
+				InTypeObject == T::StaticStruct();
+		}
+		
+		template<typename T> requires (TIsUEnumClass<T>::Value == true)
+		bool IsSameType(const EPCGMetadataTypes InType, const UObject* InTypeObject)
+		{
+			return InType == EPCGMetadataTypes::Enum &&
+				InTypeObject == StaticEnum<T>();
+		}
+		
+		//////////////////////////////
+
+		template<typename T> requires (TModels_V<CStaticStructProvider, T> || TIsUEnumClass<T>::Value == true)
+		bool IsEquivalentDesc(const FPCGMetadataAttributeDesc& Desc)
+		{
+			return Desc.ContainerTypes.IsEmpty() 
+				&& IsSameType<T>(Desc.ValueType, Desc.ValueTypeObject);
+		}
+		
+		template <typename Container> requires TIsTArray_V<Container>
+		bool IsEquivalentDesc(const FPCGMetadataAttributeDesc& Desc)
+		{
+			using T = typename Container::ElementType;
+			return !Desc.ContainerTypes.IsEmpty() 
+				&& Desc.ContainerTypes[0] == EPCGMetadataAttributeContainerTypes::Array 
+				&& IsSameType<T>(Desc.ValueType, Desc.ValueTypeObject);
+		}
+		
+		template <typename Container> requires (TIsTSet<Container>::Value == true)
+		bool IsEquivalentDesc(const FPCGMetadataAttributeDesc& Desc)
+		{
+			using T = typename Container::ElementType;
+			return !Desc.ContainerTypes.IsEmpty() 
+				&& Desc.ContainerTypes[0] == EPCGMetadataAttributeContainerTypes::Set 
+				&& IsSameType<T>(Desc.ValueType, Desc.ValueTypeObject);
+
+		}
+		
+		template<typename Container> requires (TIsTMap<Container>::Value == true)
+		bool IsEquivalentDesc(const FPCGMetadataAttributeDesc& Desc)
+		{
+			using KeyType = typename TMapTypeExtractor<Container>::KeyType;
+			using ValueType = typename TMapTypeExtractor<Container>::ValueType;
+			
+			return !Desc.ContainerTypes.IsEmpty() 
+				&& Desc.ContainerTypes[0] == EPCGMetadataAttributeContainerTypes::Map 
+				&& IsSameType<KeyType>(Desc.KeyType, Desc.KeyTypeObject)
+				&& IsSameType<ValueType>(Desc.ValueType, Desc.ValueTypeObject);
+		}
+		
+		////////////////////////
+
 		constexpr inline bool IsPCGType(uint16 TypeId)
 		{
 			return TypeId < (uint16)EPCGMetadataTypes::Count;
diff --git a/Engine/Plugins/PCG/Source/PCG/Public/Metadata/PCGMetadataCommon.h b/Engine/Plugins/PCG/Source/PCG/Public/Metadata/PCGMetadataCommon.h
--- a/Engine/Plugins/PCG/Source/PCG/Public/Metadata/PCGMetadataCommon.h
+++ b/Engine/Plugins/PCG/Source/PCG/Public/Metadata/PCGMetadataCommon.h
@@ -2,7 +2,6 @@
 
 #pragma once
 
-
 #include "PCGMetadataCommon.generated.h"
 
 typedef int64 PCGMetadataEntryKey;
@@ -217,3 +216,79 @@ struct FPCGAttributeIdentifier
 	FPCGMetadataDomainID MetadataDomain;
 };
 
+UENUM(BlueprintType)
+enum class EPCGMetadataTypes : uint8
+{
+	Float = 0,
+	Double,
+	Integer32,
+	Integer64,
+	Vector2,
+	Vector,
+	Vector4,
+	Quaternion,
+	Transform,
+	String,
+	Boolean,
+	Rotator,
+	Name,
+	SoftObjectPath,
+	SoftClassPath,
+	
+	EndLegacyTypes UMETA(Hidden),
+	
+	// Hide all the new types for now, so they are not selectable in the UI
+	Byte UMETA(Hidden),
+	Text UMETA(Hidden),
+	Enum UMETA(Hidden),
+	Struct UMETA(Hidden),
+	Object UMETA(Hidden),
+	SoftObject UMETA(Hidden),
+	Class UMETA(Hidden),
+	SoftClass UMETA(Hidden),
+
+	Count UMETA(Hidden),
+
+	Unknown = 255 UMETA(Hidden),
+};
+
+UENUM()
+enum class EPCGMetadataAttributeContainerTypes
+{
+	None,
+	Array,
+	Set,
+	Map
+};
+
+USTRUCT()
+struct FPCGMetadataAttributeDesc
+{
+	GENERATED_BODY();
+	
+	UPROPERTY()
+	FName Name;
+	
+	UPROPERTY()
+	EPCGMetadataTypes ValueType = EPCGMetadataTypes::Unknown;
+	
+	UPROPERTY()
+	TArray<EPCGMetadataAttributeContainerTypes> ContainerTypes;
+		
+	// Struct/Object/Enum Specific
+	UPROPERTY()
+	TObjectPtr<const UObject> ValueTypeObject = nullptr;
+		
+	// Map specific
+	UPROPERTY()
+	EPCGMetadataTypes KeyType = EPCGMetadataTypes::Unknown;
+	
+	UPROPERTY()
+	TObjectPtr<const UObject> KeyTypeObject = nullptr;
+	
+	PCG_API bool operator==(const FPCGMetadataAttributeDesc& Other) const;
+	PCG_API friend int32 GetTypeHash(const FPCGMetadataAttributeDesc& Desc);
+	
+	UE_DEPRECATED(all, "Should not be used anymore as some information is lost in the process.")
+	PCG_API int16 AsLegacyTypeId() const;
+};
diff --git a/Engine/Plugins/PCG/Source/PCG/Public/Metadata/PCGMetadataContainerTypes.h b/Engine/Plugins/PCG/Source/PCG/Public/Metadata/PCGMetadataContainerTypes.h
new file mode 100644
--- /dev/null
+++ b/Engine/Plugins/PCG/Source/PCG/Public/Metadata/PCGMetadataContainerTypes.h
@@ -0,0 +1,185 @@
+﻿// Copyright Epic Games, Inc. All Rights Reserved.
+
+#pragma once
+
+#include "UObject/UnrealType.h"
+
+class FPCGMetadataAttributeGeneric;
+
+namespace PCG
+{
+/**
+ * Class meant to wrap around a FScriptSetHelper, to manipulate it like a TSet and have type checking.
+ * It is not meant to be initialized by the user, but by the generic attributes (hence why only the default ctor is exposed).
+ * Since it contains a FScriptSetHelper, which contains a FScriptSet pointer, lifetime is not guaranteed. 
+ * Can convert to a TSet, but this operation will do a full copy, so it might be expensive.
+ * Note that it is also read-only, so support only const operations.
+ */
+template <typename T>
+class TScriptSetWrapper
+{
+public:
+	// HACK: We need to use a dummy property so the helper can be default constructed.
+	// We always verify that the helper is valid before using it, so in theory this is safe.
+	// We also don't care for unicity across dlls, since validity is not checking the dummy property, so static inline is fine.
+	static inline FSetProperty Dummy{nullptr, NAME_None, RF_NoFlags};
+	
+	TScriptSetWrapper()
+		: Helper(&Dummy, nullptr, FScriptSet::GetScriptLayout(0, 1))
+	{}
+	
+	bool IsValid() const { return Helper.Set != nullptr; }
+	
+	const FScriptSetHelper* GetHelper() const { return IsValid() ? &Helper : nullptr; }
+	
+	bool Contains(const T& Value) const
+	{
+		return IsValid() && Helper.FindElementIndexFromHash(&Value) != INDEX_NONE;
+	}
+	
+	int32 Num() const { return IsValid() ? Helper.Num() : 0;}
+	
+	TSet<T> Intersection(const TSet<T>& Other) const
+	{
+		if (!IsValid())
+		{
+			return {};
+		}
+		
+		TSet<T> Result;
+		for (const T& Value : Other)
+		{
+			if (Helper.FindElementIndexFromHash(&Value) != INDEX_NONE)
+			{
+				Result.Add(Value);
+			}
+		}
+		
+		return Result;
+	}
+	
+	TSet<T> Union(const TSet<T>& Other) const
+	{
+		if (!IsValid())
+		{
+			return Other;
+		}
+		
+		TSet<T> Result = Other;
+		Result.Reserve(Other.Num() + Num());
+		ForEach([&Result](const T& Value) { Result.Add(Value); });
+		
+		return Result;
+	}
+	
+	explicit operator TSet<T>() const
+	{
+		return Union(TSet<T>{});
+	}
+	
+	template <typename FuncType>
+	void ForEach(FuncType&& Func) const
+	{
+		if (!IsValid())
+		{
+			return;
+		}
+		
+		for (FScriptSetHelper::FIterator It = Helper.CreateIterator(); It; ++It)
+		{
+			Func(*reinterpret_cast<const T*>(Helper.GetElementPtr(It)));
+		}
+	}
+	
+private:
+	friend FPCGMetadataAttributeGeneric;
+	FScriptSetHelper Helper;
+};
+	
+/**
+ * Class meant to wrap around a FScriptMapHelper, to manipulate it like a TMap and have type checking.
+ * It is not meant to be initialized by the user, but by the generic attributes (hence why only the default ctor is exposed).
+ * Since it contains a FScriptMapHelper, which contains a FScriptMap pointer, lifetime is not guaranteed. 
+ * Can convert to a TMap, but this operation will do a full copy, so it might be expensive.
+ * Note that it is also read-only, so support only const operations.
+ */
+template <typename KeyType, typename ValueType>
+class TScriptMapWrapper
+{
+public:
+	// HACK: We need to use dummy properties so the helper can be default constructed.
+	// We always verify that the helper is valid before using it, so in theory this is safe.
+	// We also don't care for unicity across dlls, since validity is not checking the dummy property, so static inline is fine.
+	static inline FBoolProperty Dummy{nullptr, NAME_None, RF_NoFlags};
+	
+	TScriptMapWrapper()
+		: Helper(&Dummy, &Dummy, nullptr, FScriptMap::GetScriptLayout(0, 1, 0, 1), EMapPropertyFlags::None)
+	{}
+	
+	bool IsValid() const { return Helper.HeapMap != nullptr; }
+	
+	const FScriptSetHelper* GetHelper() const { return IsValid() ? &Helper : nullptr; }
+	
+	bool Contains(const KeyType& Key) const
+	{
+		return IsValid() && Helper.FindMapIndexWithKey(&Key) != INDEX_NONE;
+	}
+	
+	const ValueType* Find(const KeyType& Key) const
+	{
+		if (!IsValid())
+		{
+			return nullptr;
+		}
+		
+		const int32 Index = Helper.FindMapIndexWithKey(&Key);
+		if (Index == INDEX_NONE)
+		{
+			return nullptr;
+		}
+		
+		return reinterpret_cast<const ValueType*>(const_cast<FScriptMapHelper&>(Helper).GetValuePtr(Index));
+	}
+	
+	const ValueType& operator[](const KeyType& Key) const
+	{
+		const ValueType* ValuePtr = Find(Key);
+		check(ValuePtr);
+		return *ValuePtr;
+	}
+	
+	int32 Num() const { return IsValid() ? Helper.Num() : 0;}
+	
+	explicit operator TMap<KeyType, ValueType>() const
+	{
+		if (!IsValid())
+		{
+			return {};
+		}
+		
+		TMap<KeyType, ValueType> Result;
+		Result.Reserve(Helper.Num());
+		ForEach([&Result](const KeyType& Key, const ValueType& Value) { Result.Emplace(Key, Value); });
+		
+		return Result;
+	}
+	
+	template <typename FuncType>
+	void ForEach(FuncType&& Func) const
+	{
+		if (!IsValid())
+		{
+			return;
+		}
+		
+		for (FScriptMapHelper::FIterator It = Helper.CreateIterator(); It; ++It)
+		{
+			Func(*reinterpret_cast<const KeyType*>(Helper.GetKeyPtr(It)), *reinterpret_cast<const ValueType*>(Helper.GetValuePtr(It)));
+		}
+	}
+	
+private:
+	friend FPCGMetadataAttributeGeneric;
+	FScriptMapHelper Helper;
+};
+}
diff --git a/Engine/Plugins/PCG/Source/PCG/Public/Metadata/PCGMetadataDomain.h b/Engine/Plugins/PCG/Source/PCG/Public/Metadata/PCGMetadataDomain.h
--- a/Engine/Plugins/PCG/Source/PCG/Public/Metadata/PCGMetadataDomain.h
+++ b/Engine/Plugins/PCG/Source/PCG/Public/Metadata/PCGMetadataDomain.h
@@ -3,6 +3,7 @@
 #pragma once
 
 #include "PCGCommon.h"
+#include "PCGMetadataAttributeGeneric.h"
 #include "Metadata/PCGMetadataCommon.h"
 #include "Metadata/PCGMetadataAttributeTpl.h"
 #include "Utils/PCGValueRange.h"
@@ -229,6 +230,18 @@ public:
 	*/
 	template<typename T>
 	FPCGMetadataAttribute<T>* CreateAttribute(FName AttributeName, const T& DefaultValue, bool bAllowsInterpolation, bool bOverrideParent);
+	
+	// @todo_pcg TO BE REMOVED
+	/**
+	* Create a new generic attribute.
+	*/
+	template<typename T>
+	FPCGMetadataAttributeGeneric* CreateGenericAttribute(FName AttributeName, const T& DefaultValue, bool bAllowsInterpolation, bool bOverrideParent);
+	
+	UE_API FPCGMetadataAttributeGeneric* CreateGenericAttribute(const FPCGMetadataAttributeDesc& Desc, bool bAllowsInterpolation, bool bOverrideParent);
+
+	UE_API FPCGMetadataAttributeGeneric* GetMutableGenericAttribute(FName AttributeName);
+	UE_API const FPCGMetadataAttributeGeneric* GetConstGenericAttribute(FName AttributeName) const;
 
 	/**
 	* Find or create an attribute. Follows CreateAttribute signature.
@@ -330,6 +343,18 @@ FPCGMetadataAttribute<T>* FPCGMetadataDomain::CreateAttribute(FName AttributeNam
 	return NewAttribute;
 }
 
+template<typename T>
+FPCGMetadataAttributeGeneric* FPCGMetadataDomain::CreateGenericAttribute(FName AttributeName, const T& DefaultValue, bool bAllowsInterpolation, bool bOverrideParent)
+{
+	FPCGMetadataAttributeGeneric* Attribute = CreateGenericAttribute(PCG::Private::MakeAttributeDesc<T>(AttributeName), bAllowsInterpolation, bOverrideParent);
+	if (Attribute)
+	{
+		Attribute->SetDefaultValue<T>(DefaultValue);
+	}
+
+	return Attribute;
+}
+
 template<typename T>
 FPCGMetadataAttribute<T>* FPCGMetadataDomain::GetMutableTypedAttribute_Unsafe(FName AttributeName)
 {
diff --git a/Engine/Plugins/PCG/Tests/Private/Benchmarks/Metadata/PCGMetadataAttributeGenericBenchmarks.cpp b/Engine/Plugins/PCG/Tests/Private/Benchmarks/Metadata/PCGMetadataAttributeGenericBenchmarks.cpp
new file mode 100644
--- /dev/null
+++ b/Engine/Plugins/PCG/Tests/Private/Benchmarks/Metadata/PCGMetadataAttributeGenericBenchmarks.cpp
@@ -0,0 +1,281 @@
+﻿// Copyright Epic Games, Inc. All Rights Reserved.
+
+#include "PCGTestsCommon.h"
+
+#include "Data/PCGPointArrayData.h"
+#include "Data/PCGPointData.h"
+
+#include "Containers/StaticArray.h"
+
+#include <catch2/catch_test_macros.hpp>
+#include <catch2/benchmark/catch_benchmark.hpp>
+#include <catch2/generators/catch_generators_range.hpp>
+
+#include "TestHarness.h"
+
+TEST_CASE_METHOD(PCGTests::FPCGBaseTestWithBasicContext, "PCG::AttributeGeneric::Benchmark", "[PCG][AttributeGeneric][!benchmark]")
+{
+	constexpr int32 Seed = 42;
+	static const FName DoubleAttributeName = "DoubleAttr";
+	static const FName DoubleGenericAttributeName = "DoubleGenericAttr";
+	static const FName StringAttributeName = "StringAttr";
+	static const FName StringGenericAttributeName = "StringGenericAttr";
+
+	// Preparing the data
+	constexpr int32 NumPoints = 4096;
+	constexpr int32 ChunkSize = 256;
+
+	SECTION("Double Attribute")
+	{
+		UPCGPointArrayData* PointArrayData = CreateData<UPCGPointArrayData>();
+
+		PointArrayData->SetNumPoints(NumPoints);
+
+		PointArrayData->AllocateProperties(EPCGPointNativeProperties::MetadataEntry);
+		TPCGValueRange<int64> PointArrayDataMetadataEntryRange = PointArrayData->GetMetadataEntryValueRange(/*bAllocate=*/false);
+
+		FPCGMetadataDomain* Domain = PointArrayData->MutableMetadata()->GetDefaultMetadataDomain();
+		FPCGMetadataAttributeGeneric* GenericAttribute = Domain->CreateGenericAttribute<double>(DoubleGenericAttributeName, 0., true, false);
+		FPCGMetadataAttribute<double>* TypedAttribute = Domain->CreateAttribute<double>(DoubleAttributeName, 0., true, false);
+		
+		// Allocate entry keys and values
+		TStaticArray<double, NumPoints> Values;
+		for (int32 i = 0; i < NumPoints; ++i)
+		{
+			PointArrayDataMetadataEntryRange[i] = Domain->AddEntry();
+			Values[i] = static_cast<double>(i);
+		}
+
+		TConstArrayView<PCGMetadataEntryKey> SingleKey = MakeConstArrayView(&PointArrayDataMetadataEntryRange[0], 1);
+		BENCHMARK_ADVANCED("[Generic] Set single value") (Catch::Benchmark::Chronometer meter)
+		{
+			// Make sure to reset everytime so we don't stack the values.
+			GenericAttribute->Reset();
+			meter.measure([&](){ GenericAttribute->SetValues<double>(SingleKey, {1.0});});
+		};
+		
+		BENCHMARK_ADVANCED("[Typed] Set single value") (Catch::Benchmark::Chronometer meter)
+		{
+			// Make sure to reset everytime so we don't stack the values.
+			TypedAttribute->Reset();
+			meter.measure([&](){ TypedAttribute->SetValues(SingleKey, {1.0});});
+		};
+		
+		BENCHMARK_ADVANCED("[Generic] Set full range") (Catch::Benchmark::Chronometer meter)
+		{
+			// Make sure to reset everytime so we don't stack the values.
+			GenericAttribute->Reset();
+			meter.measure([&](){ GenericAttribute->SetValues<double>(PointArrayDataMetadataEntryRange, Values); });
+		};
+		
+		BENCHMARK_ADVANCED("[Typed] Set full range") (Catch::Benchmark::Chronometer meter)
+		{
+			// Make sure to reset everytime so we don't stack the values.
+			TypedAttribute->Reset();
+			meter.measure([&](){ TypedAttribute->SetValues(PointArrayDataMetadataEntryRange, Values); });
+		};
+		
+		TStaticArray<double, NumPoints> ReadValues;
+		TStaticArray<const double*, NumPoints> ReadValuesPtr;
+		
+		BENCHMARK("[Generic] Get single value")
+		{
+			GenericAttribute->GetValuesFromItemKeys<double>(SingleKey, MakeArrayView(ReadValues.GetData(), 1));
+		};
+		
+		BENCHMARK("[Generic] Get single value as pointer")
+		{
+			GenericAttribute->GetValuesFromItemKeys<const double*>(SingleKey, MakeArrayView(ReadValuesPtr.GetData(), 1));
+		};
+		
+		BENCHMARK("[Typed] Get single value")
+		{
+			TypedAttribute->GetValuesFromItemKeys(SingleKey, MakeArrayView(ReadValues.GetData(), 1));
+		};
+		
+		BENCHMARK("[Generic] Get full range")
+		{
+			GenericAttribute->GetValuesFromItemKeys<double>(PCGValueRangeHelpers::MakeConstValueRange(PointArrayDataMetadataEntryRange), ReadValues);
+		};
+		
+		BENCHMARK("[Generic] Get full range as pointer")
+		{
+			GenericAttribute->GetValuesFromItemKeys<const double*>(PCGValueRangeHelpers::MakeConstValueRange(PointArrayDataMetadataEntryRange), ReadValuesPtr);
+		};
+		
+		BENCHMARK("[Typed] Get full range")
+		{
+			TypedAttribute->GetValuesFromItemKeys(PCGValueRangeHelpers::MakeConstValueRange(PointArrayDataMetadataEntryRange), ReadValues);
+		};
+		
+		// Finally test the worst case scenario where there is non-contiguous values.
+		UPCGPointArrayData* ChildPointArrayData = CreateData<UPCGPointArrayData>();
+		ChildPointArrayData->InitializeFromDataWithParams(FPCGInitializeFromDataParams{PointArrayData});
+		TPCGValueRange<int64> ChildPointArrayDataMetadataEntryRange = ChildPointArrayData->GetMetadataEntryValueRange(/*bAllocate=*/false);
+
+		FPCGMetadataDomain* ChildDomain = ChildPointArrayData->MutableMetadata()->GetDefaultMetadataDomain();
+		FPCGMetadataAttributeGeneric* ChildGenericAttribute = ChildDomain->GetMutableGenericAttribute(DoubleGenericAttributeName);
+		FPCGMetadataAttribute<double>* ChildTypedAttribute = ChildDomain->GetMutableTypedAttribute<double>(DoubleAttributeName);
+		
+		CHECK_NOT_EQUAL(ChildGenericAttribute, nullptr);
+		CHECK_NOT_EQUAL(ChildTypedAttribute, nullptr);
+		
+		// Allocate entry keys and values
+		TStaticArray<double, NumPoints/2> ChildValues;
+		TStaticArray<PCGMetadataEntryKey, NumPoints/2> ChildEntryKeys;
+		for (int32 i = 0; i < NumPoints / 2; ++i)
+		{
+			ChildDomain->InitializeOnSet(ChildPointArrayDataMetadataEntryRange[2 * i]);
+			ChildValues[i] = static_cast<double>(i + NumPoints * 2);
+			ChildEntryKeys[i] = ChildPointArrayDataMetadataEntryRange[2 * i];
+		}
+		
+		ChildGenericAttribute->SetValues<double>(ChildEntryKeys, ChildValues);
+		ChildTypedAttribute->SetValues(ChildEntryKeys, ChildValues);
+		
+		CHECK_EQUAL(ChildGenericAttribute->GetNumberOfEntries(), NumPoints / 2);
+		CHECK_EQUAL(ChildTypedAttribute->GetNumberOfEntries(), NumPoints / 2);
+		
+		BENCHMARK("[Generic] Get full range, with fragmented memory")
+		{
+			ChildGenericAttribute->GetValuesFromItemKeys<double>(PCGValueRangeHelpers::MakeConstValueRange(ChildPointArrayDataMetadataEntryRange), ReadValues);
+		};
+		
+		BENCHMARK("[Generic] Get full range as pointer, with fragmented memory")
+		{
+			ChildGenericAttribute->GetValuesFromItemKeys<const double*>(PCGValueRangeHelpers::MakeConstValueRange(ChildPointArrayDataMetadataEntryRange), ReadValuesPtr);
+		};
+		
+		BENCHMARK("[Typed] Get full range, with fragmented memory")
+		{
+			ChildTypedAttribute->GetValuesFromItemKeys(PCGValueRangeHelpers::MakeConstValueRange(ChildPointArrayDataMetadataEntryRange), ReadValues);
+		};
+	}
+	
+	SECTION("String Attribute")
+	{
+		UPCGPointArrayData* PointArrayData = CreateData<UPCGPointArrayData>();
+
+		PointArrayData->SetNumPoints(NumPoints);
+
+		PointArrayData->AllocateProperties(EPCGPointNativeProperties::MetadataEntry);
+		TPCGValueRange<int64> PointArrayDataMetadataEntryRange = PointArrayData->GetMetadataEntryValueRange(/*bAllocate=*/false);
+
+		FPCGMetadataDomain* Domain = PointArrayData->MutableMetadata()->GetDefaultMetadataDomain();
+		FPCGMetadataAttributeGeneric* GenericAttribute = Domain->CreateGenericAttribute<FString>(StringGenericAttributeName, FString{}, true, false);
+		FPCGMetadataAttribute<FString>* TypedAttribute = Domain->CreateAttribute<FString>(StringAttributeName, FString{}, true, false);
+
+		// Allocate entry keys and values
+		TStaticArray<FString, NumPoints> Values;
+		for (int32 i = 0; i < NumPoints; ++i)
+		{
+			PointArrayDataMetadataEntryRange[i] = Domain->AddEntry();
+			Values[i] = FString::Printf(TEXT("This is a string with some characters and a value: %d"), i % 25);
+		}
+
+		TConstArrayView<PCGMetadataEntryKey> SingleKey = MakeConstArrayView(&PointArrayDataMetadataEntryRange[0], 1);
+		BENCHMARK_ADVANCED("[Generic] Set single value") (Catch::Benchmark::Chronometer meter)
+		{
+			// Make sure to reset everytime so we don't stack the values.
+			GenericAttribute->Reset();
+			meter.measure([&](){ GenericAttribute->SetValues<FString>(SingleKey, {Values[0]});});
+		};
+		
+		BENCHMARK_ADVANCED("[Typed] Set single value") (Catch::Benchmark::Chronometer meter)
+		{
+			// Make sure to reset everytime so we don't stack the values.
+			TypedAttribute->Reset();
+			meter.measure([&](){ TypedAttribute->SetValues(SingleKey, {Values[0]});});
+		};
+		
+		BENCHMARK_ADVANCED("[Generic] Set full range") (Catch::Benchmark::Chronometer meter)
+		{
+			// Make sure to reset everytime so we don't stack the values.
+			GenericAttribute->Reset();
+			meter.measure([&](){ GenericAttribute->SetValues<FString>(PointArrayDataMetadataEntryRange, Values); });
+		};
+		
+		BENCHMARK_ADVANCED("[Typed] Set full range") (Catch::Benchmark::Chronometer meter)
+		{
+			// Make sure to reset everytime so we don't stack the values.
+			TypedAttribute->Reset();
+			meter.measure([&](){ TypedAttribute->SetValues(PointArrayDataMetadataEntryRange, Values); });
+		};
+		
+		TStaticArray<FString, NumPoints> ReadValues;
+		TStaticArray<const FString*, NumPoints> ReadValuesPtr;
+		
+		BENCHMARK("[Generic] Get single value")
+		{
+			GenericAttribute->GetValuesFromItemKeys<FString>(SingleKey, MakeArrayView(ReadValues.GetData(), 1));
+		};
+		
+		BENCHMARK("[Generic] Get single value pointer")
+		{
+			GenericAttribute->GetValuesFromItemKeys<const FString*>(SingleKey, MakeArrayView<const FString*>(ReadValuesPtr.GetData(), 1));
+		};
+		
+		BENCHMARK("[Typed] Get single value")
+		{
+			TypedAttribute->GetValuesFromItemKeys(SingleKey, MakeArrayView(ReadValues.GetData(), 1));
+		};
+		
+		BENCHMARK("[Generic] Get full range")
+		{
+			GenericAttribute->GetValuesFromItemKeys<FString>(PCGValueRangeHelpers::MakeConstValueRange(PointArrayDataMetadataEntryRange), ReadValues);
+		};
+		
+		BENCHMARK("[Generic] Get full range as pointer")
+		{
+			GenericAttribute->GetValuesFromItemKeys<const FString*>(PCGValueRangeHelpers::MakeConstValueRange(PointArrayDataMetadataEntryRange), ReadValuesPtr);
+		};
+		
+		BENCHMARK("[Typed] Get full range")
+		{
+			TypedAttribute->GetValuesFromItemKeys(PCGValueRangeHelpers::MakeConstValueRange(PointArrayDataMetadataEntryRange), ReadValues);
+		};
+		
+		// Finally test the worst case scenario where there is non-contiguous values.
+		UPCGPointArrayData* ChildPointArrayData = CreateData<UPCGPointArrayData>();
+		ChildPointArrayData->InitializeFromDataWithParams(FPCGInitializeFromDataParams{PointArrayData});
+		TPCGValueRange<int64> ChildPointArrayDataMetadataEntryRange = ChildPointArrayData->GetMetadataEntryValueRange(/*bAllocate=*/false);
+
+		FPCGMetadataDomain* ChildDomain = ChildPointArrayData->MutableMetadata()->GetDefaultMetadataDomain();
+		FPCGMetadataAttributeGeneric* ChildGenericAttribute = ChildDomain->GetMutableGenericAttribute(StringGenericAttributeName);
+		FPCGMetadataAttribute<FString>* ChildTypedAttribute = ChildDomain->GetMutableTypedAttribute<FString>(StringAttributeName);
+		
+		CHECK_NOT_EQUAL(ChildGenericAttribute, nullptr);
+		CHECK_NOT_EQUAL(ChildTypedAttribute, nullptr);
+		
+		// Allocate entry keys and values
+		TStaticArray<FString, NumPoints/2> ChildValues;
+		TStaticArray<PCGMetadataEntryKey, NumPoints/2> ChildEntryKeys;
+		for (int32 i = 0; i < NumPoints / 2; ++i)
+		{
+			ChildDomain->InitializeOnSet(ChildPointArrayDataMetadataEntryRange[2 * i]);
+			ChildValues[i] = FString::Printf(TEXT("This is a string with some characters and a new value: %d"), i % 25);
+			ChildEntryKeys[i] = ChildPointArrayDataMetadataEntryRange[2 * i];
+		}
+		
+		ChildGenericAttribute->SetValues<FString>(ChildEntryKeys, ChildValues);
+		ChildTypedAttribute->SetValues(ChildEntryKeys, ChildValues);
+		
+		CHECK_EQUAL(ChildGenericAttribute->GetNumberOfEntries(), NumPoints / 2);
+		CHECK_EQUAL(ChildTypedAttribute->GetNumberOfEntries(), NumPoints / 2);
+		
+		BENCHMARK("[Generic] Get full range, with fragmented memory")
+		{
+			ChildGenericAttribute->GetValuesFromItemKeys<FString>(PCGValueRangeHelpers::MakeConstValueRange(ChildPointArrayDataMetadataEntryRange), ReadValues);
+		};
+		
+		BENCHMARK("[Generic] Get full range as pointer, with fragmented memory")
+		{
+			ChildGenericAttribute->GetValuesFromItemKeys<const FString*>(PCGValueRangeHelpers::MakeConstValueRange(ChildPointArrayDataMetadataEntryRange), ReadValuesPtr);
+		};
+		
+		BENCHMARK("[Typed] Get full range, with fragmented memory")
+		{
+			ChildTypedAttribute->GetValuesFromItemKeys(PCGValueRangeHelpers::MakeConstValueRange(ChildPointArrayDataMetadataEntryRange), ReadValues);
+		};
+	}
+}
\ No newline at end of file
diff --git a/Engine/Plugins/PCG/Tests/Private/Metadata/PCGMetadataAttributeGenericTests.cpp b/Engine/Plugins/PCG/Tests/Private/Metadata/PCGMetadataAttributeGenericTests.cpp
new file mode 100644
--- /dev/null
+++ b/Engine/Plugins/PCG/Tests/Private/Metadata/PCGMetadataAttributeGenericTests.cpp
@@ -0,0 +1,897 @@
+// Copyright Epic Games, Inc. All Rights Reserved.
+
+#include "PCGTestsCommon.h"
+
+#include "PCGParamData.h"
+#include "Data/PCGPointArrayData.h"
+#include "Data/PCGPointData.h"
+
+#include "Containers/StaticArray.h"
+
+#include <catch2/catch_test_macros.hpp>
+#include <catch2/benchmark/catch_benchmark.hpp>
+#include <catch2/generators/catch_generators_range.hpp>
+
+#include "TestHarness.h"
+
+namespace PCGAttributeGenericTests
+{
+	/**
+	 * Base class for a tester that encapsulate everything that is type specific.
+	 */
+	template <typename T, typename InReadType = T>
+	struct TypedAttributeTester
+	{
+		using Type = T;
+		using ReadType = InReadType;
+
+		TypedAttributeTester(TFunction<void(FRandomStream&, T&)> InGenerateRandom, 
+			TFunction<bool(const T&, const ReadType&)> InVerify, 
+			T InDefaultValue,
+			FPCGMetadataAttributeDesc InExpectedDesc)
+			: GenerateRandom(MoveTemp(InGenerateRandom))
+			, Verify(MoveTemp(InVerify))
+			, DefaultValue(InDefaultValue)
+			, ExpectedDesc(InExpectedDesc)
+		{}
+
+		TFunction<void(FRandomStream&, T&)> GenerateRandom;
+		TFunction<bool(const T&, const ReadType&)> Verify;
+		T DefaultValue;
+		FPCGMetadataAttributeDesc ExpectedDesc;
+		TOptional<FPCGMetadataAttributeDesc> OverrideDesc;
+	};
+	
+	static const FName AttributeName = "Attr";
+	
+	/** Getter to create a tester for different types. */
+	TypedAttributeTester<float> FloatTester()
+	{
+		auto GenerateRandom = [](FRandomStream& RandomStream, float& OutValue) { OutValue = RandomStream.FRand(); };
+		
+		auto Verify = [](const float& LHS, const float& RHS) { return LHS == RHS; };
+		
+		float DefaultValue = 5.0f;
+		
+		FPCGMetadataAttributeDesc ExpectedDesc
+		{
+			.Name = AttributeName,
+			.ValueType = EPCGMetadataTypes::Float
+		};
+		
+		return TypedAttributeTester<float>(GenerateRandom, Verify, DefaultValue, MoveTemp(ExpectedDesc));
+	}
+	
+	TypedAttributeTester<double> DoubleTester()
+	{
+		auto GenerateRandom = [](FRandomStream& RandomStream, double& OutValue) { OutValue = RandomStream.FRand(); };
+		
+		auto Verify = [](const double& LHS, const double& RHS) { return LHS == RHS; };
+		
+		double DefaultValue = 5.0;
+		
+		FPCGMetadataAttributeDesc ExpectedDesc
+		{
+			.Name = AttributeName,
+			.ValueType = EPCGMetadataTypes::Double
+		};
+		
+		return TypedAttributeTester<double>(GenerateRandom, Verify, DefaultValue, MoveTemp(ExpectedDesc));
+	}
+	
+	TypedAttributeTester<FVector> VectorTester_Basic()
+	{
+		auto GenerateRandom = [](FRandomStream& RandomStream, FVector& OutValue) { OutValue = RandomStream.VRand(); };
+		
+		auto Verify = [](const FVector& LHS, const FVector& RHS) { return LHS == RHS; };
+		
+		FVector DefaultValue = FVector(1.0, 2.0, 3.0);
+		
+		FPCGMetadataAttributeDesc ExpectedDesc
+		{
+			.Name = AttributeName,
+			.ValueType = EPCGMetadataTypes::Vector
+		};
+			
+		return TypedAttributeTester<FVector>{GenerateRandom, Verify, DefaultValue, MoveTemp(ExpectedDesc)};
+	}
+		
+	TypedAttributeTester<FVector> VectorTester_Struct()
+	{
+		TypedAttributeTester<FVector> AttributeTester = VectorTester_Basic();
+		FPCGMetadataAttributeDesc StructDesc
+		{
+			.Name = AttributeName,
+			.ValueType = EPCGMetadataTypes::Struct,
+			.ValueTypeObject = TBaseStructure<FVector>::Get()
+		};
+		
+		AttributeTester.ExpectedDesc = StructDesc;
+		AttributeTester.OverrideDesc = StructDesc;
+	
+		return AttributeTester;
+	}
+	
+	TypedAttributeTester<FString, const FString*> StringTester()
+	{
+		auto GenerateRandom = [](FRandomStream& RandomStream, FString& OutValue)
+			{
+				static int i = 0; 
+				OutValue = FString::Printf(TEXT("This is a string with some characters and a value: %d"), i++ % 25);
+			};
+		
+		auto Verify = [](const FString& LHS, FString const* const& RHS) { return LHS == *RHS; };
+		
+		FString DefaultValue = TEXT("Hi");
+		
+		FPCGMetadataAttributeDesc ExpectedDesc
+		{
+			.Name = AttributeName,
+			.ValueType = EPCGMetadataTypes::String
+		};
+		
+		return TypedAttributeTester<FString, const FString*>(GenerateRandom, Verify, MoveTemp(DefaultValue), MoveTemp(ExpectedDesc));
+	}
+	
+	TypedAttributeTester<FPCGPoint, const FPCGPoint*> FPCGPointTester()
+	{		
+		auto GenerateRandom = [](FRandomStream& RandomStream, FPCGPoint& OutValue)
+			{
+				OutValue.Density = RandomStream.FRand();
+				OutValue.Transform.SetLocation(RandomStream.VRand());
+			};
+		
+		auto Verify = [](const FPCGPoint& LHS, FPCGPoint const* const& RHS)
+			{
+				return LHS.Density == RHS->Density && LHS.Transform.GetLocation() == RHS->Transform.GetLocation();
+			};
+		
+		FPCGPoint DefaultValue{FTransform{5.0 * FVector::OneVector}, 0.4f, 42};
+		
+		FPCGMetadataAttributeDesc ExpectedDesc
+		{
+			.Name = AttributeName,
+			.ValueType = EPCGMetadataTypes::Struct,
+			.ValueTypeObject = FPCGPoint::StaticStruct()
+		};
+		
+		return TypedAttributeTester<FPCGPoint, const FPCGPoint*>(GenerateRandom, Verify, MoveTemp(DefaultValue), MoveTemp(ExpectedDesc));
+	}
+	
+	TypedAttributeTester<TArray<double>, TConstArrayView<double>> ArrayDoubleTester()
+	{		
+		auto GenerateRandom = [](FRandomStream& RandomStream, TArray<double>& OutValue)
+		{
+			const int32 RandomCount = RandomStream.RandRange(1, 4);
+			for (int j = 0; j < RandomCount; ++j)
+			{
+				OutValue.Add(RandomStream.FRand());
+			}
+		};
+		
+		auto Verify = [](const TArray<double>& LHS, const TConstArrayView<double>& RHS)
+			{
+				if (LHS.Num() != RHS.Num())
+				{
+					return false;
+				}
+				
+				for (int i = 0; i < LHS.Num(); ++i)
+				{
+					if (LHS[i] != RHS[i])
+					{
+						return false;
+					}
+				}
+				
+				return true;
+			};
+		
+		TArray<double> DefaultValue = {1, 2, 3, 4, 5};
+		
+		FPCGMetadataAttributeDesc ExpectedDesc
+		{
+			.Name = AttributeName,
+			.ValueType = EPCGMetadataTypes::Double,
+			.ContainerTypes = {EPCGMetadataAttributeContainerTypes::Array}
+		};
+		
+		return TypedAttributeTester<TArray<double>, TConstArrayView<double>>(GenerateRandom, Verify, MoveTemp(DefaultValue), MoveTemp(ExpectedDesc));
+	}
+	
+	TypedAttributeTester<TArray<FPCGPoint>, TConstArrayView<FPCGPoint>> ArrayFPCGPointTester()
+	{	
+		auto GenerateRandom = [](FRandomStream& RandomStream, TArray<FPCGPoint>& OutValue)
+			{
+				const int32 RandomCount = RandomStream.RandRange(1, 4);
+				for (int j = 0; j < RandomCount; ++j)
+				{
+					FPCGPoint& Value = OutValue.Emplace_GetRef();
+					Value.Density = RandomStream.FRand();
+					Value.Transform.SetLocation(RandomStream.VRand());
+				}
+			};
+		
+		auto Verify = [](const TArray<FPCGPoint>& LHS, const TConstArrayView<FPCGPoint>& RHS)
+			{
+				if (LHS.Num() != RHS.Num())
+				{
+					return false;
+				}
+				
+				for (int i = 0; i < LHS.Num(); ++i)
+				{
+					if (LHS[i].Density != RHS[i].Density || LHS[i].Transform.GetLocation() != RHS[i].Transform.GetLocation())
+					{
+						return false;
+					}
+				}
+				
+				return true;
+			};
+		
+		TArray<FPCGPoint> DefaultValue = {
+			{FTransform{5.0 * FVector::OneVector}, 0.4f, 42},
+			{FTransform{6.0 * FVector::OneVector}, 0.5f, 43},
+			{FTransform{7.0 * FVector::OneVector}, 0.6f, 44},
+		};
+		
+		FPCGMetadataAttributeDesc ExpectedDesc
+		{
+			.Name = AttributeName,
+			.ValueType = EPCGMetadataTypes::Struct,
+			.ContainerTypes = {EPCGMetadataAttributeContainerTypes::Array},
+			.ValueTypeObject = FPCGPoint::StaticStruct()
+		};
+		
+		return TypedAttributeTester<TArray<FPCGPoint>, TConstArrayView<FPCGPoint>>(GenerateRandom, Verify, MoveTemp(DefaultValue), MoveTemp(ExpectedDesc));
+	}
+	
+	TypedAttributeTester<TSet<double>, PCG::TScriptSetWrapper<double>> SetDoubleTester()
+	{		
+		auto GenerateRandom = [](FRandomStream& RandomStream, TSet<double>& OutValue)
+			{
+				const int32 RandomCount = RandomStream.RandRange(1, 4);
+				for (int j = 0; j < RandomCount; ++j)
+				{
+					OutValue.Add(RandomStream.FRand());
+				}
+			};
+		
+		auto Verify = [](const TSet<double>& LHS, const PCG::TScriptSetWrapper<double>& RHS)
+			{
+				if (LHS.Num() != RHS.Num())
+				{
+					return false;
+				}
+				
+				for (double Value : LHS)
+				{
+					if (!RHS.Contains(Value))
+					{
+						return false;
+					}
+				}
+				
+				return true;
+			};
+		
+		TSet<double> DefaultValue = {1, 2, 3, 4, 5};
+		
+		FPCGMetadataAttributeDesc ExpectedDesc
+		{
+			.Name = AttributeName,
+			.ValueType = EPCGMetadataTypes::Double,
+			.ContainerTypes = {EPCGMetadataAttributeContainerTypes::Set}
+		};
+		
+		return TypedAttributeTester<TSet<double>, PCG::TScriptSetWrapper<double>>(GenerateRandom, Verify, MoveTemp(DefaultValue), MoveTemp(ExpectedDesc));
+	}
+	
+	TypedAttributeTester<TMap<FString, double>, PCG::TScriptMapWrapper<FString, double>> MapStringDoubleTester()
+	{
+		auto GenerateRandom = [](FRandomStream& RandomStream, TMap<FString, double>& OutValue)
+			{
+				const int32 RandomCount = RandomStream.RandRange(1, 4);
+				for (int j = 0; j < RandomCount; ++j)
+				{
+					OutValue.Emplace(LexToString(j), RandomStream.FRand());
+				}
+			};
+		
+		auto Verify = [](const TMap<FString, double>& LHS, const PCG::TScriptMapWrapper<FString, double>& RHS)
+			{
+				if (LHS.Num() != RHS.Num())
+				{
+					return false;
+				}
+				
+				for (const auto& [Key, Value] : LHS)
+				{
+					const double* OtherValue = RHS.Find(Key);
+					if (!OtherValue || Value != *OtherValue)
+					{
+						return false;
+					}
+				}
+				
+				return true;
+			};
+		
+		TMap<FString, double> DefaultValue =
+		{
+			{TEXT("Hi"), 2.0},
+			{TEXT("Ho"), 3.0}
+		};
+		
+		FPCGMetadataAttributeDesc ExpectedDesc
+		{
+			.Name = AttributeName,
+			.ValueType = EPCGMetadataTypes::Double,
+			.ContainerTypes = {EPCGMetadataAttributeContainerTypes::Map},
+			.KeyType = EPCGMetadataTypes::String
+		};
+		
+		return TypedAttributeTester<TMap<FString, double>, PCG::TScriptMapWrapper<FString, double>>(GenerateRandom, Verify, MoveTemp(DefaultValue), MoveTemp(ExpectedDesc));
+	}
+	
+	TypedAttributeTester<EPCGMetadataTypes> EnumTester()
+	{
+		auto GenerateRandom = [](FRandomStream& RandomStream, EPCGMetadataTypes& OutValue)
+			{
+				OutValue = static_cast<EPCGMetadataTypes>(RandomStream.RandRange(0, static_cast<int32>(EPCGMetadataTypes::Count)));
+			};
+		
+		auto Verify = [](const EPCGMetadataTypes& LHS, const EPCGMetadataTypes& RHS) { return LHS == RHS; };
+		
+		EPCGMetadataTypes DefaultValue = EPCGMetadataTypes::Unknown;
+		
+		FPCGMetadataAttributeDesc ExpectedDesc
+		{
+			.Name = AttributeName,
+			.ValueType = EPCGMetadataTypes::Enum,
+			.ValueTypeObject = StaticEnum<EPCGMetadataTypes>()
+		};
+		
+		return TypedAttributeTester<EPCGMetadataTypes>(GenerateRandom, Verify, DefaultValue, MoveTemp(ExpectedDesc));
+	}
+}
+
+/**
+ * Test suite to verify basic functionality with generic attributes
+ * Rely on the TypedAttributeTester that handles anything that is type specific, the rest is templated in
+ * the Tester lambda.
+ * Some attribute don't have the same initial type with the type we read, so we distinguish between the Attribute type T
+ * and the ReadType.
+ * 
+ * Test routine is pretty standard.
+ *   - A point data is created and a generic attribute is added to it.
+ *   - Verify that the descriptor of the generic attribute is the one we expect.
+ *   - Verify that the default value is set and retrieved correctly
+ *   - Generate random values for the generic attribute, set them on the attribute and read them back.
+ *   - Verify that the values read back are the same as the one we set.
+ *   - Do it again with a child point array data inheriting from the initial one, while modifying even indices
+ *     to verify that we can still read old values (odd indices, unchanged), but also the new ones.
+ */
+TEST_CASE_METHOD(PCGTests::FPCGBaseTestWithBasicContext, "PCG::AttributeGeneric::Basic", "[PCG][AttributeGeneric]")
+{
+	constexpr int32 Seed = 42;
+	
+	auto Tester = [this]<typename T, typename ReadType>(PCGAttributeGenericTests::TypedAttributeTester<T, ReadType> AttributeTester)
+	{
+		// Preparing the data
+		constexpr int32 NumPoints = 100;
+		UPCGPointArrayData* PointArrayData = CreateData<UPCGPointArrayData>();
+
+		PointArrayData->SetNumPoints(NumPoints);
+
+		PointArrayData->AllocateProperties(EPCGPointNativeProperties::MetadataEntry);
+		TPCGValueRange<int64> PointArrayDataMetadataEntryRange = PointArrayData->GetMetadataEntryValueRange(/*bAllocate=*/false);
+
+		FPCGMetadataDomain* Domain = PointArrayData->MutableMetadata()->GetDefaultMetadataDomain();
+		
+		FRandomStream RandomStream(Seed);
+		
+		FPCGMetadataAttributeGeneric* GenericAttribute = nullptr;
+	
+		if (AttributeTester.OverrideDesc.IsSet())
+		{
+			GenericAttribute = Domain->CreateGenericAttribute(*AttributeTester.OverrideDesc, false, false);
+			REQUIRE_NOT_EQUAL(GenericAttribute, nullptr);
+			GenericAttribute->SetDefaultValue<T>(AttributeTester.DefaultValue);
+		}
+		else
+		{
+			GenericAttribute = Domain->CreateGenericAttribute<T>(PCGAttributeGenericTests::AttributeName, AttributeTester.DefaultValue, false, false);
+			REQUIRE_NOT_EQUAL(GenericAttribute, nullptr);
+		}
+		
+		SECTION("Expected attribute descriptor")
+		{
+			REQUIRE_EQUAL(GenericAttribute->GetAttributeDesc(), AttributeTester.ExpectedDesc);
+		}
+		
+		SECTION("Default value is set")
+		{
+			ReadType ReadDefaultValue{};
+			GenericAttribute->GetValuesFromItemKeys<ReadType>({PCGInvalidEntryKey}, MakeArrayView(&ReadDefaultValue, 1));
+			REQUIRE(AttributeTester.Verify(AttributeTester.DefaultValue, ReadDefaultValue));
+		}
+		
+		TStaticArray<T, NumPoints> Values{};
+		TStaticArray<ReadType, NumPoints> ReadValues;
+		
+		for (int32 i = 0; i < NumPoints; ++i)
+		{
+			PointArrayDataMetadataEntryRange[i] = Domain->AddEntry(PointArrayDataMetadataEntryRange[i]);
+			AttributeTester.GenerateRandom(RandomStream, Values[i]);
+		}
+
+		GenericAttribute->SetValues<T>(PointArrayDataMetadataEntryRange, Values);
+		REQUIRE_EQUAL(GenericAttribute->GetNumberOfEntries(), Values.Num());
+
+		SECTION("Read back and verify values")
+		{
+			GenericAttribute->GetValuesFromItemKeys<ReadType>(PCGValueRangeHelpers::MakeConstValueRange(PointArrayDataMetadataEntryRange), ReadValues);
+			CHECK_THAT(Values, Catch::Matchers::RangeEquals(ReadValues, [&AttributeTester](const auto& a, const auto& b){ return AttributeTester.Verify(a, b); } ));
+		}
+
+		SECTION("Parenting")
+		{
+			UPCGPointArrayData* ChildPointArrayData = CreateData<UPCGPointArrayData>();
+			ChildPointArrayData->InitializeFromDataWithParams(FPCGInitializeFromDataParams{ PointArrayData });
+
+			FPCGMetadataDomain* ChildDomain = ChildPointArrayData->MutableMetadata()->GetDefaultMetadataDomain();
+
+			TPCGValueRange<int64> ChildPointArrayDataMetadataEntryRange = ChildPointArrayData->GetMetadataEntryValueRange(/*bAllocate=*/false);
+
+			FPCGMetadataAttributeGeneric* ChildGenericAttribute = ChildDomain->GetMutableGenericAttribute(PCGAttributeGenericTests::AttributeName);
+			REQUIRE_NOT_EQUAL(ChildGenericAttribute, nullptr);
+			
+			SECTION("Expected attribute descriptor")
+			{
+				REQUIRE_EQUAL(ChildGenericAttribute->GetAttributeDesc(), AttributeTester.ExpectedDesc);
+			}
+			
+			SECTION("Default value is set")
+			{
+				ReadType ReadDefaultValue{};
+				ChildGenericAttribute->GetValuesFromItemKeys<ReadType>({PCGInvalidEntryKey}, MakeArrayView(&ReadDefaultValue, 1));
+				REQUIRE(AttributeTester.Verify(AttributeTester.DefaultValue, ReadDefaultValue));
+			}
+
+			// Change values of even indices.
+			TStaticArray<T, NumPoints / 2> ChildValues{};
+			TStaticArray<PCGMetadataEntryKey, NumPoints / 2> ChildEntryKeys{};
+			for (int32 i = 0; i < NumPoints; ++i)
+			{
+				if (i % 2 == 0)
+				{
+					AttributeTester.GenerateRandom(RandomStream, Values[i]);
+
+					ChildDomain->InitializeOnSet(ChildPointArrayDataMetadataEntryRange[i]);
+					ChildValues[i / 2] = Values[i];
+					ChildEntryKeys[i / 2] = ChildPointArrayDataMetadataEntryRange[i];
+				}
+			}
+
+			ChildGenericAttribute->SetValues<T>(ChildEntryKeys, ChildValues);
+			REQUIRE_EQUAL(ChildGenericAttribute->GetNumberOfEntries(), ChildValues.Num());
+			REQUIRE_EQUAL(ChildGenericAttribute->GetNumberOfEntriesWithParents(), Values.Num() + ChildValues.Num());
+
+			SECTION("Read back and verify values")
+			{
+				ChildGenericAttribute->GetValuesFromItemKeys<ReadType>(PCGValueRangeHelpers::MakeConstValueRange(ChildPointArrayDataMetadataEntryRange), ReadValues);
+				CHECK_THAT(Values, Catch::Matchers::RangeEquals(ReadValues, [&AttributeTester](const auto& a, const auto& b){ return AttributeTester.Verify(a, b); } ));
+			}
+		}
+	};
+	
+	SECTION("Float Attribute")
+	{
+		Tester(PCGAttributeGenericTests::FloatTester());
+	}
+	
+	SECTION("Double Attribute")
+	{
+		Tester(PCGAttributeGenericTests::DoubleTester());
+	}
+	
+	SECTION("Vector Attribute")
+	{
+		SECTION("As Basic Type")
+		{
+			Tester(PCGAttributeGenericTests::VectorTester_Basic());
+		}
+		
+		SECTION("As Struct")
+		{
+			Tester(PCGAttributeGenericTests::VectorTester_Struct());
+		}
+	}
+	
+	SECTION("String Attribute")
+	{
+		Tester(PCGAttributeGenericTests::StringTester());
+	}
+	
+	SECTION("FPCGPoint Attribute")
+	{		
+		Tester(PCGAttributeGenericTests::FPCGPointTester());
+	}
+	
+	SECTION("Array of Double Attribute")
+	{		
+		Tester(PCGAttributeGenericTests::ArrayDoubleTester());
+	}
+	
+	SECTION("Array of FPCGPoint Attribute")
+	{	
+		Tester(PCGAttributeGenericTests::ArrayFPCGPointTester());
+	}
+	
+	SECTION("Set of Double Attribute")
+	{		
+		Tester(PCGAttributeGenericTests::SetDoubleTester());
+	}
+	
+	SECTION("Map of String -> Double Attribute")
+	{
+		Tester(PCGAttributeGenericTests::MapStringDoubleTester());
+	}
+	
+	SECTION("Enum attribute")
+	{
+		Tester(PCGAttributeGenericTests::EnumTester());
+	}
+}
+
+/**
+ * Test that inheritance with Copy also produces the right values.
+ */
+TEST_CASE_METHOD(PCGTests::FPCGBaseTestWithBasicContext, "PCG::AttributeGeneric::Copy", "[PCG][AttributeGeneric]")
+{
+	constexpr int32 Seed = 42;
+	
+	auto Tester = [this]<typename T, typename ReadType>(PCGAttributeGenericTests::TypedAttributeTester<T, ReadType> AttributeTester)
+	{
+		// Preparing the data
+		constexpr int32 NumEntries = 10;
+		UPCGParamData* ParamData = CreateData<UPCGParamData>();
+		FPCGMetadataDomain* Domain = ParamData->MutableMetadata()->GetDefaultMetadataDomain();
+		
+		FRandomStream RandomStream(Seed);
+		
+		FPCGMetadataAttributeGeneric* GenericAttribute = nullptr;
+	
+		if (AttributeTester.OverrideDesc.IsSet())
+		{
+			GenericAttribute = Domain->CreateGenericAttribute(*AttributeTester.OverrideDesc, false, false);
+			REQUIRE_NOT_EQUAL(GenericAttribute, nullptr);
+			GenericAttribute->SetDefaultValue<T>(AttributeTester.DefaultValue);
+		}
+		else
+		{
+			GenericAttribute = Domain->CreateGenericAttribute<T>(PCGAttributeGenericTests::AttributeName, AttributeTester.DefaultValue, false, false);
+			REQUIRE_NOT_EQUAL(GenericAttribute, nullptr);
+		}
+		
+		TStaticArray<T, NumEntries> Values{};
+		TStaticArray<ReadType, NumEntries> ReadValues;
+		TStaticArray<PCGMetadataEntryKey, NumEntries> EntryKeys;
+		
+		for (int32 i = 0; i < NumEntries; ++i)
+		{
+			EntryKeys[i] = Domain->AddEntry();
+			AttributeTester.GenerateRandom(RandomStream, Values[i]);
+		}
+
+		GenericAttribute->SetValues<T>(EntryKeys, Values);
+		REQUIRE_EQUAL(GenericAttribute->GetNumberOfEntries(), Values.Num());
+
+		SECTION("Read back and verify values")
+		{
+			GenericAttribute->GetValuesFromItemKeys<ReadType>(EntryKeys, ReadValues);
+			CHECK_THAT(Values, Catch::Matchers::RangeEquals(ReadValues, [&AttributeTester](const auto& a, const auto& b){ return AttributeTester.Verify(a, b); } ));
+		}
+
+		SECTION("Parenting")
+		{
+			UPCGParamData* ChildParamData = CreateData<UPCGParamData>();
+			ChildParamData->MutableMetadata()->Initialize(ParamData->ConstMetadata());
+
+			FPCGMetadataDomain* ChildDomain = ChildParamData->MutableMetadata()->GetDefaultMetadataDomain();
+
+			FPCGMetadataAttributeGeneric* ChildGenericAttribute = ChildDomain->GetMutableGenericAttribute(PCGAttributeGenericTests::AttributeName);
+			REQUIRE_NOT_EQUAL(ChildGenericAttribute, nullptr);
+			
+			SECTION("Expected attribute descriptor")
+			{
+				REQUIRE_EQUAL(ChildGenericAttribute->GetAttributeDesc(), AttributeTester.ExpectedDesc);
+			}
+			
+			SECTION("Default value is set")
+			{
+				ReadType ReadDefaultValue{};
+				ChildGenericAttribute->GetValuesFromItemKeys<ReadType>({PCGInvalidEntryKey}, MakeArrayView(&ReadDefaultValue, 1));
+				REQUIRE(AttributeTester.Verify(AttributeTester.DefaultValue, ReadDefaultValue));
+			}
+
+			// Change values of even indices.
+			TStaticArray<T, NumEntries / 2> ChildValues{};
+			TStaticArray<PCGMetadataEntryKey, NumEntries / 2> ChildEntryKeys{};
+			for (int32 i = 0; i < NumEntries; ++i)
+			{
+				if (i % 2 == 0)
+				{
+					AttributeTester.GenerateRandom(RandomStream, Values[i]);
+
+					ChildDomain->InitializeOnSet(EntryKeys[i]);
+					ChildValues[i / 2] = Values[i];
+					ChildEntryKeys[i / 2] = EntryKeys[i];
+				}
+			}
+
+			ChildGenericAttribute->SetValues<T>(ChildEntryKeys, ChildValues);
+			REQUIRE_EQUAL(ChildGenericAttribute->GetNumberOfEntries(), Values.Num());
+			REQUIRE_EQUAL(ChildGenericAttribute->GetNumberOfEntriesWithParents(), Values.Num());
+
+			SECTION("Read back and verify values")
+			{
+				ChildGenericAttribute->GetValuesFromItemKeys<ReadType>(EntryKeys, ReadValues);
+				CHECK_THAT(Values, Catch::Matchers::RangeEquals(ReadValues, [&AttributeTester](const auto& a, const auto& b){ return AttributeTester.Verify(a, b); } ));
+			}
+		}
+	};
+	
+	SECTION("Float Attribute")
+	{
+		Tester(PCGAttributeGenericTests::FloatTester());
+	}
+	
+	SECTION("Double Attribute")
+	{
+		Tester(PCGAttributeGenericTests::DoubleTester());
+	}
+	
+	SECTION("Vector Attribute")
+	{
+		SECTION("As Basic Type")
+		{
+			Tester(PCGAttributeGenericTests::VectorTester_Basic());
+		}
+		
+		SECTION("As Struct")
+		{
+			Tester(PCGAttributeGenericTests::VectorTester_Struct());
+		}
+	}
+	
+	SECTION("String Attribute")
+	{
+		Tester(PCGAttributeGenericTests::StringTester());
+	}
+	
+	SECTION("FPCGPoint Attribute")
+	{		
+		Tester(PCGAttributeGenericTests::FPCGPointTester());
+	}
+	
+	SECTION("Array of Double Attribute")
+	{		
+		Tester(PCGAttributeGenericTests::ArrayDoubleTester());
+	}
+	
+	SECTION("Array of FPCGPoint Attribute")
+	{	
+		Tester(PCGAttributeGenericTests::ArrayFPCGPointTester());
+	}
+	
+	SECTION("Set of Double Attribute")
+	{		
+		Tester(PCGAttributeGenericTests::SetDoubleTester());
+	}
+	
+	SECTION("Map of String -> Double Attribute")
+	{
+		Tester(PCGAttributeGenericTests::MapStringDoubleTester());
+	}
+	
+	SECTION("Enum attribute")
+	{
+		Tester(PCGAttributeGenericTests::EnumTester());
+	}
+}
+
+/**
+ * Test that getting/setting values with the wrong type is not corrupting the data and fail gracefully.
+ */
+TEST_CASE_METHOD(PCGTests::FPCGBaseTestWithBasicContext, "PCG::AttributeGeneric::TypeChecking", "[PCG][AttributeGeneric]")
+{
+	auto Tester = [this]<typename RealType, typename ReadType, typename WriteType>(const RealType& InValue, ReadType DummyRead, WriteType DummyWrite)
+	{
+		UPCGParamData* ParamData = CreateData<UPCGParamData>();
+		FPCGMetadataDomain* Domain = ParamData->MutableMetadata()->GetDefaultMetadataDomain();
+		
+		FPCGMetadataAttributeGeneric* GenericAttribute = Domain->CreateGenericAttribute<RealType>(PCGAttributeGenericTests::AttributeName, InValue, false, false);
+		
+		SECTION("Get Value with wrong type")
+		{
+			PCGTests::FPCGTestsLogOutputDevice LogCapture(/*bSuppressErrors=*/true);
+			GenericAttribute->GetValuesFromItemKeys<ReadType>({PCGInvalidEntryKey}, MakeArrayView(&DummyRead, 1));
+			REQUIRE(LogCapture.NbMessageReceived.Contains(ELogVerbosity::Error));
+		}
+		
+		SECTION("Set Value with wrong type")
+		{
+			PCGTests::FPCGTestsLogOutputDevice LogCapture(/*bSuppressErrors=*/true);
+			PCGMetadataEntryKey EntryKey = Domain->AddEntry();
+			GenericAttribute->SetValues<WriteType>({EntryKey}, MakeConstArrayView(&DummyWrite, 1));
+			REQUIRE(LogCapture.NbMessageReceived.Contains(ELogVerbosity::Error));
+			REQUIRE_EQUAL(GenericAttribute->GetNumberOfEntries(), 0);
+		}
+	};
+	
+	SECTION("Double with float")
+	{
+		Tester(double{}, float{}, float{});
+	}
+	
+	SECTION("String with name")
+	{
+		Tester(FString{}, FName{}, FName{});
+	}
+	
+	SECTION("Array of double with double")
+	{
+		Tester(TArray<double>{}, double{}, double{});
+	}
+	
+	SECTION("Array of double with array of float")
+	{
+		Tester(TArray<double>{}, TConstArrayView<float>{}, TArray<float>{});
+	}
+	
+	SECTION("Set of double with double")
+	{
+		Tester(TSet<double>{}, double{}, double{});
+	}
+	
+	SECTION("Set of double with Set of float")
+	{
+		Tester(TSet<double>{}, PCG::TScriptSetWrapper<float>{}, TSet<float>{});
+	}
+	
+	SECTION("Map of String/double with String")
+	{
+		Tester(TMap<FString, double>{}, FString{}, FString{});
+	}
+	
+	SECTION("Map of String/double with Map of String/String")
+	{
+		Tester(TMap<FString, double>{}, PCG::TScriptMapWrapper<FString, FString>{}, TMap<FString, FString>{});
+	}
+	
+	SECTION("Map of String/double with Map of double/double")
+	{
+		Tester(TMap<FString, double>{}, PCG::TScriptMapWrapper<double, double>{}, TMap<double, double>{});
+	}
+	
+	SECTION("Map of String/double with Map of double/String")
+	{
+		Tester(TMap<FString, double>{}, PCG::TScriptMapWrapper<double, FString>{}, TMap<double, FString>{});
+	}
+	
+	SECTION("Map of String/double with Map of Name/float")
+	{
+		Tester(TMap<FString, double>{}, PCG::TScriptMapWrapper<FName, float>{}, TMap<FName, float>{});
+	}
+	
+	SECTION("PCGPoint with PCGPinProperties")
+	{
+		Tester(FPCGPoint{}, static_cast<const FPCGPinProperties*>(nullptr), FPCGPinProperties{});
+	}
+	
+	SECTION("Array of PCGPoint with array of PCGPinProperties")
+	{
+		Tester(TArray<FPCGPoint>{}, TConstArrayView<FPCGPinProperties>{}, TArray<FPCGPinProperties>{});
+	}
+	
+	SECTION("Set of PCGPoint with set of PCGPinProperties")
+	{
+		Tester(TSet<FPCGPoint>{}, PCG::TScriptSetWrapper<FPCGPinProperties>{}, TSet<FPCGPinProperties>{});
+	}
+	
+	SECTION("Map of String/PCGPoint with map of String/PCGPinProperties")
+	{
+		Tester(TMap<FString, FPCGPoint>{}, PCG::TScriptMapWrapper<FString, FPCGPinProperties>{}, TMap<FString, FPCGPinProperties>{});
+	}
+	
+	SECTION("EPCGMetadataTypes with EPCGMetadataAttributeContainerTypes")
+	{
+		Tester(EPCGMetadataTypes{}, EPCGMetadataAttributeContainerTypes{}, EPCGMetadataAttributeContainerTypes{});
+	}
+	
+	SECTION("Array of EPCGMetadataTypes with array of EPCGMetadataAttributeContainerTypes")
+	{
+		Tester(TArray<EPCGMetadataTypes>{}, TConstArrayView<EPCGMetadataAttributeContainerTypes>{}, TArray<EPCGMetadataAttributeContainerTypes>{});
+	}
+	
+	SECTION("Set of EPCGMetadataTypes with set of EPCGMetadataAttributeContainerTypes")
+	{
+		Tester(TSet<EPCGMetadataTypes>{}, PCG::TScriptSetWrapper<EPCGMetadataAttributeContainerTypes>{}, TSet<EPCGMetadataAttributeContainerTypes>{});
+	}
+	
+	SECTION("Map of FString/EPCGMetadataTypes with map of FString/EPCGMetadataAttributeContainerTypes")
+	{
+		Tester(TMap<FString, EPCGMetadataTypes>{}, PCG::TScriptMapWrapper<FString, EPCGMetadataAttributeContainerTypes>{}, TMap<FString, EPCGMetadataAttributeContainerTypes>{});
+	}
+	
+	SECTION("Map of EPCGMetadataTypes/FString with map of EPCGMetadataAttributeContainerTypes/FString")
+	{
+		Tester(TMap<EPCGMetadataTypes, FString>{}, PCG::TScriptMapWrapper<EPCGMetadataAttributeContainerTypes, FString>{}, TMap<EPCGMetadataAttributeContainerTypes, FString>{});
+	}
+}
+
+TEST_CASE_METHOD(PCGTests::FPCGBaseTestWithBasicContext, "PCG::AttributeGeneric::CompressedData", "[PCG][AttributeGeneric]")
+{
+	// Preparing the data
+	constexpr int32 NumPoints = 100;
+	UPCGPointArrayData* PointArrayData = CreateData<UPCGPointArrayData>();
+
+	PointArrayData->SetNumPoints(NumPoints);
+
+	PointArrayData->AllocateProperties(EPCGPointNativeProperties::MetadataEntry);
+	TPCGValueRange<int64> PointArrayDataMetadataEntryRange = PointArrayData->GetMetadataEntryValueRange(/*bAllocate=*/false);
+
+	FPCGMetadataDomain* Domain = PointArrayData->MutableMetadata()->GetDefaultMetadataDomain();
+	
+	// Initialize the default value to a value we will add later on.
+	FPCGMetadataAttributeGeneric* GenericAttribute = Domain->CreateGenericAttribute<FString>(PCGAttributeGenericTests::AttributeName, TEXT("Value: 0"), false, false);
+	REQUIRE_NOT_EQUAL(GenericAttribute, nullptr);
+	
+	// First set half the points with different values
+	TStaticArray<FString, NumPoints> Values{};
+	TStaticArray<PCGMetadataEntryKey, NumPoints/2> EntryKeys{};
+	
+	for (int32 i = 0; i < NumPoints/2; ++i)
+	{
+		Domain->InitializeOnSet(PointArrayDataMetadataEntryRange[i]);
+		EntryKeys[i] = PointArrayDataMetadataEntryRange[i];
+		Values[i] = FString::Printf(TEXT("Value: %d"), i % 10);
+	}
+	
+	// Set the values
+	GenericAttribute->SetValues<FString>(EntryKeys, MakeArrayView(Values.GetData(), NumPoints / 2));
+	
+	// Verify that we have only 10 Value keys (including the default value one)
+	REQUIRE_EQUAL(GenericAttribute->GetValueKeyOffsetForChild(), 10);
+	
+	// Then add a mix of existing strings and new strings, to test that we only add what does not exist
+	for (int32 i = NumPoints/2; i < NumPoints; ++i)
+	{
+		const int32 EntryKeyIndex = i - (NumPoints/2);
+		Domain->InitializeOnSet(PointArrayDataMetadataEntryRange[i]);
+		EntryKeys[EntryKeyIndex] = PointArrayDataMetadataEntryRange[i];
+		if (EntryKeyIndex < 10)
+		{
+			Values[i] = FString::Printf(TEXT("Value: %d"), EntryKeyIndex % 10);
+		}
+		else
+		{
+			Values[i] = FString::Printf(TEXT("New Value: %d"), EntryKeyIndex % 10);
+		}
+	}
+	
+	// Set the values
+	GenericAttribute->SetValues<FString>(EntryKeys, MakeArrayView(Values.GetData() + (NumPoints / 2), NumPoints / 2));
+	
+	// Verify that we have only 20 Value keys (including the default value one)
+	REQUIRE_EQUAL(GenericAttribute->GetValueKeyOffsetForChild(), 20);
+	
+	// Verify that all the values can be retrieved.
+	TStaticArray<const FString*, NumPoints> ReadValues{};
+	GenericAttribute->GetValuesFromItemKeys<const FString*>(PCGValueRangeHelpers::MakeConstValueRange(PointArrayDataMetadataEntryRange), ReadValues);
+	
+	CHECK_THAT(Values, Catch::Matchers::RangeEquals(ReadValues, [](const FString& a, const FString* b){ return a == *b; } ));
+}
\ No newline at end of file
diff --git a/Engine/Plugins/PCG/Tests/Private/Metadata/PCGMetadataContainerTypesTests.cpp b/Engine/Plugins/PCG/Tests/Private/Metadata/PCGMetadataContainerTypesTests.cpp
new file mode 100644
--- /dev/null
+++ b/Engine/Plugins/PCG/Tests/Private/Metadata/PCGMetadataContainerTypesTests.cpp
@@ -0,0 +1,170 @@
+﻿// Copyright Epic Games, Inc. All Rights Reserved.
+
+#include "PCGTestsCommon.h"
+
+#include "Data/PCGPointArrayData.h"
+
+#include <catch2/catch_test_macros.hpp>
+#include <catch2/benchmark/catch_benchmark.hpp>
+#include <catch2/generators/catch_generators_range.hpp>
+
+#include "TestHarness.h"
+
+TEST_CASE_METHOD(PCGTests::FPCGBaseTestWithBasicContext, "PCG::MetadataContainers", "[PCG][MetadataContainers]")
+{
+	// Since the containers are meant to be created by the generic attributes, we'll use the generic attributes
+	// to populate them
+	
+	// Preparing the data
+	constexpr int32 NumPoints = 1;
+	UPCGPointArrayData* PointArrayData = CreateData<UPCGPointArrayData>();
+
+	PointArrayData->SetNumPoints(NumPoints);
+
+	PointArrayData->AllocateProperties(EPCGPointNativeProperties::MetadataEntry);
+	TPCGValueRange<int64> PointArrayDataMetadataEntryRange = PointArrayData->GetMetadataEntryValueRange(/*bAllocate=*/false);
+	
+	FPCGMetadataDomain* Domain = PointArrayData->MutableMetadata()->GetDefaultMetadataDomain();
+	PointArrayDataMetadataEntryRange[0] = Domain->AddEntry();
+	
+	SECTION("TScriptSetWrapper")
+	{
+		FPCGMetadataAttributeGeneric* SetGenericAttribute = Domain->CreateGenericAttribute<TSet<double>>("SetAttr", TSet<double>{}, false, false);
+		REQUIRE_NOT_EQUAL(SetGenericAttribute, nullptr);
+		
+		const TSet<double> Values = {1.0, 2.0, 5.0, 3.0};
+		
+		SetGenericAttribute->SetValues<TSet<double>>(PointArrayDataMetadataEntryRange, {Values});
+		PCG::TScriptSetWrapper<double> SetWrapper{};
+		SetGenericAttribute->GetValuesFromItemKeys<PCG::TScriptSetWrapper<double>>(PCGValueRangeHelpers::MakeConstValueRange(PointArrayDataMetadataEntryRange), MakeArrayView(&SetWrapper, 1));
+		
+		REQUIRE(SetWrapper.IsValid());
+		
+		SECTION("Num")
+		{
+			REQUIRE(SetWrapper.Num() == Values.Num());
+		}
+		
+		SECTION("Contains")
+		{
+			REQUIRE(SetWrapper.Contains(1.0));
+			REQUIRE(SetWrapper.Contains(2.0));
+			REQUIRE(SetWrapper.Contains(5.0));
+			REQUIRE(SetWrapper.Contains(3.0));
+			REQUIRE_FALSE(SetWrapper.Contains(4.0));
+			REQUIRE_FALSE(SetWrapper.Contains(6.0));
+		}
+		
+		SECTION("ForEach")
+		{
+			double Sum = 0.0;
+			SetWrapper.ForEach([&Sum](double Value) { Sum += Value; });
+			REQUIRE_EQUAL(Sum, 11.0);
+		}
+		
+		SECTION("Intersection")
+		{
+			const TSet<double> Other = {2.0, 5.0, 4.0};
+			TSet<double> Intersection = SetWrapper.Intersection(Other);
+			REQUIRE_EQUAL(Intersection.Num(), 2);
+			REQUIRE_FALSE(Intersection.Contains(1.0));
+			REQUIRE(Intersection.Contains(2.0));
+			REQUIRE(Intersection.Contains(5.0));
+			REQUIRE_FALSE(Intersection.Contains(3.0));
+			REQUIRE_FALSE(Intersection.Contains(4.0));
+			REQUIRE_FALSE(Intersection.Contains(6.0));
+		}
+		
+		SECTION("Union")
+		{
+			const TSet<double> Other = {2.0, 5.0, 4.0};
+			TSet<double> Union = SetWrapper.Union(Other);
+			REQUIRE_EQUAL(Union.Num(), 5);
+			REQUIRE(Union.Contains(1.0));
+			REQUIRE(Union.Contains(2.0));
+			REQUIRE(Union.Contains(5.0));
+			REQUIRE(Union.Contains(3.0));
+			REQUIRE(Union.Contains(4.0));
+			REQUIRE_FALSE(Union.Contains(6.0));
+		}
+
+		SECTION("Copy")
+		{
+			TSet<double> Copy = static_cast<TSet<double>>(SetWrapper);
+			REQUIRE_EQUAL(Copy.Num(), 4);
+			REQUIRE(Copy.Contains(1.0));
+			REQUIRE(Copy.Contains(2.0));
+			REQUIRE(Copy.Contains(5.0));
+			REQUIRE(Copy.Contains(3.0));
+			REQUIRE_FALSE(Copy.Contains(4.0));
+			REQUIRE_FALSE(Copy.Contains(6.0));
+		}
+	}
+	
+	SECTION("TScriptMapWrapper")
+	{
+		FPCGMetadataAttributeGeneric* MapGenericAttribute = Domain->CreateGenericAttribute<TMap<FString, double>>("MapAttr", TMap<FString, double>{}, false, false);
+		REQUIRE_NOT_EQUAL(MapGenericAttribute, nullptr);
+		
+		TMap<FString, double> Values = 
+		{
+			{TEXT("Hey"), 2.0},
+			{TEXT("Ho"), 5.0}
+		};
+		
+		MapGenericAttribute->SetValues<TMap<FString, double>>(PointArrayDataMetadataEntryRange, {Values});
+		PCG::TScriptMapWrapper<FString, double> MapWrapper{};
+		MapGenericAttribute->GetValuesFromItemKeys<PCG::TScriptMapWrapper<FString, double>>(PCGValueRangeHelpers::MakeConstValueRange(PointArrayDataMetadataEntryRange), MakeArrayView(&MapWrapper, 1));
+		
+		REQUIRE(MapWrapper.IsValid());
+		
+		SECTION("Num")
+		{
+			REQUIRE(MapWrapper.Num() == Values.Num());
+		}
+		
+		SECTION("Contains")
+		{
+			REQUIRE(MapWrapper.Contains(TEXT("Hey")));
+			REQUIRE(MapWrapper.Contains(TEXT("Ho")));
+			REQUIRE_FALSE(MapWrapper.Contains(TEXT("Ah")));
+		}
+		
+		SECTION("Find")
+		{
+			REQUIRE_EQUAL(*MapWrapper.Find(TEXT("Hey")), 2.0);
+			REQUIRE_EQUAL(*MapWrapper.Find(TEXT("Ho")), 5.0);
+			REQUIRE_EQUAL(MapWrapper.Find(TEXT("Ah")), nullptr);
+		}
+		
+		SECTION("operator[]")
+		{
+			REQUIRE_EQUAL(MapWrapper[TEXT("Hey")], 2.0);
+			REQUIRE_EQUAL(MapWrapper[TEXT("Ho")], 5.0);
+		}
+		
+		SECTION("Copy")
+		{
+			TMap<FString, double> Copy = static_cast<TMap<FString, double>>(MapWrapper);
+			REQUIRE_EQUAL(Copy.Num(), 2);
+			REQUIRE_EQUAL(Copy[TEXT("Hey")], 2.0);
+			REQUIRE_EQUAL(Copy[TEXT("Ho")], 5.0);
+			REQUIRE_FALSE(Copy.Contains(TEXT("Ah")));
+		}
+		
+		SECTION("ForEach")
+		{
+			FString SumString{};
+			double Sum = 0.0;
+			
+			MapWrapper.ForEach([&SumString, &Sum](const FString& Key, const double& Value)
+				{
+					SumString += Key;
+					Sum += Value;
+				});
+				
+			REQUIRE_EQUAL(SumString, TEXT("HeyHo"));
+			REQUIRE_EQUAL(Sum, 7.0);
+		}
+	}
+}
\ No newline at end of file
diff --git a/Engine/Plugins/PCG/Tests/Private/Metadata/PCGMetadataFlattenTests.cpp b/Engine/Plugins/PCG/Tests/Private/Metadata/PCGMetadataFlattenTests.cpp
new file mode 100644
--- /dev/null
+++ b/Engine/Plugins/PCG/Tests/Private/Metadata/PCGMetadataFlattenTests.cpp
@@ -0,0 +1,241 @@
+// Copyright Epic Games, Inc. All Rights Reserved.
+
+#include "PCGTestsCommon.h"
+
+#include "Data/PCGPointArrayData.h"
+#include "Metadata/PCGMetadata.h"
+
+#include <catch2/catch_test_macros.hpp>
+#include <catch2/benchmark/catch_benchmark.hpp>
+#include <catch2/generators/catch_generators_range.hpp>
+
+#include "TestHarness.h"
+
+TEST_CASE_METHOD(PCGTests::FPCGBaseTestWithBasicContext, "PCG::Metadata::Flatten", "[PCG][Metadata][Flatten]")
+{
+	static const FName Attribute1Name = TEXT("FloatAttr");
+	static const FName Attribute2Name = TEXT("StringAttr");
+	static const FName Attribute3Name = TEXT("IntAttr");
+
+	constexpr int32 NumPoints = 10;
+	UPCGPointArrayData* RootPointData = CreateData<UPCGPointArrayData>();
+	RootPointData->SetNumPoints(NumPoints);
+	
+	FPCGMetadataAttribute<float>* Attribute1 = RootPointData->Metadata->CreateAttribute<float>(Attribute1Name, -0.1f, true, true);
+	FPCGMetadataAttribute<FString>* Attribute2 = RootPointData->Metadata->CreateAttribute<FString>(Attribute2Name, TEXT("Default"), true, true);
+
+	check(Attribute1 && Attribute2);
+
+	TPCGValueRange<int64> MetadataEntryRange = RootPointData->GetMetadataEntryValueRange();
+	for (int32 i = 0; i < RootPointData->GetNumPoints(); ++i)
+	{
+		RootPointData->Metadata->InitializeOnSet(MetadataEntryRange[i]);
+		Attribute1->SetValue(MetadataEntryRange[i], static_cast<float>(i) * 0.1f);
+
+		if (i % 2 == 0)
+		{
+			// Will be either "0" or "1"
+			Attribute2->SetValue(MetadataEntryRange[i], FString::Printf(TEXT("%d"), i % 4));
+		}
+	}
+
+	// At the end of the first set, metadata has 10 entries, and values for each points are
+	// for Attribute 1: [0.0, 0.1, 0.2, 0.3, 0.4, 0.5, 0.6, 0.7, 0.8, 0.9]
+	// for Attribute 2: ["0", "Default", "2", "Default", "0", ...]
+	SECTION("RootMetadata has the right number of entries")
+	{
+		REQUIRE_EQUAL(RootPointData->Metadata->GetItemCountForChild(), NumPoints);
+	}
+
+	UPCGBasePointData* FirstChildPointData = Cast<UPCGBasePointData>(RootPointData->DuplicateData(nullptr));
+	Attribute1 = FirstChildPointData->Metadata->GetMutableTypedAttribute<float>(Attribute1Name);
+	Attribute2 = FirstChildPointData->Metadata->GetMutableTypedAttribute<FString>(Attribute2Name);
+
+	SECTION("Attributes exists in first child")
+	{
+		REQUIRE_NOT_EQUAL(Attribute1, nullptr);
+		REQUIRE_NOT_EQUAL(Attribute2, nullptr);
+	}
+	
+	check(Attribute1 && Attribute2);
+
+	// Override all the values for Attribute 1, and replace all "1" by "0" in attribute 2
+	// Also add a third attribute
+	FPCGMetadataAttribute<int>* Attribute3 = FirstChildPointData->Metadata->CreateAttribute<int>(Attribute3Name, -1, true, true);
+	check(Attribute3);
+
+	SECTION("First child has right number of points")
+	{
+		REQUIRE_EQUAL(FirstChildPointData->GetNumPoints(), NumPoints);
+	}
+
+	TPCGValueRange<int64> FirstChildMetadataEntryRange = FirstChildPointData->GetMetadataEntryValueRange();
+	for (int32 i = 0; i < FirstChildPointData->GetNumPoints(); ++i)
+	{
+		FirstChildPointData->Metadata->InitializeOnSet(FirstChildMetadataEntryRange[i]);
+		Attribute1->SetValue(FirstChildMetadataEntryRange[i], static_cast<float>(i) * 1.1f);
+
+		if (i % 2 == 0)
+		{
+			Attribute2->SetValue(FirstChildMetadataEntryRange[i], TEXT("0"));
+		}
+
+		Attribute3->SetValue(FirstChildMetadataEntryRange[i], i);
+	}
+
+	// At the end of the second set, metadata has 20 entries, and values for each points are
+	// for Attribute 1: [0.0, 1.1, 2.2, 3.3, 4.4, 5.5, 6.6, 7.7, 8.8, 9.9]
+	// for Attribute 2: ["0", "Default", "0", "Default", "0", ...]
+	// for Attribute 3: [0, 1, 2, 3, 4, 5, 6, 7, 8, 9]
+	SECTION("First child metadata has right number of entries")
+	{
+		REQUIRE_EQUAL(FirstChildPointData->Metadata->GetItemCountForChild(), 2 * NumPoints);
+	}
+	
+
+	// Second child is keeping the metadata entry for even numbers and reset metadata entry for the rest (point back to default). Also override Attribute 3 values for even numbers
+	UPCGBasePointData* SecondChildPointData = CreateData<UPCGPointArrayData>();
+
+	FPCGInitializeFromDataParams InitializeFromDataParams(FirstChildPointData);
+	InitializeFromDataParams.bInheritSpatialData = false;
+	SecondChildPointData->InitializeFromDataWithParams(InitializeFromDataParams);
+	
+	Attribute3 = SecondChildPointData->Metadata->GetMutableTypedAttribute<int>(Attribute3Name);
+	SECTION("Attributes exists in second child")
+	{
+		REQUIRE_NOT_EQUAL(Attribute3, nullptr);
+	}
+
+	check(Attribute3);
+
+	SecondChildPointData->SetNumPoints(FirstChildPointData->GetNumPoints());
+	UPCGBasePointData::SetPoints(FirstChildPointData, SecondChildPointData, {}, /*bCopyAll=*/true);
+
+	TPCGValueRange<int64> SecondChildMetadataEntryRange = SecondChildPointData->GetMetadataEntryValueRange();
+	for (int32 i = 0; i < SecondChildPointData->GetNumPoints(); ++i)
+	{
+		if (i % 2 == 0)
+		{
+			SecondChildPointData->Metadata->InitializeOnSet(SecondChildMetadataEntryRange[i]);
+			Attribute3->SetValue(SecondChildMetadataEntryRange[i], 10 * i);
+		}
+		else
+		{
+			SecondChildMetadataEntryRange[i] = PCGInvalidEntryKey;
+		}
+	}
+
+	// At the end of the third set, metadata has 25 entries, and values for each points are:
+	// for Attribute 1: [0.0, -0.1, 2.2, -0.1, 4.4, -0.1, 6.6, -0.1, 8.8, -0.1]
+	// for Attribute 2: ["0", "Default", "0", "Default", "0", ...]
+	// for Attribute 3: [0, -1, 20, -1, 40, -1, 60, -1, 80, -1]
+	SECTION("Second Child Metadata has the right number of entries")
+	{
+		REQUIRE_EQUAL(SecondChildPointData->Metadata->GetItemCountForChild(), 2 * NumPoints + NumPoints / 2);
+	}
+	
+
+	// For final set, duplicate the data and flatten it
+	UPCGBasePointData* FinalPointData = Cast<UPCGBasePointData>(SecondChildPointData->DuplicateData(nullptr));
+	FinalPointData->Flatten();
+
+	Attribute1 = FinalPointData->Metadata->GetMutableTypedAttribute<float>(Attribute1Name);
+	Attribute2 = FinalPointData->Metadata->GetMutableTypedAttribute<FString>(Attribute2Name);
+	Attribute3 = FinalPointData->Metadata->GetMutableTypedAttribute<int>(Attribute3Name);
+
+	SECTION("Attributes exists in final child")
+	{
+		REQUIRE_NOT_EQUAL(Attribute1, nullptr);
+		REQUIRE_NOT_EQUAL(Attribute2, nullptr);
+		REQUIRE_NOT_EQUAL(Attribute3, nullptr);
+	}
+	
+	check(Attribute1 && Attribute2 && Attribute3);
+
+	SECTION("Flatten metadata has the right number of entries/values")
+	{
+		REQUIRE_EQUAL(FinalPointData->Metadata->GetItemCountForChild(), 5);
+		REQUIRE_EQUAL(Attribute1->GetValueKeyOffsetForChild(), 5);
+		REQUIRE_EQUAL(Attribute2->GetValueKeyOffsetForChild(), 1);
+		REQUIRE_EQUAL(Attribute3->GetValueKeyOffsetForChild(), 5);
+		REQUIRE_EQUAL(FinalPointData->GetNumPoints(), NumPoints);
+	}
+
+	// Validate the values
+	const TConstPCGValueRange<int64> FinalMetadataEntryRange = FinalPointData->GetConstMetadataEntryValueRange();
+	auto i = GENERATE(range(0, NumPoints));
+	
+	SECTION("Validate values")
+	{
+		const int64 MetadataEntry = FinalMetadataEntryRange[i];
+		if (i % 2 == 0)
+		{
+			REQUIRE_EQUAL(MetadataEntry, static_cast<PCGMetadataEntryKey>(i / 2));
+			REQUIRE_EQUAL(Attribute1->GetValueFromItemKey(MetadataEntry), i * 1.1f);
+			REQUIRE_EQUAL(Attribute2->GetValueFromItemKey(MetadataEntry), TEXT("0"));
+			REQUIRE_EQUAL(Attribute3->GetValueFromItemKey(MetadataEntry), i * 10);
+		}
+		else
+		{
+			REQUIRE_EQUAL(MetadataEntry, PCGInvalidEntryKey);
+			REQUIRE_EQUAL(Attribute1->GetValueFromItemKey(MetadataEntry), -0.1f);
+			REQUIRE_EQUAL(Attribute2->GetValueFromItemKey(MetadataEntry), TEXT("Default"));
+			REQUIRE_EQUAL(Attribute3->GetValueFromItemKey(MetadataEntry), -1);
+		}
+	}
+}
+
+/**
+ * Flattening metadata that has no attribute still must be flattened. Same for Data domain.
+ */
+TEST_CASE_METHOD(PCGTests::FPCGBaseTestWithBasicContext, "PCG::Metadata::FlattenEmptyMetadata", "[PCG][Metadata][Flatten]")
+{
+	// Since the containers are meant to be created by the generic attributes, we'll use the generic attributes
+	// to populate them
+	
+	// Preparing the data
+	constexpr int32 NumPoints = 1;
+	UPCGPointArrayData* PointArrayData = CreateData<UPCGPointArrayData>();
+	PointArrayData->SetNumPoints(NumPoints);
+	
+	FPCGMetadataDomain* DataDomain = PointArrayData->MutableMetadata()->GetMetadataDomain(PCGMetadataDomainID::Data);
+	REQUIRE_NOT_EQUAL(DataDomain, nullptr);
+	
+	FPCGMetadataDomain* ElementsDomain = PointArrayData->MutableMetadata()->GetMetadataDomain(PCGMetadataDomainID::Elements);
+	REQUIRE_NOT_EQUAL(ElementsDomain, nullptr);
+	
+	const PCGMetadataEntryKey EntryKey = DataDomain->AddEntry();
+	
+	const FName AttributeName = TEXT("Double");
+	
+	FPCGMetadataAttribute<double>* Attribute = DataDomain->CreateAttribute<double>(AttributeName, 1.0f, false, false);
+	Attribute->SetValue(EntryKey, 5.0);
+	
+	REQUIRE_EQUAL(Attribute->GetNumberOfEntries(), 1);
+	
+	UPCGPointArrayData* ChildPointArrayData = CreateData<UPCGPointArrayData>();
+	ChildPointArrayData->InitializeFromDataWithParams(FPCGInitializeFromDataParams{PointArrayData});
+	
+	FPCGMetadataDomain* ChildDataDomain = ChildPointArrayData->MutableMetadata()->GetMetadataDomain(PCGMetadataDomainID::Data);
+	REQUIRE_NOT_EQUAL(ChildDataDomain, nullptr);
+	
+	FPCGMetadataDomain* ChildElementDomain = ChildPointArrayData->MutableMetadata()->GetMetadataDomain(PCGMetadataDomainID::Elements);
+	REQUIRE_NOT_EQUAL(ChildElementDomain, nullptr);
+	
+	// Data domain should have a single attribute, and inherit its parent
+	REQUIRE_EQUAL(ChildDataDomain->GetAttributeCount(), 1);
+	REQUIRE_EQUAL(ChildDataDomain->GetParent(), DataDomain);
+	
+	// Elements domain should have no attribute, and inherit its parent.
+	REQUIRE_EQUAL(ChildElementDomain->GetAttributeCount(), 0);
+	REQUIRE_EQUAL(ChildElementDomain->GetParent(), ElementsDomain);
+	
+	ChildPointArrayData->Flatten();
+	
+	// After flattening, we should still have the same attribute count by domain, but no parent
+	REQUIRE_EQUAL(ChildDataDomain->GetAttributeCount(), 1);
+	REQUIRE_EQUAL(ChildDataDomain->GetParent(), nullptr);
+	
+	REQUIRE_EQUAL(ChildElementDomain->GetAttributeCount(), 0);
+	REQUIRE_EQUAL(ChildElementDomain->GetParent(), nullptr);
+}
\ No newline at end of file
