[PCG] Followup for optional lock for generic attributes

* Factorized in common

#rnx

#Synced-CL 50157473

diff --git a/Engine/Plugins/PCG/Source/PCG/Private/Metadata/PCGMetadataAttributeGeneric.cpp b/Engine/Plugins/PCG/Source/PCG/Private/Metadata/PCGMetadataAttributeGeneric.cpp
--- a/Engine/Plugins/PCG/Source/PCG/Private/Metadata/PCGMetadataAttributeGeneric.cpp
+++ b/Engine/Plugins/PCG/Source/PCG/Private/Metadata/PCGMetadataAttributeGeneric.cpp
@@ -18,45 +18,6 @@ namespace PCG::Private
 			return *this;
 		}
 	};
-	
-	template <bool bWrite>
-	struct FOptionalLock
-	{
-		FOptionalLock(PCG::FSharedLock& InLock, bool bInLock)
-			: Lock(InLock)
-			, bLock(bInLock)
-		{
-			if (bLock)
-			{
-				if constexpr (bWrite)
-				{
-					Lock.Lock();
-				}
-				else
-				{
-					Lock.LockShared();
-				}
-			}
-		}
-		
-		~FOptionalLock()
-		{
-			if (bLock)
-			{
-				if constexpr (bWrite)
-				{
-					Lock.Unlock();
-				}
-				else
-				{
-					Lock.UnlockShared();
-				}
-			}
-		}
-		
-		PCG::FSharedLock& Lock;
-		bool bLock;
-	};
 
 	/**
 	 * Equivalent of the method in PropertyBag to convert the descriptor to a property.
@@ -497,7 +458,7 @@ TArray<PCGMetadataValueKey> FPCGMetadataAttributeGeneric::AddValues_Internal(FIn
 	}
 	
 	{
-		PCG::Private::FOptionalLock</*bWrite=*/true> WriteLock(ValueLock, !bLockless);
+		PCG::TUniqueScopeLock WriteLock(ValueLock, /*bShouldLock=*/!bLockless);
 		if (ExistingIndex == INDEX_NONE)
 		{
 			check(!bLockless);
@@ -698,7 +659,7 @@ void FPCGMetadataAttributeGeneric::FindValuesInternal(const FInValuesByValue& In
 
 	if (ValueKeysSet != InValues.Count)
 	{
-		PCG::Private::FOptionalLock</*bWrite=*/false> ReadLock(ValueLock, /*bLock=*/bIsRoot);
+		PCG::TSharedScopeLock ReadLock(ValueLock, /*bShouldLock=*/bIsRoot);
 		for (int ValueIndex = 0; ValueIndex < InValues.Count; ++ValueIndex)
 		{
 			if (ValueKeys[ValueIndex] != PCGNotFoundValueKey)
@@ -858,7 +819,7 @@ FPCGMetadataAttributeBase* FPCGMetadataAttributeGeneric::Copy(FName NewName, FPC
 			const FPCGMetadataAttributeGeneric* Current = Parents[i];
 
 			// Only need to lock if Current == this, since the parents are assumed to be const.
-			PCG::Private::FOptionalLock</*bWrite=*/false> ReadLock(Current->EntryMapLock, /*bLock=*/Current == this);
+			PCG::TSharedScopeLock ReadLock(Current->EntryMapLock, /*bShouldLock=*/Current == this);
 			AttributeCopy->EntryToValueKeyMap.Append(Current->EntryToValueKeyMap);
 		}
 	}
@@ -874,7 +835,7 @@ FPCGMetadataAttributeBase* FPCGMetadataAttributeGeneric::Copy(FName NewName, FPC
 				const FPCGMetadataAttributeGeneric* Current = Parents[i];
 
 				// Only need to lock if Current == this, since the parents are assumed to be const.
-				PCG::Private::FOptionalLock</*bWrite=*/false> ReadLock(Current->ValueLock, /*bLock=*/Current == this);
+				PCG::TSharedScopeLock ReadLock(Current->ValueLock, /*bShouldLock=*/Current == this);
 
 				// Do not copy the first value, as it is the default value, already copied/initialized by the ctor
 				TOptional<FScriptArrayHelper> ParentHelper = Current->MakeArrayHelper();
@@ -1371,12 +1332,12 @@ int32 FPCGMetadataAttributeGeneric::PreallocateValues(TArrayView<PCGMetadataEntr
 
 	if (!UnderlyingProperty->bCompressData)
 	{
-		PCG::Private::FOptionalLock</*bWrite=*/true> WriteValueLock(ValueLock, !bLockless);
+		PCG::TUniqueScopeLock WriteValueLock(ValueLock, /*bShouldLock=*/!bLockless);
 		StartIndex = AllocateValues_Unsafe(EntryKeys.Num());
 	}
 
 	{
-		PCG::Private::FOptionalLock</*bWrite=*/true> WriteEntryMapLock(EntryMapLock, !bLockless);
+		PCG::TUniqueScopeLock WriteEntryMapLock(EntryMapLock, /*bShouldLock=*/!bLockless);
 		EntryToValueKeyMap.Reserve(EntryToValueKeyMap.Num() + EntryKeys.Num());
 
 		if (!UnderlyingProperty->bCompressData && StartIndex != INDEX_NONE)
diff --git a/Engine/Plugins/PCG/Source/PCG/Public/PCGCommon.h b/Engine/Plugins/PCG/Source/PCG/Public/PCGCommon.h
--- a/Engine/Plugins/PCG/Source/PCG/Public/PCGCommon.h
+++ b/Engine/Plugins/PCG/Source/PCG/Public/PCGCommon.h
@@ -43,16 +43,19 @@ namespace PCG
 	// Similar to existing UE::TScopeLock but allows optional TryLock in which case caller is responsible to 
 	// check if Lock succeeded or not with IsLocked(). This allows execution code to be skipped if non main thread
 	// tasks fail to get the lock instead of waiting.
+	// Also add an option to have an optional lock, if we have a code path that support both lock and lockless.
 	template<typename LockType> requires (!IsSharedLock<LockType>)
 	class TScopeLock final
 	{
 	public:
 		UE_NONCOPYABLE(TScopeLock);
 
-		UE_NODISCARD_CTOR inline explicit TScopeLock(LockType* InLock, bool bTryLock = false)
-			: Lock(InLock), bLocked(false)
+		UE_NODISCARD_CTOR inline explicit TScopeLock(LockType* InLock, bool bTryLock = false, bool bInShouldLock = true)
+			: Lock(InLock)
+			, bLocked(false)
+			, bShouldLock(bInShouldLock)
 		{
-			if (InLock)
+			if (InLock && bShouldLock)
 			{
 				if (bTryLock)
 				{
@@ -67,8 +70,8 @@ namespace PCG
 		}
 
 
-		UE_NODISCARD_CTOR inline explicit TScopeLock(LockType& InLock, bool bTryLock = false)
-			: TScopeLock(&InLock, bTryLock)
+		UE_NODISCARD_CTOR inline explicit TScopeLock(LockType& InLock, bool bTryLock = false, bool bInShouldLock = true)
+			: TScopeLock(&InLock, bTryLock, bInShouldLock)
 		{
 		}
 
@@ -89,6 +92,7 @@ namespace PCG
 	private:
 		LockType* Lock;
 		bool bLocked;
+		bool bShouldLock;
 	};
 
 	template <typename SharedLockType> requires (IsSharedLock<SharedLockType>)
@@ -97,19 +101,27 @@ namespace PCG
 	public:
 		UE_NONCOPYABLE(TSharedScopeLock);
 
-		UE_NODISCARD_CTOR inline explicit TSharedScopeLock(SharedLockType& InSharedLock)
+		UE_NODISCARD_CTOR inline explicit TSharedScopeLock(SharedLockType& InSharedLock, bool bInShouldLock = true)
 			: SharedLock(InSharedLock)
+			, bShouldLock(bInShouldLock)
 		{
-			SharedLock.LockShared();
+			if (bShouldLock)
+			{
+				SharedLock.LockShared();
+			}
 		}
 
 		inline ~TSharedScopeLock()
 		{
-			SharedLock.UnlockShared();
+			if (bShouldLock)
+			{
+				SharedLock.UnlockShared();
+			}
 		}
 
 	private:
 		SharedLockType& SharedLock;
+		bool bShouldLock;
 	};
 
 	template <typename SharedLockType> requires (IsSharedLock<SharedLockType>)
@@ -118,19 +130,27 @@ namespace PCG
 	public:
 		UE_NONCOPYABLE(TUniqueScopeLock);
 
-		UE_NODISCARD_CTOR inline explicit TUniqueScopeLock(SharedLockType& InSharedLock)
+		UE_NODISCARD_CTOR inline explicit TUniqueScopeLock(SharedLockType& InSharedLock, bool bInShouldLock = true)
 			: SharedLock(InSharedLock)
+			, bShouldLock(bInShouldLock)
 		{
-			SharedLock.Lock();
+			if (bShouldLock)
+			{
+				SharedLock.Lock();
+			}
 		}
 
 		inline ~TUniqueScopeLock()
 		{
-			SharedLock.Unlock();
+			if (bShouldLock)
+			{
+				SharedLock.Unlock();
+			}
 		}
 
 	private:
 		SharedLockType& SharedLock;
+		bool bShouldLock;
 	};
 }
 
