[PCG] Restore Cached Pointer for Spline mesh managed resource as it is expensive to resolve the soft object and it might be expected to have a lot of spline mesh components

* Make it a Weak instead of a raw pointer to be resilient with BP reinstanciation
* Make it thread-safe

#rnx

#Synced-CL 50157473

diff --git a/Engine/Plugins/PCG/Source/PCG/Private/PCGManagedResource.cpp b/Engine/Plugins/PCG/Source/PCG/Private/PCGManagedResource.cpp
--- a/Engine/Plugins/PCG/Source/PCG/Private/PCGManagedResource.cpp
+++ b/Engine/Plugins/PCG/Source/PCG/Private/PCGManagedResource.cpp
@@ -1205,13 +1205,50 @@ void UPCGManagedISKMComponent::SetComponent(UInstancedSkinnedMeshComponent* InCo
 	GeneratedComponent = InComponent;
 }
 
+void UPCGManagedSplineMeshComponent::PostLoad()
+{
+	Super::PostLoad();
+
+	// Cache raw ptr
+	CacheAndGetComponent();
+}
+
+void UPCGManagedSplineMeshComponent::ForgetComponent()
+{
+	Super::ForgetComponent();
+	CachedGeneratedComponent.Reset();
+}
+
 USplineMeshComponent* UPCGManagedSplineMeshComponent::GetComponent() const
 {
-	TRACE_CPUPROFILER_EVENT_SCOPE(UPCGManagedSplineMeshComponent::GetComponent);
-	return Cast<USplineMeshComponent>(GeneratedComponent.Get());
+	// Read first, if it is cached, return it.
+	{
+		PCG::TSharedScopeLock ReadLock(CachedLock);
+		if (USplineMeshComponent* GeneratedComponentPtr = CachedGeneratedComponent.Get())
+		{
+			return GeneratedComponentPtr;
+		}
+	}
+	
+	// Otherwise, lock on write, check again and cache it if needed.
+	return CacheAndGetComponent();
 }
 
 void UPCGManagedSplineMeshComponent::SetComponent(USplineMeshComponent* InComponent)
 {
 	GeneratedComponent = InComponent;
+	CachedGeneratedComponent = InComponent;
+}
+
+USplineMeshComponent* UPCGManagedSplineMeshComponent::CacheAndGetComponent() const
+{
+	PCG::TUniqueScopeLock WriteLock(CachedLock);
+	USplineMeshComponent* GeneratedComponentPtr = CachedGeneratedComponent.Get();
+	if (!GeneratedComponentPtr)
+	{
+		GeneratedComponentPtr = Cast<USplineMeshComponent>(GeneratedComponent.Get());
+		CachedGeneratedComponent = GeneratedComponentPtr;
+	}
+		
+	return GeneratedComponentPtr;
 }
diff --git a/Engine/Plugins/PCG/Source/PCG/Public/PCGManagedResource.h b/Engine/Plugins/PCG/Source/PCG/Public/PCGManagedResource.h
--- a/Engine/Plugins/PCG/Source/PCG/Public/PCGManagedResource.h
+++ b/Engine/Plugins/PCG/Source/PCG/Public/PCGManagedResource.h
@@ -463,7 +463,12 @@ class UPCGManagedSplineMeshComponent : public UPCGManagedComponent
 	GENERATED_BODY()
 
 public:
+	//~Begin UObject interface
+	UE_API virtual void PostLoad() override;
+	//~End UObject interface
+
 	//~Begin UPCGManagedComponents interface
+	UE_API virtual void ForgetComponent() override;
 	virtual void ResetComponent() override { /* Does nothing, but implementation is required to support reuse. */ }
 	virtual bool SupportsComponentReset() const override { return true; }
 	//~End UPCGManagedComponents interface
@@ -495,6 +500,14 @@ protected:
 
 	UPROPERTY()
 	FPCGCrc SettingsCrc;
+	
+private:
+	USplineMeshComponent* CacheAndGetComponent() const;
+	
+	UPROPERTY(Transient)
+	mutable TWeakObjectPtr<USplineMeshComponent> CachedGeneratedComponent;
+
+	mutable PCG::FSharedLock CachedLock;
 };
 
 #undef UE_API
