[PCG] Produce spatial hashes/culling cells that define bounds for ranges of instances and allow early culling within the renderer.

This saves a large amount of memory on low end devices.

We already run an analysis pass on the input points to static mesh spawner to count how many instances per primitive. This change expands that functionality to also allow bucketing the instances into the cells at the same time.

[RN] [PCG] Added experimental support for passing bounds data for subranges of instances in a primitive to the renderer which can save substantial memory allocated for draw calls on low end

#Synced-CL 50382473

diff --git a/Engine/Plugins/Experimental/PCGInterops/PCGFastGeoInterop/Source/PCGFastGeoInterop/Private/Compute/PrimitiveFactories/PCGPrimitiveFactoryFastGeoPISMC.cpp b/Engine/Plugins/Experimental/PCGInterops/PCGFastGeoInterop/Source/PCGFastGeoInterop/Private/Compute/PrimitiveFactories/PCGPrimitiveFactoryFastGeoPISMC.cpp
--- a/Engine/Plugins/Experimental/PCGInterops/PCGFastGeoInterop/Source/PCGFastGeoInterop/Private/Compute/PrimitiveFactories/PCGPrimitiveFactoryFastGeoPISMC.cpp
+++ b/Engine/Plugins/Experimental/PCGInterops/PCGFastGeoInterop/Source/PCGFastGeoInterop/Private/Compute/PrimitiveFactories/PCGPrimitiveFactoryFastGeoPISMC.cpp
@@ -18,7 +18,7 @@ bool FPCGPrimitiveFactoryFastGeoPISMC::IsRenderStateCreated() const
 {
 	if (Components.IsEmpty())
 	{
-		return false;
+		return true;
 	}
 
 	// Use first component to get to container.
@@ -27,7 +27,7 @@ bool FPCGPrimitiveFactoryFastGeoPISMC::IsRenderStateCreated() const
 
 	if (!Container)
 	{
-		return false;
+		return true;
 	}
 
 	if (Container->HasAnyPendingTasks())
@@ -64,14 +64,40 @@ FPrimitiveSceneProxy* FPCGPrimitiveFactoryFastGeoPISMC::GetSceneProxy(int32 InPr
 	}
 }
 
-int32 FPCGPrimitiveFactoryFastGeoPISMC::GetNumInstances(int32 InPrimitiveIndex) const
+int32 FPCGPrimitiveFactoryFastGeoPISMC::GetNumInstances(int32 InPrimitiveIndex, int32 InCellID) const
 {
-	return ensure(InstanceCounts.IsValidIndex(InPrimitiveIndex)) ? InstanceCounts[InPrimitiveIndex] : 0;
+	if (ensure(PrimitiveInfos.IsValidIndex(InPrimitiveIndex)))
+	{
+		for (const FPCGInstanceRange& PrimitiveInfo : PrimitiveInfos[InPrimitiveIndex].InstanceRanges)
+		{
+			if (PrimitiveInfo.GetCellID() == InCellID)
+			{
+				return PrimitiveInfo.GetNumInstances();
+			}
+		}
+	}
+
+	return 0;
+}
+
+int32 FPCGPrimitiveFactoryFastGeoPISMC::GetNumInstancesTotal(int32 InPrimitiveIndex) const
+{
+	int32 NumInstances = 0;
+
+	if (ensure(PrimitiveInfos.IsValidIndex(InPrimitiveIndex)))
+	{
+		for (const FPCGInstanceRange& PrimitiveInfo : PrimitiveInfos[InPrimitiveIndex].InstanceRanges)
+		{
+			NumInstances += PrimitiveInfo.GetNumInstances();
+		}
+	}
+
+	return NumInstances;
 }
 
 void FPCGPrimitiveFactoryFastGeoPISMC::Initialize(FParameters&& InParameters)
 {
-	Descriptors = MoveTemp(InParameters.Descriptors);
+	PrimitiveInfos = MoveTemp(InParameters.PrimitiveInfos);
 	CustomPrimitiveData = MoveTemp(InParameters.CustomPrimitiveData);
 }
 
@@ -100,17 +126,29 @@ bool FPCGPrimitiveFactoryFastGeoPISMC::Create(FPCGContext* InContext)
 	const FString DebugName = ExecutionSource->GetExecutionState().GetDebugName();
 	TUniquePtr<FFastGeoComponentCluster> ComponentCluster = MakeUnique<FFastGeoComponentCluster>(FastGeo, *FString::Printf(TEXT("FastGeoComponentCluster_%s"), *DebugName));
 
-	const int32 NumToCreate = FMath::Min(Descriptors.Num(), PCGComputeConstants::MAX_PRIMITIVE_COMPONENTS_PER_SPAWNER);
+	const int32 NumToCreate = FMath::Min(PrimitiveInfos.Num(), PCGComputeConstants::MAX_PRIMITIVE_COMPONENTS_PER_SPAWNER);
 
 	MeshBounds.SetNumUninitialized(NumToCreate);
 
 	for (int32 DescIndex = 0; DescIndex < NumToCreate; ++DescIndex)
 	{
-		const FPCGProceduralISMComponentDescriptor& Desc = Descriptors[DescIndex];
+		FPCGPrimitiveInfo& PrimitiveInfo = PrimitiveInfos[DescIndex];
+		FPCGProceduralISMComponentDescriptor& Desc = PrimitiveInfo.Descriptor;
+
+		TArray<FInstanceSceneDataBuffers::FCompressedSpatialHashItem> SpatialHashes;
+		FBox BoundsOfAllCells(EForceInit::ForceInit);
+
+		PopulateHashes(PrimitiveInfo, SpatialHashes, BoundsOfAllCells);
+
+		if (BoundsOfAllCells.IsValid)
+		{
+			Desc.WorldBounds = Desc.WorldBounds.Overlap(BoundsOfAllCells);
+		}
 
 		FFastGeoProceduralISMComponent& FastGeoComponent = static_cast<FFastGeoProceduralISMComponent&>(ComponentCluster->AddComponent(FFastGeoProceduralISMComponent::Type));
 		FastGeoComponent.InitializeFromComponentDescriptor(Desc);
 		FastGeoComponent.SetCustomPrimitiveData(CustomPrimitiveData.GetViewForIndex(DescIndex));
+		FastGeoComponent.SetSpatialHashes(MoveTemp(SpatialHashes));
 
 		MeshBounds[DescIndex] = Desc.StaticMesh->GetBoundingBox();
 	}
@@ -131,8 +169,8 @@ bool FPCGPrimitiveFactoryFastGeoPISMC::Create(FPCGContext* InContext)
 	{
 		InComponentCluster.ForEachComponent([This, &ComponentIndex](FFastGeoComponent& InComponent)
 		{
-			check(This->Descriptors.IsValidIndex(ComponentIndex));
-			const FPCGProceduralISMComponentDescriptor& Desc = This->Descriptors[ComponentIndex];
+			check(This->PrimitiveInfos.IsValidIndex(ComponentIndex));
+			const FPCGProceduralISMComponentDescriptor& Desc = This->PrimitiveInfos[ComponentIndex].Descriptor;
 
 			// If not true, weakptr assignment below will fail.
 			check(InComponent.GetOwnerComponentCluster());
@@ -164,10 +202,12 @@ FBox FPCGPrimitiveFactoryFastGeoPISMC::GetMeshBounds(int32 InPrimitiveIndex) con
 TArray<TObjectPtr<UObject>> FPCGPrimitiveFactoryFastGeoPISMC::CollectObjectReferences()
 {
 	TArray<TObjectPtr<UObject>> ObjectReferences;
-	ObjectReferences.Reserve(Descriptors.Num() * 4); // Heuristic
+	ObjectReferences.Reserve(PrimitiveInfos.Num() * 4); // Heuristic
 
-	for (const FPCGProceduralISMComponentDescriptor& Descriptor : Descriptors)
+	for (const FPCGPrimitiveInfo& PrimitiveInfo : PrimitiveInfos)
 	{
+		const FPCGProceduralISMComponentDescriptor& Descriptor = PrimitiveInfo.Descriptor;
+
 		ObjectReferences.Add(Descriptor.StaticMesh.Get());
 
 		if (Descriptor.OverlayMaterial)
diff --git a/Engine/Plugins/Experimental/PCGInterops/PCGFastGeoInterop/Source/PCGFastGeoInterop/Private/Compute/PrimitiveFactories/PCGPrimitiveFactoryFastGeoPISMC.h b/Engine/Plugins/Experimental/PCGInterops/PCGFastGeoInterop/Source/PCGFastGeoInterop/Private/Compute/PrimitiveFactories/PCGPrimitiveFactoryFastGeoPISMC.h
--- a/Engine/Plugins/Experimental/PCGInterops/PCGFastGeoInterop/Source/PCGFastGeoInterop/Private/Compute/PrimitiveFactories/PCGPrimitiveFactoryFastGeoPISMC.h
+++ b/Engine/Plugins/Experimental/PCGInterops/PCGFastGeoInterop/Source/PCGFastGeoInterop/Private/Compute/PrimitiveFactories/PCGPrimitiveFactoryFastGeoPISMC.h
@@ -19,7 +19,8 @@ public:
 	virtual bool IsRenderStateCreated() const override;
 	virtual int32 GetNumPrimitives() const override { return Components.Num(); }
 	virtual FPrimitiveSceneProxy* GetSceneProxy(int32 InPrimitiveIndex) const override;
-	virtual int32 GetNumInstances(int32 InPrimitiveIndex) const override;
+	virtual int32 GetNumInstances(int32 InPrimitiveIndex, int32 InCellID) const override;
+	virtual int32 GetNumInstancesTotal(int32 InPrimitiveIndex) const override;
 	//~ End IPCGRuntimePrimitiveFactory interface
 
 	//~ Begin IPCGPrimitiveFactoryISMBase interface
@@ -32,7 +33,7 @@ protected:
 	TArray<TObjectPtr<UObject>> CollectObjectReferences();
 
 protected:
-	TArray<FPCGProceduralISMComponentDescriptor> Descriptors;
+	TArray<FPCGPrimitiveInfo> PrimitiveInfos;
 	FPCGPackedCustomData CustomPrimitiveData;
 
 	TArray<FWeakFastGeoComponent> Components;
diff --git a/Engine/Plugins/PCG/Shaders/Private/BuiltInKernels/PCGCountUniqueAttributeValues.usf b/Engine/Plugins/PCG/Shaders/Private/BuiltInKernels/PCGCountUniqueAttributeValues.usf
--- a/Engine/Plugins/PCG/Shaders/Private/BuiltInKernels/PCGCountUniqueAttributeValues.usf
+++ b/Engine/Plugins/PCG/Shaders/Private/BuiltInKernels/PCGCountUniqueAttributeValues.usf
@@ -8,10 +8,10 @@
 // Output (PCG_OUTPUT_TO_RAW_BUFFER): Uint array consisting of (string key value, value instance count) pairs. Number of pairs is the number of unique values (known CPU-side).
 // Output (!PCG_OUTPUT_TO_RAW_BUFFER): Attribute set with a ValueCount attribute which has the count of each value. Number of elements is the number of unique values (known CPU-side).
 
-void InitializeValueKey(uint InDataIndex, uint InValueKey)
+void InitializeValueKey(uint InDataIndex, uint InValueKey, uint InBucketIndex)
 {
 #if PCG_OUTPUT_TO_RAW_BUFFER
-	const uint FirstElementIndex = InValueKey * 2u;
+	const uint FirstElementIndex = (InValueKey * CountUniqueValues_GetNumBucketsTotal() + InBucketIndex) * 2u;
 	if (FirstElementIndex < Out_GetNumElements(0))
 	{
 		// We currently have an output element per possible string key value. This may be compacted in the future.
@@ -19,8 +19,9 @@ void InitializeValueKey(uint InDataIndex, uint InValueKey)
 	}
 #else
 	// Write value of attribute. We currently have an output element per possible string key value.
+	const uint ElementIndex = InValueKey * CountUniqueValues_GetNumBucketsTotal() + InBucketIndex;
  	const int OutputValueAttributeId = CountUniqueValues_GetOutputValueAttributeId();
- 	Out_SetInt(InDataIndex, InValueKey, OutputValueAttributeId, InValueKey);
+ 	Out_SetInt(InDataIndex, ElementIndex, OutputValueAttributeId, InValueKey);
 #endif
 }
 
@@ -37,19 +38,20 @@ void PCGCountUniqueAttributeValuesCS(uint3 GroupId : SV_GroupID, uint GroupIndex
 #endif
 
 	const uint ThreadIndex = GetUnWrappedDispatchThreadId(GroupId, GroupIndex, 64);
+	const int NumBucketsTotal = CountUniqueValues_GetNumBucketsTotal(); // Should always be greater than 0
 
 	// Pre-job - if this thread maps to any element in the output, initialize the value key. The dispatch thread count will be large enough to ensure
 	// all output elements are initialized.
 	{
 #if PCG_OUTPUT_TO_RAW_BUFFER
 		// If outputting to raw buffer, always single data.
-		InitializeValueKey(/*DataIndex=*/0, ThreadIndex);
+		InitializeValueKey(/*DataIndex=*/0, ThreadIndex / NumBucketsTotal, ThreadIndex % NumBucketsTotal);
 #else
 		uint Out_DataIndex;
 		uint Out_ElementIndex;
 		if (Out_GetThreadData(ThreadIndex, Out_DataIndex, Out_ElementIndex))
 		{
-			InitializeValueKey(Out_DataIndex, Out_ElementIndex);
+			InitializeValueKey(Out_DataIndex, Out_ElementIndex / NumBucketsTotal, Out_ElementIndex % NumBucketsTotal);
 		}
 #endif
 	}
@@ -81,12 +83,25 @@ void PCGCountUniqueAttributeValuesCS(uint3 GroupId : SV_GroupID, uint GroupIndex
 		// If not outputting per-data counts, then all counters are aggregated on a single data.
 		const uint Out_DataIndex = CountUniqueValues_GetEmitPerDataCounts() ? In_DataIndex : 0;
 
+		uint BucketIndex = 0;
+		if (NumBucketsTotal > 1)
+		{
+			const float3 Pos = In_GetPosition(In_DataIndex, In_ElementIndex);
+			const int3 BucketCoord = int3(floor(Pos / CountUniqueValues_GetBucketExtent()));
+			int3 LocalBucketCoord = BucketCoord - int3(floor(CountUniqueValues_GetBoundsMin() / CountUniqueValues_GetBucketExtent()));
+
+			const int3 NumBuckets = CountUniqueValues_GetNumBuckets();
+			LocalBucketCoord = clamp(LocalBucketCoord, 0, NumBuckets - 1);
+
+			BucketIndex = (NumBuckets.y * NumBuckets.x * LocalBucketCoord.z) + (NumBuckets.x * LocalBucketCoord.y) + LocalBucketCoord.x;
+		}
+
 #if PCG_OUTPUT_TO_RAW_BUFFER
-		const uint ElementIndex = AttributeValue * 2u + 1;
+		const uint ElementIndex = (AttributeValue * NumBucketsTotal + BucketIndex) * 2u + 1;
 		Out_AtomicAdd(/*DataIndex=*/0, ElementIndex, 1u);
 #else
 		// Always one bucket per attribute value.
-		const uint Out_ElementIndex = AttributeValue;
+		const uint Out_ElementIndex = AttributeValue * NumBucketsTotal + BucketIndex;
 
 		const int OutputCountAttributeId = CountUniqueValues_GetOutputCountAttributeId();
 
diff --git a/Engine/Plugins/PCG/Shaders/Private/Elements/PCGStaticMeshSpawner.usf b/Engine/Plugins/PCG/Shaders/Private/Elements/PCGStaticMeshSpawner.usf
--- a/Engine/Plugins/PCG/Shaders/Private/Elements/PCGStaticMeshSpawner.usf
+++ b/Engine/Plugins/PCG/Shaders/Private/Elements/PCGStaticMeshSpawner.usf
@@ -47,6 +47,7 @@ void PCGStaticMeshSpawnerCS(uint3 GroupId : SV_GroupID, uint GroupIndex : SV_Gro
 	}
 
 	uint PrimitiveIndex = (uint)-1;
+	uint CullingCellId = 0;
 
 	const int NumPrimitives = InstanceData_GetNumPrimitives();
 
@@ -81,13 +82,26 @@ void PCGStaticMeshSpawnerCS(uint3 GroupId : SV_GroupID, uint GroupIndex : SV_Gro
 				}
 			}
 		}
+
+		const int NumCullingCellsTotal = SMSpawner_GetNumCullingCellsTotal();
+		if (NumCullingCellsTotal > 1)
+		{
+			const float3 Pos = In_GetPosition(In_DataIndex, In_ElementIndex);
+			const int3 CullingCellCoord = int3(floor(Pos / SMSpawner_GetBucketExtent()));
+			int3 LocalCullingCellCoord = CullingCellCoord - int3(floor(SMSpawner_GetBoundsMin() / SMSpawner_GetBucketExtent()));
+
+			const int3 NumCullingCells = SMSpawner_GetNumCullingCells();
+			LocalCullingCellCoord = clamp(LocalCullingCellCoord, 0, NumCullingCells - 1);
+
+			CullingCellId = (NumCullingCells.y * NumCullingCells.x * LocalCullingCellCoord.z) + (NumCullingCells.x * LocalCullingCellCoord.y) + LocalCullingCellCoord.x;
+		}
 	}
 
 	// Get instance index to write to. Will return a negative value if we have saturated instances for this primitive, in which case we will not push the data.
 	int InstanceIndexInt = -1;
 	if (PrimitiveIndex != (uint)-1)
 	{
-		InstanceIndexInt = InstanceData_GetIndexToWriteTo(PrimitiveIndex);
+		InstanceIndexInt = InstanceData_GetIndexToWriteTo(PrimitiveIndex, CullingCellId);
 	}
 
 	// Write instance transform, only if not culled by primitive instance index
diff --git a/Engine/Plugins/PCG/Shaders/Private/PCGInstanceDataInterface.ush b/Engine/Plugins/PCG/Shaders/Private/PCGInstanceDataInterface.ush
--- a/Engine/Plugins/PCG/Shaders/Private/PCGInstanceDataInterface.ush
+++ b/Engine/Plugins/PCG/Shaders/Private/PCGInstanceDataInterface.ush
@@ -1,10 +1,11 @@
 // Copyright Epic Games, Inc. All Rights Reserved.
 
-uint4 {DataInterfaceName}_PrimitiveParameters[PCG_MAX_PRIMITIVES]; // (NumInstancesAllocated, InstanceOffset, Unused, Unused)
-
 uint {DataInterfaceName}_NumPrimitives;
 uint {DataInterfaceName}_NumInstancesAllPrimitives;
 uint {DataInterfaceName}_NumCustomFloatsPerInstance;
+int {DataInterfaceName}_NumCullingCellsTotal;
+
+StructuredBuffer<uint2> {DataInterfaceName}_PrimitiveParameters; // Per primitive, per culling cell: (NumInstancesAllocated, InstanceOffset)
 
 RWStructuredBuffer<uint4> {DataInterfaceName}_InstanceData;
 RWStructuredBuffer<uint> {DataInterfaceName}_InstanceCustomFloatData;
@@ -20,19 +21,23 @@ uint InstanceData_GetNumInstancesAllPrimitives_{DataInterfaceName}()
 	return {DataInterfaceName}_NumInstancesAllPrimitives;
 }
 
-int InstanceData_GetIndexToWriteTo_{DataInterfaceName}(uint InPrimitiveIndex)
+int InstanceData_GetIndexToWriteTo_{DataInterfaceName}(uint InPrimitiveIndex, uint InCullingCellId)
 {
+	const uint Index = {DataInterfaceName}_NumCullingCellsTotal * InPrimitiveIndex + InCullingCellId;
+
 	uint ValueBefore;
-	InterlockedAdd({DataInterfaceName}_WriteCounters[InPrimitiveIndex], 1u, ValueBefore);
+	InterlockedAdd({DataInterfaceName}_WriteCounters[Index], 1u, ValueBefore);
+
+	uint2 PrimitiveParams = {DataInterfaceName}_PrimitiveParameters[Index];
 
-	const uint NumInstancesAllocated = {DataInterfaceName}_PrimitiveParameters[InPrimitiveIndex][0];
+	const uint NumInstancesAllocated = PrimitiveParams[0];
 	if (ValueBefore >= NumInstancesAllocated)
 	{
-		// All instances for this primitive have been written to.
+		// All instances for this primitive have been written to. Should not happen if instance counts were computed correctly.
 		return -1;
 	}
 
-	const uint InstanceOffset = {DataInterfaceName}_PrimitiveParameters[InPrimitiveIndex][1];
+	const uint InstanceOffset = PrimitiveParams[1];
 	return (int)(ValueBefore + InstanceOffset);
 }
 
diff --git a/Engine/Plugins/PCG/Shaders/Private/PCGSceneWriter.usf b/Engine/Plugins/PCG/Shaders/Private/PCGSceneWriter.usf
--- a/Engine/Plugins/PCG/Shaders/Private/PCGSceneWriter.usf
+++ b/Engine/Plugins/PCG/Shaders/Private/PCGSceneWriter.usf
@@ -17,6 +17,7 @@ uint InPrimitiveId;
 uint InCustomDataCount;
 uint InPayloadDataFlags;
 uint InSeed;
+uint InNumCullingCellsTotal;
 
 // This follows the implementation in PCGHelpers.
 uint ComputeSeed(uint A, uint B)
@@ -55,8 +56,15 @@ void SceneWriter_CS(const uint3 GroupId : SV_GroupID, const uint GroupIndex : SV
 	
 	const uint InstanceId = PrimitiveData.InstanceSceneDataOffset + LocalInstanceId;
 	
+	// Loop over all culling cells for this primitive and add up total instance count for this primitive.
+	uint InstanceCount = 0;
+	for (uint CullingCellIndex = 0; CullingCellIndex < InNumCullingCellsTotal; ++CullingCellIndex)
+	{
+		InstanceCount += InWriteCounters[InNumCullingCellsTotal * InPrimitiveIndex + CullingCellIndex];
+	}
+
 	// Cull any instances at tail end (assume all valid instances are written).
-	if (LocalInstanceId >= InWriteCounters[InPrimitiveIndex])
+	if (LocalInstanceId >= InstanceCount)
 	{
 		CullInstance(InstanceId);
 		return;
diff --git a/Engine/Plugins/PCG/Source/PCG/Private/Components/PCGProceduralISMComponent.cpp b/Engine/Plugins/PCG/Source/PCG/Private/Components/PCGProceduralISMComponent.cpp
--- a/Engine/Plugins/PCG/Source/PCG/Private/Components/PCGProceduralISMComponent.cpp
+++ b/Engine/Plugins/PCG/Source/PCG/Private/Components/PCGProceduralISMComponent.cpp
@@ -341,6 +341,16 @@ void UPCGProceduralISMComponent::SetMinDrawDistance(int32 InMinDrawDistance)
 	}
 }
 
+void UPCGProceduralISMComponent::SetSpatialHashes(TArray<FInstanceSceneDataBuffers::FCompressedSpatialHashItem>&& InSpatialHashes)
+{
+	// A simple ordered compare of the hashes used as this will be set infrequently, and hash order impacts what cell merging in the renderer.
+	if (SpatialHashes != InSpatialHashes)
+	{
+		SpatialHashes = MoveTemp(InSpatialHashes);
+		MarkRenderStateDirty();
+	}
+}
+
 void UPCGProceduralISMComponent::GetSceneProxyDesc(FInstancedStaticMeshSceneProxyDesc& OutSceneProxyDesc) const
 {
 	const UWorld* World = GetWorld();
@@ -361,6 +371,11 @@ void UPCGProceduralISMComponent::GetSceneProxyDesc(FInstancedStaticMeshSceneProx
 		ProxyData.Flags.bHasPerInstanceCustomData = ProxyData.NumCustomDataFloats > 0;
 		ProxyData.Flags.bHasPerInstanceRandom = bHasPerInstanceRandom;
 
+		if (!SpatialHashes.IsEmpty())
+		{
+			InstanceSceneDataBuffers.SetImmutable(FInstanceSceneDataImmutable(SpatialHashes), AccessTag);
+		}
+
 		InstanceSceneDataBuffers.EndWriteAccess(AccessTag);
 		InstanceSceneDataBuffers.ValidateData();
 	}
diff --git a/Engine/Plugins/PCG/Source/PCG/Private/Components/PCGProceduralISMComponent.h b/Engine/Plugins/PCG/Source/PCG/Private/Components/PCGProceduralISMComponent.h
--- a/Engine/Plugins/PCG/Source/PCG/Private/Components/PCGProceduralISMComponent.h
+++ b/Engine/Plugins/PCG/Source/PCG/Private/Components/PCGProceduralISMComponent.h
@@ -93,6 +93,9 @@ public:
 	UFUNCTION(BlueprintCallable, Category = "Components|PCGProceduralISMC")
 	void GetMinDrawDistance(int32& OutMinDrawDistance) const { OutMinDrawDistance = InstanceMinDrawDistance; }
 
+	/** Set spatial hashes used by the renderer for early culling of subranges of instances. */
+	void SetSpatialHashes(TArray<FInstanceSceneDataBuffers::FCompressedSpatialHashItem>&& InSpatialHashes);
+
 	//~ Begin UActorComponent Interface
 	bool ShouldCreatePhysicsState() const override { return false; }
 	bool IsHLODRelevant() const override { return false; }
@@ -179,6 +182,9 @@ protected:
 	UE_DEPRECATED(5.6, "Use SetBounds to set world-space bounds instead.")
 	UFUNCTION(BlueprintCallable, Category = "Components|PCGProceduralISMC", meta = (DeprecatedProperty, DeprecationMessage = "Use SetBounds to set world-space bounds instead."))
 	void SetLocalBounds(const FBox& InWorldBounds) {}
+
+	/** Spatial hashes provide bounds for sub-ranges of instances within this primitive, enabling early CPU culling. */
+	TArray<FInstanceSceneDataBuffers::FCompressedSpatialHashItem> SpatialHashes;
 };
 
 struct FPCGProceduralISMCBuilderParameters
diff --git a/Engine/Plugins/PCG/Source/PCG/Private/Compute/BuiltInKernels/PCGCountUniqueAttributeValuesKernel.cpp b/Engine/Plugins/PCG/Source/PCG/Private/Compute/BuiltInKernels/PCGCountUniqueAttributeValuesKernel.cpp
--- a/Engine/Plugins/PCG/Source/PCG/Private/Compute/BuiltInKernels/PCGCountUniqueAttributeValuesKernel.cpp
+++ b/Engine/Plugins/PCG/Source/PCG/Private/Compute/BuiltInKernels/PCGCountUniqueAttributeValuesKernel.cpp
@@ -73,22 +73,101 @@ bool UPCGCountUniqueAttributeValuesKernel::IsKernelDataValid(FPCGContext* InCont
 				return false;
 			}
 
-			if (AttributeDesc.GetAttributeKey().GetType() != EPCGKernelAttributeType::Int
-				&& AttributeDesc.GetAttributeKey().GetType() != EPCGKernelAttributeType::StringKey)
+			const bool bRequireStringKeyAttribute = MaxNumUniqueValuesSource == EPCGMaxNumUniqueValuesSource::StringKeyValues;
+			const bool bTypeSupported = (!bRequireStringKeyAttribute && AttributeDesc.GetAttributeKey().GetType() == EPCGKernelAttributeType::Int)
+				|| AttributeDesc.GetAttributeKey().GetType() == EPCGKernelAttributeType::StringKey;
+
+			if (!bTypeSupported)
 			{
 				// Attribute value counting only currently supported for attributes of type Int or StringKey.
 				PCG_KERNEL_VALIDATION_ERR(Context, GetSettings(), FText::Format(
-					LOCTEXT("AttributeTypeInvalid", "Cannot count values for attribute '{0}', only attributes of type Int or String Key are currently supported."),
+					LOCTEXT("AttributeTypeInvalid", "Cannot count values for attribute '{0}'. Attribute type must either be Int, or can be String Key when counting via String Key."),
 					FText::FromName(AttributeName)));
 
 				return false;
 			}
+
+			if (MaxNumUniqueValuesSource == EPCGMaxNumUniqueValuesSource::UniqueValueTable)
+			{
+				const TSharedPtr<const FPCGDataCollectionDesc> InputTableDesc = DataBinding->GetCachedKernelPinDataDesc(this, PCGCountUniqueAttributeValuesConstants::UniqueValueTablePinLabel, /*bIsInputPin=*/true);
+				if (!InputTableDesc || InputTableDesc->ComputeTotalElementCount() == 0)
+				{
+					PCG_KERNEL_VALIDATION_ERR(InContext, GetSettings(), FText::Format(
+						LOCTEXT("NoInputDataToDriveUniqueValueCount", "No input data provided on kernel pin '{0}' to drive unique value count calculation."),
+						FText::FromName(PCGCountUniqueAttributeValuesConstants::UniqueValueTablePinLabel)));
+
+					return false;
+				}
+			}
+			else if (MaxNumUniqueValuesSource == EPCGMaxNumUniqueValuesSource::ExplicitMaxValue && MaxIntValue <= 0)
+			{
+				PCG_KERNEL_VALIDATION_ERR(InContext, GetSettings(), FText::Format(
+					LOCTEXT("InvalidExplicitAttributeValueCount", "Explicit maximum value count {0} is invalid, must be greater than 0."),
+					FText::AsNumber(MaxIntValue)));
+
+				return false;
+			}
 		}
 	}
 
 	return true;
 }
 
+int32 UPCGCountUniqueAttributeValuesKernel::ComputeNumCountersRequired(UPCGDataBinding* InBinding, TSharedPtr<const FPCGDataCollectionDesc> InInputDesc) const
+{
+	int32 MaxNumUniqueValues = INDEX_NONE;
+
+	if (MaxNumUniqueValuesSource == EPCGMaxNumUniqueValuesSource::StringKeyValues)
+	{
+		// If counted attribute is a stringkey, use max string key value as num elements.
+		const int32 StringKeyAttributeId = InBinding->GetAttributeId(AttributeName, EPCGKernelAttributeType::StringKey);
+
+		if (StringKeyAttributeId == INDEX_NONE)
+		{
+			return 0;
+		}
+
+		TArray<int32> UniqueStringKeyValues;
+		check(InInputDesc);
+		InInputDesc->GetUniqueStringKeyValues(StringKeyAttributeId, UniqueStringKeyValues);
+
+		// If the highest string key value is 3, we'll allocate 4 counters. String key value 0 is reserved for empty/null string.
+		const int32* MaxStringKeyValue = Algo::MaxElement(UniqueStringKeyValues);
+		MaxNumUniqueValues = (MaxStringKeyValue ? *MaxStringKeyValue : 0) + 1;
+	}
+	else if (MaxNumUniqueValuesSource == EPCGMaxNumUniqueValuesSource::ExplicitMaxValue)
+	{
+		if (MaxIntValue < 0)
+		{
+			UE_LOG(LogPCG, Warning, TEXT("PCGCountUniqueAttributeValuesKernel: number of unique values (%d) not set."), MaxIntValue);
+			return 0;
+		}
+
+		MaxNumUniqueValues = MaxIntValue + 1;
+	}
+	else if (MaxNumUniqueValuesSource == EPCGMaxNumUniqueValuesSource::UniqueValueTable)
+	{
+		const FPCGKernelPin TableKernelPin(GetKernelIndex(), PCGCountUniqueAttributeValuesConstants::UniqueValueTablePinLabel, /*bIsInput=*/true);
+		const TSharedPtr<const FPCGDataCollectionDesc> TableDesc = InBinding->ComputeKernelPinDataDesc(TableKernelPin);
+		if (!ensure(TableDesc))
+		{
+			return 0;
+		}
+
+		MaxNumUniqueValues = TableDesc->ComputeTotalElementCount();
+	}
+	else
+	{
+		UE_LOG(LogPCG, Error, TEXT("PCGCountUniqueAttributeValuesKernel: MaxNumUniqueValuesSource (%d) is not valid."), MaxNumUniqueValuesSource);
+		return 0;
+	}
+
+	const FIntVector SubCounters = InBinding->GetMaxNumCullingCells(GetBucketExtent());
+	MaxNumUniqueValues *= SubCounters.X * SubCounters.Y * SubCounters.Z;
+
+	return MaxNumUniqueValues;
+}
+
 TSharedPtr<const FPCGDataCollectionDesc> UPCGCountUniqueAttributeValuesKernel::ComputeOutputBindingDataDesc(FName InOutputPinLabel, UPCGDataBinding* InBinding) const
 {
 	check(InBinding);
@@ -110,46 +189,25 @@ TSharedPtr<const FPCGDataCollectionDesc> UPCGCountUniqueAttributeValuesKernel::C
 	TSharedPtr<FPCGDataCollectionDesc> OutDataDesc = FPCGDataCollectionDesc::MakeShared();
 	TArray<FPCGDataDesc>& OutDataDescs = OutDataDesc->GetDataDescriptionsMutable();
 
-	const int32 StringKeyAttributeId = InBinding->GetAttributeId(AttributeName, EPCGKernelAttributeType::StringKey);
-	const bool bAttributePresent = StringKeyAttributeId != INDEX_NONE || InBinding->GetAttributeId(AttributeName, EPCGKernelAttributeType::Int) != INDEX_NONE;
+	int32 NumCounters = 0;
 
-	if (bAttributePresent)
+	if (InBinding->GetAttributeId(AttributeName, EPCGKernelAttributeType::StringKey) != INDEX_NONE || InBinding->GetAttributeId(AttributeName, EPCGKernelAttributeType::Int) != INDEX_NONE)
 	{
-		int32 ElementCount = INDEX_NONE;
-
-		// If counted attribute is a stringkey, use max string key value as num elements.
-		if (StringKeyAttributeId != INDEX_NONE)
-		{
-			TArray<int32> UniqueStringKeyValues;
-			InputDesc->GetUniqueStringKeyValues(StringKeyAttributeId, UniqueStringKeyValues);
-
-			// If the highest string key value is 3, we'll allocate 4 counters. String key value 0 is reserved for empty/null string.
-			const int32* MaxStringKeyValue = Algo::MaxElement(UniqueStringKeyValues);
-			ElementCount = (MaxStringKeyValue ? *MaxStringKeyValue : 0) + 1;
-		}
-		else
-		{
-			if (MaxIntValue < 0)
-			{
-				UE_LOG(LogPCG, Warning, TEXT("PCGCountUniqueAttributeValuesKernel: number of unique values (%d) not set."), MaxIntValue);
-				return OutDataDesc;
-			}
-
-			ElementCount = MaxIntValue + 1;
-		}
-
-		check(ElementCount > 0);
+		NumCounters = ComputeNumCountersRequired(InBinding, InputDesc);
+	}
 
+	if (NumCounters > 0)
+	{
 		// Base data description object, copied N times below to populate the final description.
 		FPCGDataDesc CountDataDesc;
 		
 		if (bOutputRawBuffer)
 		{
-			CountDataDesc = FPCGDataDesc(FPCGDataTypeInfoRawBuffer::AsId(), ElementCount * 2);
+			CountDataDesc = FPCGDataDesc(FPCGDataTypeInfoRawBuffer::AsId(), NumCounters * 2);
 		}
 		else
 		{
-			CountDataDesc = FPCGDataDesc(EPCGDataType::Param, ElementCount);
+			CountDataDesc = FPCGDataDesc(EPCGDataType::Param, NumCounters);
 			CountDataDesc.AddAttribute(FPCGKernelAttributeKey(PCGCountUniqueAttributeValuesConstants::ValueAttributeName, EPCGKernelAttributeType::Int), InBinding);
 			CountDataDesc.AddAttribute(FPCGKernelAttributeKey(PCGCountUniqueAttributeValuesConstants::ValueCountAttributeName, EPCGKernelAttributeType::Int), InBinding);
 		}
@@ -206,20 +264,13 @@ int UPCGCountUniqueAttributeValuesKernel::ComputeThreadCount(const UPCGDataBindi
 {
 	if (const TSharedPtr<const FPCGDataCollectionDesc> InputPinDesc = InBinding->GetCachedKernelPinDataDesc(this, PCGPinConstants::DefaultInputLabel, /*bIsInput=*/true); ensure(InputPinDesc))
 	{
-		const int ElementCount = InputPinDesc->ComputeTotalElementCount();
+		const int InputElementCount = InputPinDesc->ComputeTotalElementCount();
 
-		int ValueCount = 0;
-		if (const int32 PartitionAttributeId = InBinding->GetAttributeId(AttributeName, EPCGKernelAttributeType::StringKey); PartitionAttributeId != INDEX_NONE)
-		{
-			ValueCount = InputPinDesc->GetNumStringKeyValues(PartitionAttributeId);
-		}
-		else
-		{
-			ValueCount = FMath::Max(ValueCount, MaxIntValue + 1);
-		}
+		const int32 PartitionAttributeId = InBinding->GetAttributeId(AttributeName, EPCGKernelAttributeType::StringKey);
+		const int OutputElementCount = ComputeNumCountersRequired(const_cast<UPCGDataBinding*>(InBinding), InputPinDesc);
 
 		// One thread per element, but also make sure we have enough threads to write each unique attribute value (rather than leaving uninitialized).
-		return FMath::Max(ElementCount, ValueCount);
+		return FMath::Max(InputElementCount, OutputElementCount);
 	}
 	else
 	{
@@ -241,7 +292,7 @@ FString UPCGCountUniqueAttributeValuesKernel::GetCookedSource(FPCGGPUCompilation
 	
 	TMap<FString, FStringFormatArg> TemplateArgs =
 	{
-		{ TEXT("RawDataOutputEnabled"), bOutputRawBuffer ? TEXT("1") : TEXT("0") },
+		{ TEXT("RawDataOutputEnabled"), bOutputRawBuffer ? 1 : 0 },
 	};
 
 	SourceFile = FString::Format(*SourceFile, TemplateArgs);
@@ -258,6 +309,7 @@ void UPCGCountUniqueAttributeValuesKernel::CreateAdditionalInputDataInterfaces(F
 	KernelDI->SetAttributeToCountName(AttributeName);
 	KernelDI->SetEmitPerDataCounts(bEmitPerDataCounts);
 	KernelDI->SetOutputRawBuffer(bOutputRawBuffer);
+	KernelDI->SetBucketingCellExtent(GetBucketExtent());
 
 	OutDataInterfaces.Add(KernelDI);
 }
@@ -276,6 +328,7 @@ void UPCGCountUniqueAttributeValuesKernel::GetKernelAttributeKeys(TArray<FPCGKer
 void UPCGCountUniqueAttributeValuesKernel::GetInputPins(TArray<FPCGPinProperties>& OutPins) const
 {
 	OutPins.Emplace(PCGPinConstants::DefaultInputLabel, EPCGDataType::Point);
+	OutPins.Emplace(PCGCountUniqueAttributeValuesConstants::UniqueValueTablePinLabel, EPCGDataType::Param);
 }
 
 void UPCGCountUniqueAttributeValuesKernel::GetOutputPins(TArray<FPCGPinPropertiesGPU>& OutPins) const
diff --git a/Engine/Plugins/PCG/Source/PCG/Private/Compute/BuiltInKernels/PCGCountUniqueAttributeValuesKernel.h b/Engine/Plugins/PCG/Source/PCG/Private/Compute/BuiltInKernels/PCGCountUniqueAttributeValuesKernel.h
--- a/Engine/Plugins/PCG/Source/PCG/Private/Compute/BuiltInKernels/PCGCountUniqueAttributeValuesKernel.h
+++ b/Engine/Plugins/PCG/Source/PCG/Private/Compute/BuiltInKernels/PCGCountUniqueAttributeValuesKernel.h
@@ -10,8 +10,21 @@ namespace PCGCountUniqueAttributeValuesConstants
 {
 	const FName ValueAttributeName = TEXT("UniqueValue");
 	const FName ValueCountAttributeName = TEXT("UniqueValueCount");
+	const FName UniqueValueTablePinLabel = TEXT("UniqueValueTable");
 }
 
+UENUM()
+enum class EPCGMaxNumUniqueValuesSource : uint8
+{
+	Invalid = 0,
+	/** Treat the attribute type as string key and derive maximum unique value count from static data descriptions. */
+	StringKeyValues = 1,
+	/** Explicitly set maximum value count. */
+	ExplicitMaxValue = 2,
+	/** Treat the attribute value as an element index into a table, and derive the max count based on the number of elements encountered in input data. */
+	UniqueValueTable = 3,
+};
+
 /**
 * Counts how many unique values of a string key or int attribute are present in an input data collection. Output attribute set is a table of unique
 * values and corresponding counts. Can output an attribute set per input data, or a single attribute set that counts across all data.
@@ -33,6 +46,9 @@ public:
 	/** Whether to output a raw array of (string key value, instance count) values from analysis instead of constructing an attribute set. */
 	void SetOutputRawBuffer(bool bInOutputRawBuffer) { bOutputRawBuffer = bInOutputRawBuffer; }
 
+	EPCGMaxNumUniqueValuesSource GetMaxNumUniqueValuesSource() const { return MaxNumUniqueValuesSource; }
+	void SetMaxNumUniqueValuesSource(EPCGMaxNumUniqueValuesSource InSource) { MaxNumUniqueValuesSource = InSource; }
+
 	virtual bool IsKernelDataValid(FPCGContext* InContext) const override;
 	virtual TSharedPtr<const FPCGDataCollectionDesc> ComputeOutputBindingDataDesc(FName InOutputPinLabel, UPCGDataBinding* InBinding) const override;
 	virtual int ComputeThreadCount(const UPCGDataBinding* Binding) const override;
@@ -50,6 +66,11 @@ public:
 	virtual void GetInputPins(TArray<FPCGPinProperties>& OutPins) const override;
 	virtual void GetOutputPins(TArray<FPCGPinPropertiesGPU>& OutPins) const override;
 
+	int32 GetBucketExtent() const { return BucketExtent; }
+	void SetBucketExtent(int32 InBucketExtent) { BucketExtent = InBucketExtent; }
+
+	int32 ComputeNumCountersRequired(UPCGDataBinding* InBinding, TSharedPtr<const FPCGDataCollectionDesc> InInputDesc) const;
+
 protected:
 	UPROPERTY()
 	FName AttributeName;
@@ -64,4 +85,11 @@ protected:
 
 	UPROPERTY()
 	bool bOutputRawBuffer = false;
+
+	UPROPERTY()
+	EPCGMaxNumUniqueValuesSource MaxNumUniqueValuesSource = EPCGMaxNumUniqueValuesSource::Invalid;
+
+	/** Extent of spatial bucket in cm. */
+	UPROPERTY()
+	int32 BucketExtent = 0;
 };
diff --git a/Engine/Plugins/PCG/Source/PCG/Private/Compute/DataInterfaces/BuiltInKernels/PCGCountUniqueAttributeValuesDataInterface.cpp b/Engine/Plugins/PCG/Source/PCG/Private/Compute/DataInterfaces/BuiltInKernels/PCGCountUniqueAttributeValuesDataInterface.cpp
--- a/Engine/Plugins/PCG/Source/PCG/Private/Compute/DataInterfaces/BuiltInKernels/PCGCountUniqueAttributeValuesDataInterface.cpp
+++ b/Engine/Plugins/PCG/Source/PCG/Private/Compute/DataInterfaces/BuiltInKernels/PCGCountUniqueAttributeValuesDataInterface.cpp
@@ -2,6 +2,7 @@
 
 #include "PCGCountUniqueAttributeValuesDataInterface.h"
 
+#include "PCGGraphExecutionStateInterface.h"
 #include "Compute/PCGComputeCommon.h"
 #include "Compute/PCGDataBinding.h"
 #include "Compute/PCGDataDescription.h"
@@ -37,6 +38,22 @@ void UPCGCountUniqueAttributeValuesDataInterface::GetSupportedInputs(TArray<FSha
 	OutFunctions.AddDefaulted_GetRef()
 		.SetName(TEXT("CountUniqueValues_GetEmitPerDataCounts"))
 		.AddReturnType(EShaderFundamentalType::Bool);
+
+	OutFunctions.AddDefaulted_GetRef()
+		.SetName(TEXT("CountUniqueValues_GetBoundsMin"))
+		.AddReturnType(EShaderFundamentalType::Float, 3);
+
+	OutFunctions.AddDefaulted_GetRef()
+		.SetName(TEXT("CountUniqueValues_GetBucketExtent"))
+		.AddReturnType(EShaderFundamentalType::Float);
+
+	OutFunctions.AddDefaulted_GetRef()
+		.SetName(TEXT("CountUniqueValues_GetNumBuckets"))
+		.AddReturnType(EShaderFundamentalType::Int, 3);
+
+	OutFunctions.AddDefaulted_GetRef()
+		.SetName(TEXT("CountUniqueValues_GetNumBucketsTotal"))
+		.AddReturnType(EShaderFundamentalType::Int);
 }
 
 BEGIN_SHADER_PARAMETER_STRUCT(FPCGCountUniqueAttributeValuesDataInterfaceParameters,)
@@ -44,6 +61,9 @@ BEGIN_SHADER_PARAMETER_STRUCT(FPCGCountUniqueAttributeValuesDataInterfaceParamet
 	SHADER_PARAMETER(int32, OutputValueAttributeId)
 	SHADER_PARAMETER(int32, OutputCountAttributeId)
 	SHADER_PARAMETER(uint32, EmitPerDataCounts)
+	SHADER_PARAMETER(FVector3f, BoundsMin)
+	SHADER_PARAMETER(float, BucketExtent)
+	SHADER_PARAMETER(FIntVector, NumBuckets)
 END_SHADER_PARAMETER_STRUCT()
 
 void UPCGCountUniqueAttributeValuesDataInterface::GetShaderParameters(TCHAR const* UID, FShaderParametersMetadataBuilder& InOutBuilder, FShaderParametersMetadataAllocations& InOutAllocations) const
@@ -63,6 +83,9 @@ void UPCGCountUniqueAttributeValuesDataInterface::GetHLSL(FString& OutHLSL, FStr
 		"int {DataInterfaceName}_OutputValueAttributeId;\n"
 		"int {DataInterfaceName}_OutputCountAttributeId;\n"
 		"uint {DataInterfaceName}_EmitPerDataCounts;\n"
+		"float3 {DataInterfaceName}_BoundsMin;\n"
+		"float {DataInterfaceName}_BucketExtent;\n"
+		"int3 {DataInterfaceName}_NumBuckets;\n"
 		"\n"
 		"int CountUniqueValues_GetAttributeToCountId_{DataInterfaceName}()\n"
 		"{\n"
@@ -83,7 +106,27 @@ void UPCGCountUniqueAttributeValuesDataInterface::GetHLSL(FString& OutHLSL, FStr
 		"{\n"
 		"	return {DataInterfaceName}_EmitPerDataCounts > 0;\n"
 		"}\n"
-		), TemplateArgs);
+		"\n"
+		"float3 CountUniqueValues_GetBoundsMin_{DataInterfaceName}()\n"
+		"{\n"
+		"	return {DataInterfaceName}_BoundsMin;\n"
+		"}\n"
+		"\n"
+		"float CountUniqueValues_GetBucketExtent_{DataInterfaceName}()\n"
+		"{\n"
+		"	return {DataInterfaceName}_BucketExtent;\n"
+		"}\n"
+		"\n"
+		"int3 CountUniqueValues_GetNumBuckets_{DataInterfaceName}()\n"
+		"{\n"
+		"	return {DataInterfaceName}_NumBuckets;\n"
+		"}\n"
+		"\n"
+		"int CountUniqueValues_GetNumBucketsTotal_{DataInterfaceName}()\n"
+		"{\n"
+		"	return {DataInterfaceName}_NumBuckets.x * {DataInterfaceName}_NumBuckets.y * {DataInterfaceName}_NumBuckets.z;\n"
+		"}\n"
+	), TemplateArgs);
 }
 
 UComputeDataProvider* UPCGCountUniqueAttributeValuesDataInterface::CreateDataProvider() const
@@ -102,6 +145,12 @@ void UPCGCountUniqueAttributeValuesDataProvider::Initialize(const UComputeDataIn
 	AttributeToCountName = DataInterface->GetAttributeToCountName();
 	bEmitPerDataCounts = DataInterface->GetEmitPerDataCounts();
 	bOutputRawBuffer = DataInterface->GetOutputRawBuffer();
+
+	UPCGDataBinding* Binding = CastChecked<UPCGDataBinding>(InBinding);
+	BoundsMin = FVector3f(Binding->GetExecutionSource()->GetExecutionState().GetBounds().Min);
+
+	NumBuckets = Binding->GetMaxNumCullingCells(DataInterface->GetBucketingCellExtent());
+	BucketExtent = static_cast<double>(DataInterface->GetBucketingCellExtent());
 }
 
 bool UPCGCountUniqueAttributeValuesDataProvider::PrepareForExecute_GameThread(UPCGDataBinding* InBinding)
@@ -133,11 +182,16 @@ bool UPCGCountUniqueAttributeValuesDataProvider::PrepareForExecute_GameThread(UP
 
 FComputeDataProviderRenderProxy* UPCGCountUniqueAttributeValuesDataProvider::GetRenderProxy()
 {
-	FPCGCountUniqueAttributeValuesProviderProxy::FCountUniqueValuesData_RenderThread Data;
-	Data.AttributeToCountId = AttributeToCountId;
-	Data.OutputValueAttributeId = OutputValueAttributeId;
-	Data.OutputCountAttributeId = OutputCountAttributeId;
-	Data.bEmitPerDataCounts = bEmitPerDataCounts;
+	FPCGCountUniqueAttributeValuesProviderProxy::FCountUniqueValuesData_RenderThread Data =
+	{
+		.AttributeToCountId = AttributeToCountId,
+		.OutputValueAttributeId = OutputValueAttributeId,
+		.OutputCountAttributeId = OutputCountAttributeId,
+		.bEmitPerDataCounts = bEmitPerDataCounts,
+		.BoundsMin = BoundsMin,
+		.BucketExtent = static_cast<float>(BucketExtent),
+		.NumBuckets = NumBuckets,
+	};
 
 	return new FPCGCountUniqueAttributeValuesProviderProxy(MoveTemp(Data));
 }
@@ -150,6 +204,9 @@ void UPCGCountUniqueAttributeValuesDataProvider::Reset()
 	OutputCountAttributeId = INDEX_NONE;
 	bEmitPerDataCounts = true;
 	bOutputRawBuffer = false;
+	BoundsMin = FVector3f::ZeroVector;
+	NumBuckets = FIntVector();
+	BucketExtent = 0.0;
 
 	Super::Reset();
 }
@@ -170,5 +227,8 @@ void FPCGCountUniqueAttributeValuesProviderProxy::GatherDispatchData(FDispatchDa
 		Parameters.OutputValueAttributeId = Data.OutputValueAttributeId;
 		Parameters.OutputCountAttributeId = Data.OutputCountAttributeId;
 		Parameters.EmitPerDataCounts = Data.bEmitPerDataCounts ? 1u : 0u;
+		Parameters.BoundsMin = Data.BoundsMin;
+		Parameters.BucketExtent = Data.BucketExtent;
+		Parameters.NumBuckets = Data.NumBuckets;
 	}
 }
diff --git a/Engine/Plugins/PCG/Source/PCG/Private/Compute/DataInterfaces/BuiltInKernels/PCGCountUniqueAttributeValuesDataInterface.h b/Engine/Plugins/PCG/Source/PCG/Private/Compute/DataInterfaces/BuiltInKernels/PCGCountUniqueAttributeValuesDataInterface.h
--- a/Engine/Plugins/PCG/Source/PCG/Private/Compute/DataInterfaces/BuiltInKernels/PCGCountUniqueAttributeValuesDataInterface.h
+++ b/Engine/Plugins/PCG/Source/PCG/Private/Compute/DataInterfaces/BuiltInKernels/PCGCountUniqueAttributeValuesDataInterface.h
@@ -33,6 +33,9 @@ public:
 	bool GetOutputRawBuffer() const { return bOutputRawBuffer; }
 	void SetOutputRawBuffer(bool bInOutputRawBuffer) { bOutputRawBuffer = bInOutputRawBuffer; }
 
+	int32 GetBucketingCellExtent() const { return BucketingCellExtent; }
+	void SetBucketingCellExtent(int32 InGrid) { BucketingCellExtent = InGrid; }
+
 protected:
 	UPROPERTY()
 	FName AttributeToCountName;
@@ -42,6 +45,9 @@ protected:
 
 	UPROPERTY()
 	bool bOutputRawBuffer = false;
+
+	UPROPERTY()
+	int32 BucketingCellExtent = 0;
 };
 
 UCLASS()
@@ -72,6 +78,10 @@ public:
 	bool bEmitPerDataCounts = true;
 
 	bool bOutputRawBuffer = false;
+
+	FVector3f BoundsMin = FVector3f::ZeroVector;
+	double BucketExtent = 0.0;
+	FIntVector NumBuckets = FIntVector(1, 1, 1);
 };
 
 class FPCGCountUniqueAttributeValuesProviderProxy : public FComputeDataProviderRenderProxy
@@ -83,6 +93,9 @@ public:
 		int32 OutputValueAttributeId = INDEX_NONE;
 		int32 OutputCountAttributeId = INDEX_NONE;
 		bool bEmitPerDataCounts = true;
+		FVector3f BoundsMin = FVector3f::ZeroVector;
+		float BucketExtent = 0.0f;
+		FIntVector NumBuckets = FIntVector(1, 1, 1);
 	};
 
 	FPCGCountUniqueAttributeValuesProviderProxy(FCountUniqueValuesData_RenderThread InData)
diff --git a/Engine/Plugins/PCG/Source/PCG/Private/Compute/DataInterfaces/Elements/PCGStaticMeshSpawnerDataInterface.cpp b/Engine/Plugins/PCG/Source/PCG/Private/Compute/DataInterfaces/Elements/PCGStaticMeshSpawnerDataInterface.cpp
--- a/Engine/Plugins/PCG/Source/PCG/Private/Compute/DataInterfaces/Elements/PCGStaticMeshSpawnerDataInterface.cpp
+++ b/Engine/Plugins/PCG/Source/PCG/Private/Compute/DataInterfaces/Elements/PCGStaticMeshSpawnerDataInterface.cpp
@@ -12,7 +12,6 @@
 #include "Compute/Data/PCGProxyForGPUData.h"
 #include "Compute/Data/PCGRawBufferData.h"
 #include "Compute/Packing/PCGMeshSpawnerPacking.h"
-#include "Compute/PrimitiveFactories/PCGPrimitiveFactoryPISMC.h"
 #include "Elements/PCGStaticMeshSpawner.h"
 #include "Elements/PCGStaticMeshSpawnerKernel.h"
 #include "InstanceDataPackers/PCGInstanceDataPackerBase.h"
@@ -104,6 +103,22 @@ void UPCGStaticMeshSpawnerDataInterface::GetSupportedInputs(TArray<FShaderFuncti
 	OutFunctions.AddDefaulted_GetRef()
 		.SetName(TEXT("SMSpawner_GetSelectedMeshAttributeId"))
 		.AddReturnType(EShaderFundamentalType::Uint); // Attribute id to output mesh path string key to
+
+	OutFunctions.AddDefaulted_GetRef()
+		.SetName(TEXT("SMSpawner_GetNumCullingCells"))
+		.AddReturnType(EShaderFundamentalType::Int, 3); // Number of culling cells that the primitive volume overlaps, per dimension
+
+	OutFunctions.AddDefaulted_GetRef()
+		.SetName(TEXT("SMSpawner_GetNumCullingCellsTotal"))
+		.AddReturnType(EShaderFundamentalType::Int); // Total number of culling cells that the primitive volume overlaps
+
+	OutFunctions.AddDefaulted_GetRef()
+		.SetName(TEXT("SMSpawner_GetBoundsMin"))
+		.AddReturnType(EShaderFundamentalType::Float, 3);
+
+	OutFunctions.AddDefaulted_GetRef()
+		.SetName(TEXT("SMSpawner_GetBucketExtent"))
+		.AddReturnType(EShaderFundamentalType::Float);
 }
 
 BEGIN_SHADER_PARAMETER_STRUCT(FPCGStaticMeshSpawnerDataInterfaceParameters,)
@@ -118,6 +133,10 @@ BEGIN_SHADER_PARAMETER_STRUCT(FPCGStaticMeshSpawnerDataInterfaceParameters,)
 	SHADER_PARAMETER(int32, SelectorAttributeId)
 	SHADER_PARAMETER(int32, SelectedMeshAttributeId)
 	SHADER_PARAMETER(uint32, ApplyBounds)
+	SHADER_PARAMETER(FIntVector, NumCullingCells)
+	SHADER_PARAMETER(int32, NumCullingCellsTotal)
+	SHADER_PARAMETER(FVector3f, BoundsMin)
+	SHADER_PARAMETER(float, BucketExtent) // TODO rename
 END_SHADER_PARAMETER_STRUCT()
 
 void UPCGStaticMeshSpawnerDataInterface::GetShaderParameters(TCHAR const* UID, FShaderParametersMetadataBuilder& InOutBuilder, FShaderParametersMetadataAllocations& InOutAllocations) const
@@ -141,6 +160,11 @@ void UPCGStaticMeshSpawnerDataInterface::GetHLSL(FString& OutHLSL, FString const
 		"uint {DataInterfaceName}_ApplyBounds;\n"
 		"int {DataInterfaceName}_SelectedMeshAttributeId;\n"
 		"uint4 {DataInterfaceName}_AttributeIdOffsetStrides[{MaxAttributes}];\n"
+		"int3 {DataInterfaceName}_NumCullingCells;\n"
+		"int {DataInterfaceName}_NumCullingCellsTotal;\n"
+		"float3 {DataInterfaceName}_BoundsMin;\n"
+		"float {DataInterfaceName}_BucketExtent;\n"
+		"\n"
 		"StructuredBuffer<float4> {DataInterfaceName}_PrimitiveMeshBoundsMin;\n"
 		"StructuredBuffer<float4> {DataInterfaceName}_PrimitiveMeshBoundsMax;\n"
 		"StructuredBuffer<int> {DataInterfaceName}_PrimitiveIds;\n"
@@ -208,7 +232,27 @@ void UPCGStaticMeshSpawnerDataInterface::GetHLSL(FString& OutHLSL, FString const
 		"int SMSpawner_GetSelectedMeshAttributeId_{DataInterfaceName}()\n"
 		"{\n"
 		"	return {DataInterfaceName}_SelectedMeshAttributeId;\n"
-		"}\n\n"
+		"}\n"
+		"\n"
+		"int3 SMSpawner_GetNumCullingCells_{DataInterfaceName}()\n"
+		"{\n"
+		"	return {DataInterfaceName}_NumCullingCells;\n"
+		"}\n"
+		"\n"
+		"int SMSpawner_GetNumCullingCellsTotal_{DataInterfaceName}()\n"
+		"{\n"
+		"	return {DataInterfaceName}_NumCullingCellsTotal;\n"
+		"}\n"
+		"\n"
+		"float3 SMSpawner_GetBoundsMin_{DataInterfaceName}()\n"
+		"{\n"
+		"	return {DataInterfaceName}_BoundsMin;\n"
+		"}\n"
+		"\n"
+		"float SMSpawner_GetBucketExtent_{DataInterfaceName}()\n"
+		"{\n"
+		"	return {DataInterfaceName}_BucketExtent;\n"
+		"}\n"
 	), TemplateArgs);
 }
 
@@ -224,6 +268,15 @@ void UPCGStaticMeshSpawnerDataProvider::Initialize(const UComputeDataInterface*
 	const UPCGStaticMeshSpawnerSettings* Settings = CastChecked<UPCGStaticMeshSpawnerSettings>(GetProducerKernel()->GetSettings());
 	bSpawningByAttribute = !!Cast<UPCGMeshSelectorByAttribute>(Settings->MeshSelectorParameters);
 	bSpawningByPrimitiveData = !!Cast<UPCGMeshSelectorPrimitiveData>(Settings->MeshSelectorParameters);
+
+	UPCGDataBinding* Binding = CastChecked<UPCGDataBinding>(InBinding);
+
+	NumCullingCells = Binding->GetMaxNumCullingCells(Settings->GetCullingCellSize());
+	CullingCellExtent = static_cast<float>(Settings->GetCullingCellSize());
+
+	NumCullingCellsTotal = NumCullingCells.X * NumCullingCells.Y * NumCullingCells.Z;
+
+	BoundsMin = FVector3f(Binding->GetExecutionSource()->GetExecutionState().GetBounds().Min);
 }
 
 bool UPCGStaticMeshSpawnerDataProvider::PerformPreExecuteReadbacks_GameThread(UPCGDataBinding* InBinding)
@@ -292,12 +345,33 @@ bool UPCGStaticMeshSpawnerDataProvider::PerformPreExecuteReadbacks_GameThread(UP
 
 			const int32 NumElements = AnalysisMetadata->GetItemCountForChild();
 
-			PrimitiveIdToInstanceCount.Reserve(NumElements);
+			int32 LastPrimitiveId = INDEX_NONE;
+			uint32 CellId = 0;
+			TArray<FPCGInstanceRange>* CurrentInstanceRanges = nullptr;
 
 			// TODO: Range based get would scale better.
 			for (int64 MetadataKey = 0; MetadataKey < NumElements; ++MetadataKey)
 			{
-				PrimitiveIdToInstanceCount.Add(ValueAttribute->GetValue(MetadataKey), CountAttribute->GetValue(MetadataKey));
+				const int32 PrimitiveId = ValueAttribute->GetValue(MetadataKey);
+				const int32 InstanceCount = CountAttribute->GetValue(MetadataKey);
+
+				if (LastPrimitiveId != PrimitiveId)
+				{
+					// Reset cell ID if we encounter a new primitive.
+					CellId = 0;
+
+					// Cache pointer to avoid repeated map finds.
+					CurrentInstanceRanges = &PrimitiveIdToInstanceRanges.FindOrAdd(PrimitiveId);
+				}
+
+				LastPrimitiveId = PrimitiveId;
+
+				if (InstanceCount > 0)
+				{
+					CurrentInstanceRanges->Add(FPCGInstanceRange(InstanceCount, CellId));
+				}
+
+				++CellId;
 			}
 		}
 		else
@@ -320,10 +394,33 @@ bool UPCGStaticMeshSpawnerDataProvider::PerformPreExecuteReadbacks_GameThread(UP
 		const int32 NumUints = RawAnalysisData->GetNumUint32s();
 		ensure((NumUints % 2) == 0);
 
+		int32 LastPrimitiveId = INDEX_NONE;
+		uint32 CellId = 0;
+		TArray<FPCGInstanceRange>* CurrentInstanceRanges = nullptr;
+
 		const uint32* RawData = RawAnalysisData->GetConstData().GetData();
 		for (const uint32* It = RawData; It < RawData + NumUints; It += 2)
 		{
-			PrimitiveIdToInstanceCount.Add(static_cast<int32>(It[0]), It[1]);
+			const int32 PrimitiveId = static_cast<int32>(It[0]);
+			const uint32 InstanceCount = It[1];
+
+			if (LastPrimitiveId != PrimitiveId)
+			{
+				// Reset cell ID if we encounter a new primitive.
+				CellId = 0;
+
+				// Cache pointer to avoid repeated map finds.
+				CurrentInstanceRanges = &PrimitiveIdToInstanceRanges.FindOrAdd(PrimitiveId);
+			}
+
+			LastPrimitiveId = PrimitiveId;
+
+			if (InstanceCount > 0)
+			{
+				CurrentInstanceRanges->Add(FPCGInstanceRange(InstanceCount, CellId));
+			}
+
+			++CellId;
 		}
 	}
 
@@ -428,14 +525,14 @@ bool UPCGStaticMeshSpawnerDataProvider::PrepareForExecute_GameThread(UPCGDataBin
 		CreatePrimitiveDescriptors(Context, InBinding);
 
 		TArray<FSoftObjectPath> PathsToLoad;
-		PathsToLoad.Reserve(PrimitiveDescriptors.Num());
+		PathsToLoad.Reserve(PrimitiveInfos.Num());
 
-		for (const FPCGProceduralISMComponentDescriptor& PrimitiveDescriptor : PrimitiveDescriptors)
+		for (const FPCGPrimitiveInfo& PrimitiveInfo : PrimitiveInfos)
 		{
-			if (!PrimitiveDescriptor.StaticMesh.IsNull())
+			if (!PrimitiveInfo.Descriptor.StaticMesh.IsNull())
 			{
-				UE_LOG(LogPCG, Verbose, TEXT("Request '%s' to load."), *PrimitiveDescriptor.StaticMesh.ToString());
-				PathsToLoad.Add(PrimitiveDescriptor.StaticMesh.ToSoftObjectPath());
+				UE_LOG(LogPCG, Verbose, TEXT("Request '%s' to load."), *PrimitiveInfo.Descriptor.StaticMesh.ToString());
+				PathsToLoad.Add(PrimitiveInfo.Descriptor.StaticMesh.ToSoftObjectPath());
 			}
 		}
 
@@ -455,11 +552,11 @@ bool UPCGStaticMeshSpawnerDataProvider::PrepareForExecute_GameThread(UPCGDataBin
 
 	if (!bStaticMeshesLoaded)
 	{
-		for (const FPCGProceduralISMComponentDescriptor& PrimitiveDescriptor : PrimitiveDescriptors)
+		for (const FPCGPrimitiveInfo& PrimitiveInfo : PrimitiveInfos)
 		{
-			if (PrimitiveDescriptor.StaticMesh.IsPending())
+			if (PrimitiveInfo.Descriptor.StaticMesh.IsPending())
 			{
-				UE_LOG(LogPCG, Verbose, TEXT("Waiting for '%s' to load."), *PrimitiveDescriptor.StaticMesh.ToString());
+				UE_LOG(LogPCG, Verbose, TEXT("Waiting for '%s' to load."), *PrimitiveInfo.Descriptor.StaticMesh.ToString());
 
 				return false;
 			}
@@ -468,18 +565,18 @@ bool UPCGStaticMeshSpawnerDataProvider::PrepareForExecute_GameThread(UPCGDataBin
 		bool bAnyInvalid = false;
 
 		// User could pass us any soft object path so verify static meshes were loaded.
-		for (const FPCGProceduralISMComponentDescriptor& Descriptor : PrimitiveDescriptors)
+		for (const FPCGPrimitiveInfo& PrimitiveInfo : PrimitiveInfos)
 		{
-			if (!Cast<UStaticMesh>(Descriptor.StaticMesh.Get()))
+			if (!Cast<UStaticMesh>(PrimitiveInfo.Descriptor.StaticMesh.Get()))
 			{
-				UE_LOG(LogPCG, Error, TEXT("Tried to use asset '%s' as a static mesh."), *Descriptor.StaticMesh.ToString());
+				UE_LOG(LogPCG, Error, TEXT("Tried to use asset '%s' as a static mesh."), *PrimitiveInfo.Descriptor.StaticMesh.ToString());
 				bAnyInvalid = true;
 			}
 		}
 
 		if (bAnyInvalid)
 		{
-			PrimitiveDescriptors.Reset();
+			PrimitiveInfos.Reset();
 		}
 
 		bStaticMeshesLoaded = true;
@@ -526,7 +623,7 @@ bool UPCGStaticMeshSpawnerDataProvider::PrepareForExecute_GameThread(UPCGDataBin
 FComputeDataProviderRenderProxy* UPCGStaticMeshSpawnerDataProvider::GetRenderProxy()
 {
 	return new FPCGStaticMeshSpawnerDataProviderProxy(
-		AttributeIdOffsetStrides, SelectorAttributeId, PrimitiveIds, PrimitiveStringKeys, PrimitiveSelectionCDF, SelectedMeshAttributeId, PrimitiveMeshBounds);
+		AttributeIdOffsetStrides, SelectorAttributeId, PrimitiveIds, PrimitiveStringKeys, PrimitiveSelectionCDF, SelectedMeshAttributeId, PrimitiveMeshBounds, NumCullingCells, NumCullingCellsTotal, BoundsMin, CullingCellExtent);
 }
 
 void UPCGStaticMeshSpawnerDataProvider::Reset()
@@ -541,18 +638,21 @@ void UPCGStaticMeshSpawnerDataProvider::Reset()
 	PrimitiveMeshBounds.Empty();
 	PrimitiveSelectionCDF.Empty();
 	SelectorAttributeId = INDEX_NONE;
-	NumInputPoints = 0;
 	SelectedMeshAttributeId = INDEX_NONE;
-	PrimitiveIdToInstanceCount.Empty();
+	PrimitiveIdToInstanceRanges.Empty();
 	AnalysisDataIndex = INDEX_NONE;
 	PrimitiveDataIndex = INDEX_NONE;
 	bPrimitiveDescriptorsCreated = false;
-	PrimitiveDescriptors.Empty();
+	PrimitiveInfos.Empty();
 	CustomFloatCount = 0;
 	bPrimitivesSetUp = false;
 	NumPrimitivesSetup = 0;
 	bRegisteredPrimitives = false;
 	bStaticMeshesLoaded = false;
+	NumCullingCells = FIntVector(1, 1, 1);
+	NumCullingCellsTotal = 1;
+	BoundsMin = FVector3f::ZeroVector;
+	CullingCellExtent = 0.0;
 	LoadHandle.Reset();
 	CustomPrimitiveData.Reset();
 }
@@ -562,7 +662,7 @@ void UPCGStaticMeshSpawnerDataProvider::CreatePrimitiveDescriptors(FPCGContext*
 	TRACE_CPUPROFILER_EVENT_SCOPE(UPCGStaticMeshSpawnerDataProvider::CreatePrimitiveDescriptors);
 	check(InContext);
 
-	PrimitiveDescriptors.Empty();
+	PrimitiveInfos.Empty();
 
 	const UPCGStaticMeshSpawnerSettings* Settings = CastChecked<UPCGStaticMeshSpawnerSettings>(GetProducerKernel()->GetSettings());
 
@@ -634,66 +734,33 @@ void UPCGStaticMeshSpawnerDataProvider::CreatePrimitiveDescriptors(FPCGContext*
 	{
 		const FName SelectorName = SelectorByAttribute->AttributeName;
 
-		// Compute how many instances we expect for each mesh, if we do not already have the answer from analysis.
-		if (PrimitiveIdToInstanceCount.IsEmpty())
+		PrimitiveInfos.Reserve(PrimitiveIdToInstanceRanges.Num());
+
+		for (TPair<int32, TArray<FPCGInstanceRange>>& PrimitiveIdAndInstanceInfo : PrimitiveIdToInstanceRanges)
 		{
-			for (const FPCGDataDesc& Desc : InputDataDesc->GetDataDescriptions())
+			uint32 TotalInstanceCount = 0;
+			for (const FPCGInstanceRange& InstanceRange : PrimitiveIdAndInstanceInfo.Value)
 			{
-				if (Desc.GetElementCount().X <= 0)
-				{
-					continue;
-				}
-
-				for (const FPCGKernelAttributeDesc& AttributeDesc : Desc.GetAttributeDescriptions())
-				{
-					if (AttributeDesc.GetAttributeKey().GetIdentifier().Name == SelectorName)
-					{
-						if (AttributeDesc.GetAttributeKey().GetType() == EPCGKernelAttributeType::StringKey)
-						{
-							for (int32 AttributeStringKeyValue : AttributeDesc.GetUniqueStringKeys())
-							{
-								if (AttributeStringKeyValue > 0)
-								{
-									PrimitiveIdToInstanceCount.FindOrAdd(AttributeStringKeyValue, 0) += Desc.GetElementCountForAttribute(AttributeDesc).X;
-								}
-							}
-
-							// We could early out but we currently continue looping to generate the below warning for mismatched types.
-						}
-						else
-						{
-							// Currently only a single type per attribute name is supported (the name/type in the attribute table). It's possible to wire a graph
-							// for which an attribute is present with multiple types. Warn if this is encountered.
-							PCG_KERNEL_VALIDATION_WARN(InContext, Settings, FText::Format(
-								LOCTEXT("MeshSelectorAttributeNotUsable", "Attribute '{0}' not usable for mesh selection, only attributes of type String Key are supported."),
-								FText::FromName(SelectorName)));
-						}
-					}
-				}
+				TotalInstanceCount += InstanceRange.GetNumInstances();
 			}
-		}
 
-		PrimitiveDescriptors.Reserve(PrimitiveIdToInstanceCount.Num());
-
-		for (const TPair<int32, uint32>& PrimitiveIdAndInstanceCount : PrimitiveIdToInstanceCount)
-		{
-			if (PrimitiveIdAndInstanceCount.Value == 0)
+			if (TotalInstanceCount == 0)
 			{
 				continue;
 			}
 
-			if (!ensure(InBinding->GetStringTable().IsValidIndex(PrimitiveIdAndInstanceCount.Key)))
+			if (!ensure(InBinding->GetStringTable().IsValidIndex(PrimitiveIdAndInstanceInfo.Key)))
 			{
 				continue;
 			}
 
-			const FString& MeshPathString = InBinding->GetStringTable()[PrimitiveIdAndInstanceCount.Key];
+			const FString& MeshPathString = InBinding->GetStringTable()[PrimitiveIdAndInstanceInfo.Key];
 			if (MeshPathString.IsEmpty())
 			{
 				continue;
 			}
 
-			if (PrimitiveDescriptors.Num() >= PCGComputeConstants::MAX_PRIMITIVE_COMPONENTS_PER_SPAWNER)
+			if (PrimitiveInfos.Num() >= PCGComputeConstants::MAX_PRIMITIVE_COMPONENTS_PER_SPAWNER)
 			{
 				PCG_KERNEL_VALIDATION_WARN(InContext, Settings,
 					FText::Format(PCGStaticMeshSpawnerDataInterface::TooManyPrimitivesFormat, FText::FromString(FString::FromInt(PCGComputeConstants::MAX_PRIMITIVE_COMPONENTS_PER_SPAWNER))));
@@ -703,7 +770,7 @@ void UPCGStaticMeshSpawnerDataProvider::CreatePrimitiveDescriptors(FPCGContext*
 
 			FPCGProceduralISMComponentDescriptor Descriptor;
 			Descriptor = SelectorByAttribute->TemplateDescriptor;
-			Descriptor.NumInstances = PrimitiveIdAndInstanceCount.Value;
+			Descriptor.NumInstances = TotalInstanceCount;
 			Descriptor.WorldBounds = InContext->ExecutionSource->GetExecutionState().GetBounds();
 			Descriptor.NumCustomFloats = CustomFloatCount;
 			Descriptor.StaticMesh = FSoftObjectPath(MeshPathString);
@@ -714,18 +781,21 @@ void UPCGStaticMeshSpawnerDataProvider::CreatePrimitiveDescriptors(FPCGContext*
 				Descriptor.NumInstances = TotalInputPointCount;
 			}
 
-			PrimitiveIds.Add(PrimitiveIdAndInstanceCount.Key);
-			PrimitiveDescriptors.Add(MoveTemp(Descriptor));
+			PrimitiveIds.Add(PrimitiveIdAndInstanceInfo.Key);
+
+			FPCGPrimitiveInfo& PrimitiveInfo = PrimitiveInfos.Emplace_GetRef();
+			PrimitiveInfo.Descriptor = MoveTemp(Descriptor);
+			PrimitiveInfo.InstanceRanges = MoveTemp(PrimitiveIdAndInstanceInfo.Value);
 
 			// For by-attribute spawning, the primitive ID is the mesh path string key.
-			PrimitiveStringKeys.Add(PrimitiveIdAndInstanceCount.Key);
+			PrimitiveStringKeys.Add(PrimitiveIdAndInstanceInfo.Key);
 		}
 
-		PrimitiveSelectionCDF.SetNumZeroed(PrimitiveDescriptors.Num());
+		PrimitiveSelectionCDF.SetNumZeroed(PrimitiveInfos.Num());
 	}
 	else if (const UPCGMeshSelectorPrimitiveData* SelectorPrimitiveData = Cast<UPCGMeshSelectorPrimitiveData>(Settings->MeshSelectorParameters))
 	{
-		if (PrimitiveIdToInstanceCount.IsEmpty())
+		if (PrimitiveIdToInstanceRanges.IsEmpty())
 		{
 			return;
 		}
@@ -764,24 +834,30 @@ void UPCGStaticMeshSpawnerDataProvider::CreatePrimitiveDescriptors(FPCGContext*
 		FPCGMeshMaterialOverrideHelper MaterialOverrideHelper;
 		MaterialOverrideHelper.Initialize(*InContext, /*bInByAttributeOverride=*/true, SelectorPrimitiveData->MaterialOverrideAttributes, PrimitiveMetadata);
 
-		SelectorPrimitiveData->PackCustomPrimitiveData(PrimitiveData, PrimitiveIdToInstanceCount.Num(), CustomPrimitiveData, InContext);
+		SelectorPrimitiveData->PackCustomPrimitiveData(PrimitiveData, PrimitiveIdToInstanceRanges.Num(), CustomPrimitiveData, InContext);
 
 		const int32 NumPrimitiveDataElements = PrimitiveMetadata->GetItemCountForChild();
 
-		PrimitiveDescriptors.Reserve(PrimitiveIdToInstanceCount.Num());
-		PrimitiveIds.Reserve(PrimitiveIdToInstanceCount.Num());
-		PrimitiveStringKeys.Reserve(PrimitiveIdToInstanceCount.Num());
+		PrimitiveInfos.Reserve(PrimitiveIdToInstanceRanges.Num());
+		PrimitiveIds.Reserve(PrimitiveIdToInstanceRanges.Num());
+		PrimitiveStringKeys.Reserve(PrimitiveIdToInstanceRanges.Num());
 		
-		for (const TPair<int32, uint32>& PrimitiveIdAndInstanceCount : PrimitiveIdToInstanceCount)
+		for (TPair<int32, TArray<FPCGInstanceRange>>& PrimitiveIdAndInstanceInfo : PrimitiveIdToInstanceRanges)
 		{
-			if (PrimitiveIdAndInstanceCount.Value == 0)
+			uint32 TotalInstanceCount = 0;
+			for (const FPCGInstanceRange& InstanceRange : PrimitiveIdAndInstanceInfo.Value)
+			{
+				TotalInstanceCount += InstanceRange.GetNumInstances();
+			}
+
+			if (TotalInstanceCount == 0)
 			{
 				continue;
 			}
 
 			TRACE_CPUPROFILER_EVENT_SCOPE(UPCGStaticMeshSpawnerDataProvider::CreatePrimitiveDescriptor);
 
-			const int32 PrimitiveId = PrimitiveIdAndInstanceCount.Key;
+			const int32 PrimitiveId = PrimitiveIdAndInstanceInfo.Key;
 			if (PrimitiveId >= NumPrimitiveDataElements)
 			{
 				PCG_KERNEL_VALIDATION_WARN(InContext, Settings,
@@ -792,7 +868,7 @@ void UPCGStaticMeshSpawnerDataProvider::CreatePrimitiveDescriptors(FPCGContext*
 			const PCGMetadataValueKey ValueKey = MeshAttributeSoftObjectPath ? MeshAttributeSoftObjectPath->GetValueKey(PrimitiveId) : MeshAttributeString->GetValueKey(PrimitiveId);
 			const FSoftObjectPath MeshPath = MeshAttributeSoftObjectPath ? MeshAttributeSoftObjectPath->GetValue(ValueKey) : FSoftObjectPath(MeshAttributeString->GetValue(ValueKey));
 
-			if (PrimitiveDescriptors.Num() >= PCGComputeConstants::MAX_PRIMITIVE_COMPONENTS_PER_SPAWNER)
+			if (PrimitiveInfos.Num() >= PCGComputeConstants::MAX_PRIMITIVE_COMPONENTS_PER_SPAWNER)
 			{
 				PCG_KERNEL_VALIDATION_WARN(InContext, Settings,
 					FText::Format(PCGStaticMeshSpawnerDataInterface::TooManyPrimitivesFormat, FText::FromString(FString::FromInt(PCGComputeConstants::MAX_PRIMITIVE_COMPONENTS_PER_SPAWNER))));
@@ -800,7 +876,7 @@ void UPCGStaticMeshSpawnerDataProvider::CreatePrimitiveDescriptors(FPCGContext*
 				break;
 			}
 
-			Descriptor.NumInstances = PrimitiveIdAndInstanceCount.Value;
+			Descriptor.NumInstances = TotalInstanceCount;
 			Descriptor.StaticMesh = MeshPath;
 
 			PropertyOverrides.Apply(PrimitiveId);
@@ -821,13 +897,18 @@ void UPCGStaticMeshSpawnerDataProvider::CreatePrimitiveDescriptors(FPCGContext*
 			}
 
 			PrimitiveIds.Add(PrimitiveId);
-			PrimitiveDescriptors.Add(MoveTemp(Descriptor));
+
+			FPCGPrimitiveInfo& PrimitiveInfo = PrimitiveInfos.Emplace_GetRef();
+			PrimitiveInfo.Descriptor = MoveTemp(Descriptor);
+			PrimitiveInfo.InstanceRanges = PrimitiveIdAndInstanceInfo.Value;
+			PrimitiveInfo.CullingCellExtent = CullingCellExtent;
+			PrimitiveInfo.NumCullingCells = NumCullingCells;
 
 			// Look up the mesh path string key value.
 			PrimitiveStringKeys.Add(InBinding->GetStringTable().IndexOfByKey(MeshPath));
 		}
 
-		PrimitiveSelectionCDF.SetNumZeroed(PrimitiveDescriptors.Num());
+		PrimitiveSelectionCDF.SetNumZeroed(PrimitiveInfos.Num());
 	}
 	else if (const UPCGMeshSelectorWeighted* SelectorWeighted = Cast<UPCGMeshSelectorWeighted>(Settings->MeshSelectorParameters))
 	{
@@ -852,7 +933,7 @@ void UPCGStaticMeshSpawnerDataProvider::CreatePrimitiveDescriptors(FPCGContext*
 
 		PrimitiveSelectionCDF.Reserve(SelectorWeighted->MeshEntries.Num());
 
-		PrimitiveDescriptors.Reserve(SelectorWeighted->MeshEntries.Num());
+		PrimitiveInfos.Reserve(SelectorWeighted->MeshEntries.Num());
 
 		for (const FPCGMeshSelectorWeightedEntry& Entry : SelectorWeighted->MeshEntries)
 		{
@@ -869,7 +950,7 @@ void UPCGStaticMeshSpawnerDataProvider::CreatePrimitiveDescriptors(FPCGContext*
 				continue;
 			}
 
-			if (PrimitiveDescriptors.Num() >= PCGComputeConstants::MAX_PRIMITIVE_COMPONENTS_PER_SPAWNER)
+			if (PrimitiveInfos.Num() >= PCGComputeConstants::MAX_PRIMITIVE_COMPONENTS_PER_SPAWNER)
 			{
 				PCG_KERNEL_VALIDATION_WARN(InContext, Settings,
 					FText::Format(PCGStaticMeshSpawnerDataInterface::TooManyPrimitivesFormat, FText::FromString(FString::FromInt(PCGComputeConstants::MAX_PRIMITIVE_COMPONENTS_PER_SPAWNER))));
@@ -910,7 +991,11 @@ void UPCGStaticMeshSpawnerDataProvider::CreatePrimitiveDescriptors(FPCGContext*
 
 			Descriptor.NumInstances = FMath::Min(InstanceCount, TotalInputPointCount);
 
-			PrimitiveDescriptors.Add(MoveTemp(Descriptor));
+			FPCGPrimitiveInfo& PrimitiveInfo = PrimitiveInfos.Emplace_GetRef();
+			PrimitiveInfo.Descriptor = MoveTemp(Descriptor);
+			PrimitiveInfo.InstanceRanges.Emplace(Descriptor.NumInstances);
+			PrimitiveInfo.CullingCellExtent = CullingCellExtent;
+			PrimitiveInfo.NumCullingCells = NumCullingCells;
 		}
 	}
 
@@ -941,15 +1026,15 @@ bool UPCGStaticMeshSpawnerDataProvider::SetupPrimitives(FPCGContext* InContext,
 	{
 		// If we have no descriptors, we can early out. Only checked when the primitive factory is not initialized,
 		// as we're moving the descritors into the factory (which will emtpy PrimitiveDescriptors)
-		if (PrimitiveDescriptors.IsEmpty())
+		if (PrimitiveInfos.IsEmpty())
 		{
 			return true;
 		}
 
 		// Ensure SMs are compiled, required to prevent ensure when partially built SM is accessed later in FastGeo code.
-		for (const FPCGProceduralISMComponentDescriptor& Desc : PrimitiveDescriptors)
+		for (const FPCGPrimitiveInfo& PrimitiveInfo : PrimitiveInfos)
 		{
-			if (ensure(Desc.StaticMesh.IsValid()) && Desc.StaticMesh->IsCompiling())
+			if (ensure(PrimitiveInfo.Descriptor.StaticMesh.IsValid()) && PrimitiveInfo.Descriptor.StaticMesh->IsCompiling())
 			{
 				return false;
 			}
@@ -977,7 +1062,7 @@ bool UPCGStaticMeshSpawnerDataProvider::SetupPrimitives(FPCGContext* InContext,
 		Primitives->PrimitiveFactory = Factory;
 
 		IPCGPrimitiveFactoryISMBase::FParameters Params;
-		Params.Descriptors = MoveTemp(PrimitiveDescriptors);
+		Params.PrimitiveInfos = MoveTemp(PrimitiveInfos);
 		Params.CustomPrimitiveData = MoveTemp(CustomPrimitiveData);
 		Params.TargetActor = Settings->TargetActor.Get() ? Settings->TargetActor.Get() : InContext->GetTypedExecutionTarget<AActor>();
 
@@ -1014,13 +1099,21 @@ FPCGStaticMeshSpawnerDataProviderProxy::FPCGStaticMeshSpawnerDataProviderProxy(
 	const TArray<int32>& InPrimitiveStringKeys,
 	TArray<float> InSelectionCDF,
 	int32 InSelectedMeshAttributeId,
-	const TArray<FBox>& InPrimitiveMeshBounds)
+	const TArray<FBox>& InPrimitiveMeshBounds,
+	FIntVector InNumSubdivisions,
+	int32 InNumSubdivisionsTotal,
+	FVector3f InBoundsMin,
+	float InBucketExtent)
 	: AttributeIdOffsetStrides(InAttributeIdOffsetStrides)
 	, SelectionCDF(InSelectionCDF)
 	, SelectorAttributeId(InSelectorAttributeId)
 	, PrimitiveIds(InPrimitiveIds)
 	, PrimitiveStringKeys(InPrimitiveStringKeys)
 	, SelectedMeshAttributeId(InSelectedMeshAttributeId)
+	, NumCullingCells(InNumSubdivisions)
+	, NumCullingCellsTotal(InNumSubdivisionsTotal)
+	, CullingCellExtent(InBucketExtent)
+	, BoundsMin(InBoundsMin)
 {
 	PrimitiveMeshBoundsMin.Reserve(InPrimitiveMeshBounds.Num());
 	PrimitiveMeshBoundsMax.Reserve(InPrimitiveMeshBounds.Num());
@@ -1034,7 +1127,7 @@ FPCGStaticMeshSpawnerDataProviderProxy::FPCGStaticMeshSpawnerDataProviderProxy(
 
 bool FPCGStaticMeshSpawnerDataProviderProxy::IsValid(FValidationData const& InValidationData) const
 {
-	return InValidationData.ParameterStructSize == sizeof(FParameters);
+	return InValidationData.ParameterStructSize == sizeof(FParameters) && NumCullingCellsTotal > 0;
 }
 
 void FPCGStaticMeshSpawnerDataProviderProxy::AllocateResources(FRDGBuilder& GraphBuilder, FAllocationData const& InAllocationData)
@@ -1107,6 +1200,10 @@ void FPCGStaticMeshSpawnerDataProviderProxy::GatherDispatchData(FDispatchData co
 		Parameters.NumPrimitives = SelectionCDF.Num();
 		Parameters.SelectorAttributeId = SelectorAttributeId;
 		Parameters.SelectedMeshAttributeId = SelectedMeshAttributeId;
+		Parameters.NumCullingCells = NumCullingCells;
+		Parameters.NumCullingCellsTotal = NumCullingCellsTotal;
+		Parameters.BoundsMin = BoundsMin;
+		Parameters.BucketExtent = CullingCellExtent;
 
 		for (int32 Index = 0; Index < AttributeIdOffsetStrides.Num(); ++Index)
 		{
diff --git a/Engine/Plugins/PCG/Source/PCG/Private/Compute/DataInterfaces/Elements/PCGStaticMeshSpawnerDataInterface.h b/Engine/Plugins/PCG/Source/PCG/Private/Compute/DataInterfaces/Elements/PCGStaticMeshSpawnerDataInterface.h
--- a/Engine/Plugins/PCG/Source/PCG/Private/Compute/DataInterfaces/Elements/PCGStaticMeshSpawnerDataInterface.h
+++ b/Engine/Plugins/PCG/Source/PCG/Private/Compute/DataInterfaces/Elements/PCGStaticMeshSpawnerDataInterface.h
@@ -3,6 +3,7 @@
 #pragma once
 
 #include "Compute/DataInterfaces/PCGComputeDataInterface.h"
+#include "Compute/PrimitiveFactories/PCGPrimitiveFactoryPISMC.h"
 #include "ComputeFramework/ComputeDataProvider.h"
 
 #include "Components/PCGProceduralISMComponentDescriptor.h"
@@ -74,13 +75,11 @@ public:
 	/** Attribute Id for mesh selector. */
 	int32 SelectorAttributeId = INDEX_NONE;
 
-	int32 NumInputPoints = 0;
-
 	/** Output attribute Id for the output selected mesh. */
 	int32 SelectedMeshAttributeId = INDEX_NONE;
 
-	/** The number of instances per primitive ID, used for by-attribute spawning. */
-	TMap<int32, uint32> PrimitiveIdToInstanceCount;
+	/** Ranges of instances for each primitive. */
+	TMap</*Primitive Id*/int32, TArray<FPCGInstanceRange>> PrimitiveIdToInstanceRanges;
 
 	int32 AnalysisDataIndex = INDEX_NONE;
 
@@ -88,7 +87,7 @@ public:
 
 	bool bPrimitiveDescriptorsCreated = false;
 
-	TArray<FPCGProceduralISMComponentDescriptor> PrimitiveDescriptors;
+	TArray<FPCGPrimitiveInfo> PrimitiveInfos;
 
 	FPCGPackedCustomData CustomPrimitiveData;
 
@@ -105,6 +104,12 @@ public:
 
 	/** Keeps loaded objects alive. */
 	TSharedPtr<FStreamableHandle> LoadHandle;
+
+	FIntVector NumCullingCells = FIntVector(1, 1, 1);
+	int32 NumCullingCellsTotal = 1;
+	float CullingCellExtent = 0.0;
+
+	FVector3f BoundsMin = FVector3f::ZeroVector;
 };
 
 class FPCGStaticMeshSpawnerDataProviderProxy : public FComputeDataProviderRenderProxy
@@ -117,7 +122,11 @@ public:
 		const TArray<int32>& InPrimitiveStringKeys,
 		TArray<float> InSelectionCDF,
 		int32 InSelectedMeshAttributeId,
-		const TArray<FBox>& InPrimitiveMeshBounds);
+		const TArray<FBox>& InPrimitiveMeshBounds,
+		FIntVector InNumSubdivisions,
+		int32 InNumSubdivisionsTotal,
+		FVector3f InBoundsMin,
+		float InBucketExtent);
 
 	//~ Begin FComputeDataProviderRenderProxy Interface
 	bool IsValid(FValidationData const& InValidationData) const override;
@@ -146,4 +155,9 @@ protected:
 	FRDGBufferSRVRef PrimitiveMeshBoundsMaxBufferSRV = nullptr;
 
 	int32 SelectedMeshAttributeId = INDEX_NONE;
+
+	FIntVector NumCullingCells = FIntVector(1, 1, 1);
+	int32 NumCullingCellsTotal = 1;
+	float CullingCellExtent = 0.0f;
+	FVector3f BoundsMin = FVector3f::ZeroVector;
 };
diff --git a/Engine/Plugins/PCG/Source/PCG/Private/Compute/DataInterfaces/PCGInstanceDataInterface.cpp b/Engine/Plugins/PCG/Source/PCG/Private/Compute/DataInterfaces/PCGInstanceDataInterface.cpp
--- a/Engine/Plugins/PCG/Source/PCG/Private/Compute/DataInterfaces/PCGInstanceDataInterface.cpp
+++ b/Engine/Plugins/PCG/Source/PCG/Private/Compute/DataInterfaces/PCGInstanceDataInterface.cpp
@@ -70,7 +70,8 @@ void UPCGInstanceDataInterface::GetSupportedOutputs(TArray<FShaderFunctionDefini
 	OutFunctions.AddDefaulted_GetRef()
 		.SetName(TEXT("InstanceData_GetIndexToWriteTo"))
 		.AddReturnType(EShaderFundamentalType::Int)
-		.AddParam(EShaderFundamentalType::Uint); // InPrimitiveIndex
+		.AddParam(EShaderFundamentalType::Uint) // InPrimitiveIndex
+		.AddParam(EShaderFundamentalType::Uint); // InCullingCellId
 
 	OutFunctions.AddDefaulted_GetRef()
 		.SetName(TEXT("InstanceData_WriteInstanceLocalToWorld"))
@@ -91,10 +92,11 @@ void UPCGInstanceDataInterface::GetSupportedOutputs(TArray<FShaderFunctionDefini
 }
 
 BEGIN_SHADER_PARAMETER_STRUCT(FPCGInstanceDataInterfaceParameters,)
-	SHADER_PARAMETER_ARRAY(FUintVector4, PrimitiveParameters, [PCGComputeConstants::MAX_PRIMITIVE_COMPONENTS_PER_SPAWNER]) // (NumInstancesAllocated, InstanceOffset, Unused, Unused)
 	SHADER_PARAMETER(uint32, NumPrimitives)
 	SHADER_PARAMETER(uint32, NumInstancesAllPrimitives)
 	SHADER_PARAMETER(uint32, NumCustomFloatsPerInstance)
+	SHADER_PARAMETER(int32, NumCullingCellsTotal)
+	SHADER_PARAMETER_RDG_BUFFER_SRV(StructuredBuffer<FUintVector2>, PrimitiveParameters)
 	SHADER_PARAMETER_RDG_BUFFER_UAV(RWStructuredBuffer<FUintVector4>, InstanceData)
 	SHADER_PARAMETER_RDG_BUFFER_UAV(RWStructuredBuffer<uint32>, InstanceCustomFloatData)
 	SHADER_PARAMETER_RDG_BUFFER_UAV(RWStructuredBuffer<uint32>, WriteCounters)
@@ -143,6 +145,16 @@ UComputeDataProvider* UPCGInstanceDataInterface::CreateDataProvider() const
 	return NewObject<UPCGInstanceDataProvider>();
 }
 
+void UPCGInstanceDataProvider::Initialize(const UComputeDataInterface* InDataInterface, UObject* InBinding, uint64 InInputMask, uint64 InOutputMask)
+{
+	Super::Initialize(InDataInterface, InBinding, InInputMask, InOutputMask);
+
+	const UPCGInstanceDataInterface* InstanceDI = CastChecked<UPCGInstanceDataInterface>(InDataInterface);
+
+	const FIntVector NumCullingCells = CastChecked<UPCGDataBinding>(InBinding)->GetMaxNumCullingCells(InstanceDI->GetCullingCellExtent());
+	NumCullingCellsTotal = NumCullingCells.X * NumCullingCells.Y * NumCullingCells.Z;
+}
+
 bool UPCGInstanceDataProvider::PrepareForExecute_GameThread(UPCGDataBinding* InBinding)
 {
 	TRACE_CPUPROFILER_EVENT_SCOPE(UPCGInstanceDataProvider::PrepareForExecute_GameThread);
@@ -183,9 +195,9 @@ bool UPCGInstanceDataProvider::PrepareForExecute_GameThread(UPCGDataBinding* InB
 
 	NumInstancesAllPrimitives = 0;
 
-	for (int32 Index = 0; Index < PrimitiveFactory->GetNumPrimitives(); ++Index)
+	for (int32 PrimitiveIndex = 0; PrimitiveIndex < PrimitiveFactory->GetNumPrimitives(); ++PrimitiveIndex)
 	{
-		NumInstancesAllPrimitives += PrimitiveFactory->GetNumInstances(Index);
+		NumInstancesAllPrimitives += PrimitiveFactory->GetNumInstancesTotal(PrimitiveIndex);
 	}
 
 	NumCustomFloatsPerInstance = FoundPrimitives->NumCustomFloats;
@@ -216,6 +228,7 @@ FComputeDataProviderRenderProxy* UPCGInstanceDataProvider::GetRenderProxy()
 		.PrimitiveFactory = PrimitiveFactory,
 		.NumInstancesAllPrimitives = NumInstancesAllPrimitives,
 		.NumCustomFloatsPerInstance = NumCustomFloatsPerInstance,
+		.NumCullingCellsTotal = NumCullingCellsTotal,
 		.Seed = Settings ? static_cast<uint32>(Settings->GetSeed()) : 42,
 		.DataProvider = MakeWeakObjectPtr(this),
 		.ContextHandle = ContextHandle,
@@ -231,6 +244,7 @@ void UPCGInstanceDataProvider::Reset()
 	PrimitiveFactory.Reset();
 	NumInstancesAllPrimitives = 0;
 	NumCustomFloatsPerInstance = 0;
+	NumCullingCellsTotal = 1;
 	bWroteInstances = false;
 	ContextHandle.Reset();
 }
@@ -239,6 +253,7 @@ FPCGInstanceDataProviderProxy::FPCGInstanceDataProviderProxy(const FParams& InPa
 	: PrimitiveFactory(InParams.PrimitiveFactory)
 	, NumInstancesAllPrimitives(InParams.NumInstancesAllPrimitives)
 	, NumCustomFloatsPerInstance(InParams.NumCustomFloatsPerInstance)
+	, NumCullingCellsTotal(InParams.NumCullingCellsTotal)
 	, Seed(InParams.Seed)
 	, DataProvider(InParams.DataProvider)
 	, ContextHandle(InParams.ContextHandle)
@@ -276,7 +291,7 @@ FPCGInstanceDataProviderProxy::FPCGInstanceDataProviderProxy(const FParams& InPa
 
 bool FPCGInstanceDataProviderProxy::IsValid(FValidationData const& InValidationData) const
 {
-	return bIsValid && InValidationData.ParameterStructSize == sizeof(FParameters);
+	return bIsValid && InValidationData.ParameterStructSize == sizeof(FParameters) && NumCullingCellsTotal > 0;
 }
 
 void FPCGInstanceDataProviderProxy::AllocateResources(FRDGBuilder& GraphBuilder, FAllocationData const& InAllocationData)
@@ -285,22 +300,48 @@ void FPCGInstanceDataProviderProxy::AllocateResources(FRDGBuilder& GraphBuilder,
 
 	const uint32 StrideUint4s = 3;
 	InstanceData = GraphBuilder.CreateBuffer(FRDGBufferDesc::CreateStructuredDesc(sizeof(FUintVector4), FMath::Max(1u, NumInstancesAllPrimitives) * StrideUint4s), TEXT("PCGInstanceDataBuffer"));
-	InstanceDataSRV = GraphBuilder.CreateSRV(InstanceData);
 	InstanceDataUAV = GraphBuilder.CreateUAV(InstanceData);
 
 	const uint32 CustomFloatsRequired = FMath::Max(1u, NumInstancesAllPrimitives * NumCustomFloatsPerInstance);
 	InstanceCustomFloatData = GraphBuilder.CreateBuffer(FRDGBufferDesc::CreateStructuredDesc(sizeof(uint32), CustomFloatsRequired), TEXT("PCGInstanceCustomFloatDataBuffer"));
-	InstanceCustomFloatDataSRV = GraphBuilder.CreateSRV(InstanceCustomFloatData);
 	InstanceCustomFloatDataUAV = GraphBuilder.CreateUAV(InstanceCustomFloatData);
 
-	const int32 NumCountersRequired = FMath::Max(1u, PrimitiveFactory ? PrimitiveFactory->GetNumPrimitives() : 0u);
-	WriteCounters = GraphBuilder.CreateBuffer(FRDGBufferDesc::CreateStructuredDesc(sizeof(uint32), NumCountersRequired), TEXT("PCGWriteCounters"));
-	WriteCountersSRV = GraphBuilder.CreateSRV(WriteCounters);
+	const int32 NumPrimitiveCells = FMath::Max(1u, PrimitiveFactory ? PrimitiveFactory->GetNumPrimitives() : 0u) * NumCullingCellsTotal;
+
+	WriteCounters = GraphBuilder.CreateBuffer(FRDGBufferDesc::CreateStructuredDesc(sizeof(uint32), NumPrimitiveCells), TEXT("PCGWriteCounters"));
 	WriteCountersUAV = GraphBuilder.CreateUAV(WriteCounters);
 
-	TArray<uint32, TInlineAllocator<32>> Zeros;
-	Zeros.SetNumZeroed(NumCountersRequired);
-	GraphBuilder.QueueBufferUpload(WriteCounters, MakeConstArrayView(Zeros));
+	TArray<uint32, TInlineAllocator<128>> Zeroes;
+	Zeroes.SetNumZeroed(NumPrimitiveCells);
+	GraphBuilder.QueueBufferUpload(WriteCounters, MakeConstArrayView(Zeroes));
+
+	FRDGBufferRef PrimitiveParameters = GraphBuilder.CreateBuffer(FRDGBufferDesc::CreateStructuredDesc(sizeof(FUintVector2), NumPrimitiveCells), TEXT("PCGPrimitiveParameters"));
+	PrimitiveParametersSRV = GraphBuilder.CreateSRV(PrimitiveParameters);
+
+	if (PrimitiveFactory && PrimitiveFactory->GetNumPrimitives() > 0)
+	{
+		TArray<FUintVector2, TInlineAllocator<128>> PrimitiveParams;
+		uint32 CumulativeInstanceCount = 0;
+
+		for (int32 PrimitiveIndex = 0; PrimitiveIndex < PrimitiveFactory->GetNumPrimitives(); ++PrimitiveIndex)
+		{
+			for (int32 CellID = 0; CellID < NumCullingCellsTotal; ++CellID)
+			{
+				// This is the full map (not just populated cells) - for each incoming instance the kernel will look up an instance offset from this array.
+				const uint32 InstanceCount = PrimitiveFactory->GetNumInstances(PrimitiveIndex, CellID);
+
+				PrimitiveParams.Add(FUintVector2(InstanceCount, CumulativeInstanceCount));
+
+				CumulativeInstanceCount += InstanceCount;
+			}
+		}
+
+		GraphBuilder.QueueBufferUpload(PrimitiveParameters, MakeConstArrayView(PrimitiveParams));
+	}
+	else
+	{
+		GraphBuilder.QueueBufferUpload(PrimitiveParameters, MakeConstArrayView(&FUintVector2::ZeroValue, 1));
+	}
 }
 
 void FPCGInstanceDataProviderProxy::GatherDispatchData(FDispatchData const& InDispatchData)
@@ -315,23 +356,8 @@ void FPCGInstanceDataProviderProxy::GatherDispatchData(FDispatchData const& InDi
 		Parameters.InstanceData = InstanceDataUAV;
 		Parameters.InstanceCustomFloatData = InstanceCustomFloatDataUAV;
 		Parameters.WriteCounters = WriteCountersUAV;
-
-		if (PrimitiveFactory)
-		{
-			uint32 CumulativeInstanceCount = 0;
-
-			for (int32 Index = 0; Index < PrimitiveFactory->GetNumPrimitives(); ++Index)
-			{
-				const uint32 PrimitiveInstanceCount = PrimitiveFactory->GetNumInstances(Index);
-
-				Parameters.PrimitiveParameters[Index] = FUintVector4(
-					PrimitiveInstanceCount,
-					CumulativeInstanceCount,
-					0, 0); // Unused
-
-				CumulativeInstanceCount += PrimitiveInstanceCount;
-			}
-		}
+		Parameters.PrimitiveParameters = PrimitiveParametersSRV;
+		Parameters.NumCullingCellsTotal = NumCullingCellsTotal;
 	}
 }
 
@@ -348,7 +374,7 @@ void FPCGInstanceDataProviderProxy::PostSubmit(FComputeContext& Context) const
 	const TRefCountPtr<FRDGPooledBuffer> InstanceCustomFloatDataExported = Context.GraphBuilder.ConvertToExternalBuffer(InstanceCustomFloatData);
 	const TRefCountPtr<FRDGPooledBuffer> WriteCountersExported = Context.GraphBuilder.ConvertToExternalBuffer(WriteCounters);
 	
-	FTSTicker::GetCoreTicker().AddTicker(TEXT("ApplyPrimitiveSceneUpdates"), /*InDelay=*/0.0f, [DataProviderWeak=DataProvider, DataProviderGeneration=DataProviderGeneration, InstanceDataExported, InstanceCustomFloatDataExported, WriteCountersExported, ContextHandle=ContextHandle, Seed=Seed](float)
+	FTSTicker::GetCoreTicker().AddTicker(TEXT("ApplyPrimitiveSceneUpdates"), /*InDelay=*/0.0f, [DataProviderWeak=DataProvider, DataProviderGeneration=DataProviderGeneration, InstanceDataExported, InstanceCustomFloatDataExported, WriteCountersExported, ContextHandle=ContextHandle, Seed=Seed, NumCullingCellsTotal=NumCullingCellsTotal](float)
 	{
 		UPCGInstanceDataProvider* DataProviderInner = DataProviderWeak.Get();
 		if (!DataProviderInner || DataProviderInner->GetGenerationCounter() != DataProviderGeneration)
@@ -399,7 +425,8 @@ void FPCGInstanceDataProviderProxy::PostSubmit(FComputeContext& Context) const
 
 		for (int32 PrimitiveIndex = 0; PrimitiveIndex < Factory->GetNumPrimitives(); ++PrimitiveIndex)
 		{
-			const int32 PrimitiveNumInstances = Factory->GetNumInstances(PrimitiveIndex);
+			const int32 PrimitiveNumInstances = Factory->GetNumInstancesTotal(PrimitiveIndex);
+
 			if (PrimitiveNumInstances <= 0)
 			{
 				UE_LOG(LogPCG, Warning, TEXT("Primitive with 0 instances encountered during instancing, should not happen."));
@@ -425,7 +452,7 @@ void FPCGInstanceDataProviderProxy::PostSubmit(FComputeContext& Context) const
 			FPrimitiveSceneDesc PrimitiveSceneDesc;
 			PrimitiveSceneDesc.SceneProxy = SceneProxy;
 
-			Scene->UpdatePrimitiveInstancesFromCompute(&PrimitiveSceneDesc, FGPUSceneWriteDelegate::CreateLambda([PrimitiveIndex, PrimitiveNumInstances, CumulativeInstanceCount, Seed, bTriggerCapture, InstanceDataExported, InstanceCustomFloatDataExported, WriteCountersExported](FRDGBuilder& InGraphBuilder, const FGPUSceneWriteDelegateParams& Params)
+			Scene->UpdatePrimitiveInstancesFromCompute(&PrimitiveSceneDesc, FGPUSceneWriteDelegate::CreateLambda([PrimitiveIndex, PrimitiveNumInstances, CumulativeInstanceCount, Seed, bTriggerCapture, InstanceDataExported, InstanceCustomFloatDataExported, WriteCountersExported, NumCullingCellsTotal=NumCullingCellsTotal](FRDGBuilder& InGraphBuilder, const FGPUSceneWriteDelegateParams& Params)
 			{
 				RDG_EVENT_SCOPE(InGraphBuilder, "SceneComputeUpdateInterface->EnqueueUpdate");
 				check(Params.PersistentPrimitiveId != (uint32)INDEX_NONE);
@@ -453,6 +480,7 @@ void FPCGInstanceDataProviderProxy::PostSubmit(FComputeContext& Context) const
 				Parameters->InCustomDataCount = Params.NumCustomDataFloats;
 				Parameters->InPayloadDataFlags = Params.PackedInstanceSceneDataFlags;
 				Parameters->InSeed = Seed;
+				Parameters->InNumCullingCellsTotal = NumCullingCellsTotal;
 
 #if PCG_INSTANCE_DATA_LOGGING
 				UE_LOG(LogPCG, Log, TEXT("\tScene writer delegate [%d]:\tPrimitive ID %u,\tsource instance offset %u,\tInstanceSceneDataOffset %u, num instances %u"),
diff --git a/Engine/Plugins/PCG/Source/PCG/Private/Compute/DataInterfaces/PCGInstanceDataInterface.h b/Engine/Plugins/PCG/Source/PCG/Private/Compute/DataInterfaces/PCGInstanceDataInterface.h
--- a/Engine/Plugins/PCG/Source/PCG/Private/Compute/DataInterfaces/PCGInstanceDataInterface.h
+++ b/Engine/Plugins/PCG/Source/PCG/Private/Compute/DataInterfaces/PCGInstanceDataInterface.h
@@ -34,10 +34,19 @@ public:
 	UComputeDataProvider* CreateDataProvider() const override;
 	//~ End UComputeDataInterface Interface
 
-public:
+	void SetInputPinProvidingData(FName InInputPinProvidingData) { InputPinProvidingData = InInputPinProvidingData; }
+
+	int32 GetCullingCellExtent() const { return CullingCellExtent; }
+	void SetCullingCellExtent(int32 InGridSize) { CullingCellExtent = InGridSize; }
+
+protected:
 	UPROPERTY()
 	FName InputPinProvidingData = NAME_None;
 
+	/** Extent of culling cells in cm. */
+	UPROPERTY()
+	int32 CullingCellExtent = 0;
+
 private:
 	static TCHAR const* TemplateFilePath;
 };
@@ -50,6 +59,7 @@ class UPCGInstanceDataProvider : public UPCGComputeDataProvider
 
 public:
 	//~ Begin UComputeDataProvider Interface
+	virtual void Initialize(const UComputeDataInterface* InDataInterface, UObject* InBinding, uint64 InInputMask, uint64 InOutputMask) override;
 	virtual bool PrepareForExecute_GameThread(UPCGDataBinding* InBinding) override;
 	virtual bool PostExecute(UPCGDataBinding* InBinding) override;
 	FComputeDataProviderRenderProxy* GetRenderProxy() override;
@@ -65,6 +75,8 @@ public:
 	/** Whether instance data has been applied to scene and operation is complete. */
 	bool bWroteInstances = false;
 
+	int32 NumCullingCellsTotal = 1;
+
 	TWeakPtr<FPCGContextHandle> ContextHandle;
 };
 
@@ -76,6 +88,7 @@ public:
 		TSharedPtr<IPCGRuntimePrimitiveFactory> PrimitiveFactory;
 		uint32 NumInstancesAllPrimitives = 0;
 		uint32 NumCustomFloatsPerInstance = 0;
+		int32 NumCullingCellsTotal = 1;
 		uint32 Seed = 42;
 		TWeakObjectPtr<UPCGInstanceDataProvider> DataProvider;
 		TWeakPtr<FPCGContextHandle> ContextHandle;
@@ -99,21 +112,21 @@ protected:
 
 	/** Instance transforms of all instances across all primitives. Three float4s per instance. */
 	FRDGBufferRef InstanceData = nullptr;
-	FRDGBufferSRVRef InstanceDataSRV = nullptr;
 	FRDGBufferUAVRef InstanceDataUAV = nullptr;
 
 	/** Custom float data for all instances across all primitives. N per instance. Number of custom floats N comes from inputs (and is static). */
 	FRDGBufferRef InstanceCustomFloatData = nullptr;
-	FRDGBufferSRVRef InstanceCustomFloatDataSRV = nullptr;
 	FRDGBufferUAVRef InstanceCustomFloatDataUAV = nullptr;
 
 	/** A per-primitive instance counter. Updated atomically and used to place instances into a segmented array (one segment per primitive). */
 	FRDGBufferRef WriteCounters = nullptr;
-	FRDGBufferSRVRef WriteCountersSRV = nullptr;
 	FRDGBufferUAVRef WriteCountersUAV = nullptr;
 
+	FRDGBufferSRVRef PrimitiveParametersSRV = nullptr;
+
 	uint32 NumInstancesAllPrimitives = 0;
 	uint32 NumCustomFloatsPerInstance = 0;
+	int32 NumCullingCellsTotal = 1;
 	uint32 Seed = 42;
 
 	TWeakObjectPtr<UPCGInstanceDataProvider> DataProvider;
diff --git a/Engine/Plugins/PCG/Source/PCG/Private/Compute/PCGDataBinding.cpp b/Engine/Plugins/PCG/Source/PCG/Private/Compute/PCGDataBinding.cpp
--- a/Engine/Plugins/PCG/Source/PCG/Private/Compute/PCGDataBinding.cpp
+++ b/Engine/Plugins/PCG/Source/PCG/Private/Compute/PCGDataBinding.cpp
@@ -149,9 +149,6 @@ TSharedPtr<const FPCGDataCollectionDesc> UPCGDataBinding::ComputeKernelPinDataDe
 
 TSharedPtr<const FPCGDataCollectionDesc> UPCGDataBinding::ComputeKernelBindingDataDesc(int32 InGraphBindingIndex)
 {
-	// Uncomment to be notified when a kernel pin description is computed late:
-	//ensureAlwaysMsgf(!IsInGameThread(), TEXT("Data descriptions should only be computed when priming the cache, which takes place on a worker thread."));
-	ensure(!bIsDataDescriptionCachePrimed);
 	ensure(bTablesInitialized);
 
 	if (!ensure(Graph))
@@ -165,6 +162,8 @@ TSharedPtr<const FPCGDataCollectionDesc> UPCGDataBinding::ComputeKernelBindingDa
 		return *FoundDesc;
 	}
 
+	ensureMsgf(!bIsDataDescriptionCachePrimed, TEXT("Data description cache miss, investigate why the data description was not cached for this pin during priming."));
+
 	// If the description wasn't found in the cache, compute it now.
 	const TSharedPtr<const FPCGDataCollectionDesc> ComputedDesc = Graph->ComputeKernelBindingDataDesc(InGraphBindingIndex, this);
 	DataDescriptionCache.Add(InGraphBindingIndex, ComputedDesc);
@@ -401,6 +400,26 @@ const FPCGDataCollection& UPCGDataBinding::GetOutputDataCollection() const
 	return OutputDataCollection;
 }
 
+FIntVector UPCGDataBinding::GetMaxNumCullingCells(int32 InCullingCellExtent) const
+{
+	if (InCullingCellExtent <= 0)
+	{
+		return FIntVector(1, 1, 1);
+	}
+
+	const IPCGGraphExecutionSource* ExecSource = GetExecutionSource();
+	if (!ExecSource)
+	{
+		return FIntVector(1, 1, 1);
+	}
+
+	const FBox Bounds = ExecSource->GetExecutionState().GetBounds();
+	const double CellExtent = static_cast<double>(InCullingCellExtent);
+	const FIntVector LowerLeft(FMath::Floor(Bounds.Min.X / CellExtent), FMath::Floor(Bounds.Min.Y / CellExtent), FMath::Floor(Bounds.Min.Z / CellExtent));
+	const FIntVector UpperRight(FMath::Floor(Bounds.Max.X / CellExtent), FMath::Floor(Bounds.Max.Y / CellExtent), FMath::Floor(Bounds.Max.Z / CellExtent));
+	return UpperRight - LowerLeft + FIntVector(1, 1, 1);
+}
+
 void UPCGDataBinding::AddInputDataAttributesToTable()
 {
 	if (InputDataCollection.TaggedData.IsEmpty())
diff --git a/Engine/Plugins/PCG/Source/PCG/Private/Compute/PrimitiveFactories/IPCGRuntimePrimitiveFactory.cpp b/Engine/Plugins/PCG/Source/PCG/Private/Compute/PrimitiveFactories/IPCGRuntimePrimitiveFactory.cpp
new file mode 100644
--- /dev/null
+++ b/Engine/Plugins/PCG/Source/PCG/Private/Compute/PrimitiveFactories/IPCGRuntimePrimitiveFactory.cpp
@@ -0,0 +1,11 @@
+// Copyright Epic Games, Inc. All Rights Reserved.
+
+#include "Compute/PrimitiveFactories/IPCGRuntimePrimitiveFactory.h"
+
+int32 IPCGRuntimePrimitiveFactory::GetNumInstances(int32 InPrimitiveIndex, int32 InCellID) const
+{
+	// Fallback used if this new GetNumInstances function is not implemented.
+	PRAGMA_DISABLE_DEPRECATION_WARNINGS
+	return GetNumInstances(InPrimitiveIndex);
+	PRAGMA_ENABLE_DEPRECATION_WARNINGS
+}
diff --git a/Engine/Plugins/PCG/Source/PCG/Private/Compute/PrimitiveFactories/PCGPrimitiveFactoryISKMC.cpp b/Engine/Plugins/PCG/Source/PCG/Private/Compute/PrimitiveFactories/PCGPrimitiveFactoryISKMC.cpp
--- a/Engine/Plugins/PCG/Source/PCG/Private/Compute/PrimitiveFactories/PCGPrimitiveFactoryISKMC.cpp
+++ b/Engine/Plugins/PCG/Source/PCG/Private/Compute/PrimitiveFactories/PCGPrimitiveFactoryISKMC.cpp
@@ -25,9 +25,14 @@ FPrimitiveSceneProxy* FPCGPrimitiveFactoryISKMC::GetSceneProxy(int32 InPrimitive
 	}
 }
 
-int32 FPCGPrimitiveFactoryISKMC::GetNumInstances(int32 InPrimitiveIndex) const
+int32 FPCGPrimitiveFactoryISKMC::GetNumInstances(int32 InPrimitiveIndex, int32 InCellID) const
 {
-	return NumInstances.IsValidIndex(InPrimitiveIndex) ? NumInstances[InPrimitiveIndex] : -1;
+	return NumInstances.IsValidIndex(InPrimitiveIndex) ? NumInstances[InPrimitiveIndex] : 0;
+}
+
+int32 FPCGPrimitiveFactoryISKMC::GetNumInstancesTotal(int32 InPrimitiveIndex) const
+{
+	return NumInstances.IsValidIndex(InPrimitiveIndex) ? NumInstances[InPrimitiveIndex] : 0;
 }
 
 bool FPCGPrimitiveFactoryISKMC::IsRenderStateCreated() const
diff --git a/Engine/Plugins/PCG/Source/PCG/Private/Compute/PrimitiveFactories/PCGPrimitiveFactoryISKMC.h b/Engine/Plugins/PCG/Source/PCG/Private/Compute/PrimitiveFactories/PCGPrimitiveFactoryISKMC.h
--- a/Engine/Plugins/PCG/Source/PCG/Private/Compute/PrimitiveFactories/PCGPrimitiveFactoryISKMC.h
+++ b/Engine/Plugins/PCG/Source/PCG/Private/Compute/PrimitiveFactories/PCGPrimitiveFactoryISKMC.h
@@ -16,7 +16,8 @@ public:
 	//~ Begin IPCGRuntimePrimitiveFactory interface
 	virtual int32 GetNumPrimitives() const override { return Components.Num(); }
 	virtual FPrimitiveSceneProxy* GetSceneProxy(int32 InPrimitiveIndex) const override;
-	virtual int32 GetNumInstances(int32 InPrimitiveIndex) const override;
+	virtual int32 GetNumInstances(int32 InPrimitiveIndex, int32 InCellID) const override;
+	virtual int32 GetNumInstancesTotal(int32 InPrimitiveIndex) const override;
 	virtual bool IsRenderStateCreated() const override;
 	//~ End IPCGRuntimePrimitiveFactory interface
 
diff --git a/Engine/Plugins/PCG/Source/PCG/Private/Compute/PrimitiveFactories/PCGPrimitiveFactoryISMBase.cpp b/Engine/Plugins/PCG/Source/PCG/Private/Compute/PrimitiveFactories/PCGPrimitiveFactoryISMBase.cpp
--- a/Engine/Plugins/PCG/Source/PCG/Private/Compute/PrimitiveFactories/PCGPrimitiveFactoryISMBase.cpp
+++ b/Engine/Plugins/PCG/Source/PCG/Private/Compute/PrimitiveFactories/PCGPrimitiveFactoryISMBase.cpp
@@ -26,3 +26,38 @@ namespace PCGPrimitiveFactoryHelpers
 		}
 	}
 }
+
+void IPCGPrimitiveFactoryISMBase::PopulateHashes(const FPCGPrimitiveInfo& InPrimitiveInfo, TArray<FInstanceSceneDataBuffers::FCompressedSpatialHashItem>& OutHashes, FBox& OutBounds) const
+{
+	if (InPrimitiveInfo.InstanceRanges.IsEmpty() || InPrimitiveInfo.CullingCellExtent == 0.0)
+	{
+		return;
+	}
+
+	OutHashes.Reserve(InPrimitiveInfo.InstanceRanges.Num());
+
+	// Minus 1 because CalcLevel returns the level that fully contains this extent, and will return the next level up if a power of 2 cell size is given (which is always
+	// the case for us).
+	const int32 Level = RenderingSpatialHash::CalcLevel(InPrimitiveInfo.CullingCellExtent - 1.0f);
+	RenderingSpatialHash::FLocation64 BaseLocation = RenderingSpatialHash::ToCellLoc(Level, InPrimitiveInfo.Descriptor.WorldBounds.Min);
+
+	for (int32 Index = 0; Index < InPrimitiveInfo.InstanceRanges.Num(); ++Index)
+	{
+		const int32 NumInstances = InPrimitiveInfo.InstanceRanges[Index].GetNumInstances();
+		if (NumInstances > 0)
+		{
+			FInstanceSceneDataBuffers::FCompressedSpatialHashItem& Hash = OutHashes.Emplace_GetRef();
+			Hash.NumInstances = NumInstances;
+			Hash.Location = BaseLocation;
+
+			const uint32 SubdivIndex = InPrimitiveInfo.InstanceRanges[Index].GetCellID();
+			check(InPrimitiveInfo.NumCullingCells.X > 0 && InPrimitiveInfo.NumCullingCells.Y > 0 && InPrimitiveInfo.NumCullingCells.Z > 0);
+			Hash.Location.Coord.X += SubdivIndex % InPrimitiveInfo.NumCullingCells.X;
+			Hash.Location.Coord.Y += (SubdivIndex / InPrimitiveInfo.NumCullingCells.X) % InPrimitiveInfo.NumCullingCells.Y;
+			Hash.Location.Coord.Z += SubdivIndex / (InPrimitiveInfo.NumCullingCells.X * InPrimitiveInfo.NumCullingCells.Y);
+
+			const FVector3d CellBottomCorner = RenderingSpatialHash::CalcWorldPosition(Hash.Location);
+			OutBounds += FBox(CellBottomCorner, CellBottomCorner + FVector3d(InPrimitiveInfo.CullingCellExtent));
+		}
+	}
+}
diff --git a/Engine/Plugins/PCG/Source/PCG/Private/Compute/PrimitiveFactories/PCGPrimitiveFactoryPISMC.cpp b/Engine/Plugins/PCG/Source/PCG/Private/Compute/PrimitiveFactories/PCGPrimitiveFactoryPISMC.cpp
--- a/Engine/Plugins/PCG/Source/PCG/Private/Compute/PrimitiveFactories/PCGPrimitiveFactoryPISMC.cpp
+++ b/Engine/Plugins/PCG/Source/PCG/Private/Compute/PrimitiveFactories/PCGPrimitiveFactoryPISMC.cpp
@@ -24,9 +24,35 @@ FPrimitiveSceneProxy* FPCGPrimitiveFactoryPISMC::GetSceneProxy(int32 InPrimitive
 	}
 }
 
-int32 FPCGPrimitiveFactoryPISMC::GetNumInstances(int32 InPrimitiveIndex) const
+int32 FPCGPrimitiveFactoryPISMC::GetNumInstances(int32 InPrimitiveIndex, int32 InCellID) const
 {
-	return NumInstances.IsValidIndex(InPrimitiveIndex) ? NumInstances[InPrimitiveIndex] : -1;
+	if (ensure(PrimitiveInfos.IsValidIndex(InPrimitiveIndex)))
+	{
+		for (const FPCGInstanceRange& InstanceRange : PrimitiveInfos[InPrimitiveIndex].InstanceRanges)
+		{
+			if (InstanceRange.GetCellID() == InCellID)
+			{
+				return InstanceRange.GetNumInstances();
+			}
+		}
+	}
+
+	return 0;
+}
+
+int32 FPCGPrimitiveFactoryPISMC::GetNumInstancesTotal(int32 InPrimitiveIndex) const
+{
+	int32 NumInstances = 0;
+
+	if (ensure(PrimitiveInfos.IsValidIndex(InPrimitiveIndex)))
+	{
+		for (const FPCGInstanceRange& PrimitiveInfo : PrimitiveInfos[InPrimitiveIndex].InstanceRanges)
+		{
+			NumInstances += PrimitiveInfo.GetNumInstances();
+		}
+	}
+
+	return NumInstances;
 }
 
 bool FPCGPrimitiveFactoryPISMC::IsRenderStateCreated() const
@@ -45,7 +71,7 @@ bool FPCGPrimitiveFactoryPISMC::IsRenderStateCreated() const
 
 void FPCGPrimitiveFactoryPISMC::Initialize(FParameters&& InParameters)
 {
-	Descriptors = MoveTemp(InParameters.Descriptors);
+	PrimitiveInfos = MoveTemp(InParameters.PrimitiveInfos);
 	CustomPrimitiveData = MoveTemp(InParameters.CustomPrimitiveData);
 	TargetActor = InParameters.TargetActor;
 }
@@ -63,11 +89,22 @@ bool FPCGPrimitiveFactoryPISMC::Create(FPCGContext* InContext)
 		return true;
 	}
 
-	const int32 NumToCreate = FMath::Min(Descriptors.Num(), PCGComputeConstants::MAX_PRIMITIVE_COMPONENTS_PER_SPAWNER);
+	const int32 NumToCreate = FMath::Min(PrimitiveInfos.Num(), PCGComputeConstants::MAX_PRIMITIVE_COMPONENTS_PER_SPAWNER);
 
 	while (NumPrimitivesProcessed < NumToCreate)
 	{
-		FPCGProceduralISMComponentDescriptor& Desc = Descriptors[NumPrimitivesProcessed];
+		FPCGPrimitiveInfo& PrimitiveInfo = PrimitiveInfos[NumPrimitivesProcessed];
+		FPCGProceduralISMComponentDescriptor& Desc = PrimitiveInfo.Descriptor;
+
+		TArray<FInstanceSceneDataBuffers::FCompressedSpatialHashItem> Hashes;
+		FBox BoundsOfAllCells(EForceInit::ForceInit);
+		
+		PopulateHashes(PrimitiveInfo, Hashes, BoundsOfAllCells);
+		
+		if (BoundsOfAllCells.IsValid)
+		{
+			Desc.WorldBounds = Desc.WorldBounds.Overlap(BoundsOfAllCells);
+		}
 
 		const int32 PrimNumInstances = Desc.NumInstances;
 		FBox BoundingBox = Desc.StaticMesh.IsValid() ? Desc.StaticMesh->GetBoundingBox() : FBox{};
@@ -83,8 +120,10 @@ bool FPCGPrimitiveFactoryPISMC::Create(FPCGContext* InContext)
 
 		if (ManagedComponent)
 		{
-			Components.Add(ManagedComponent->GetComponent());
-			NumInstances.Add(PrimNumInstances);
+			UPCGProceduralISMComponent* PISMC = ManagedComponent->GetComponent();
+			PISMC->SetSpatialHashes(MoveTemp(Hashes));
+
+			Components.Add(PISMC);
 			MeshBounds.Add(MoveTemp(BoundingBox));
 		}
 
diff --git a/Engine/Plugins/PCG/Source/PCG/Private/Compute/PrimitiveFactories/PCGPrimitiveFactoryPISMC.h b/Engine/Plugins/PCG/Source/PCG/Private/Compute/PrimitiveFactories/PCGPrimitiveFactoryPISMC.h
--- a/Engine/Plugins/PCG/Source/PCG/Private/Compute/PrimitiveFactories/PCGPrimitiveFactoryPISMC.h
+++ b/Engine/Plugins/PCG/Source/PCG/Private/Compute/PrimitiveFactories/PCGPrimitiveFactoryPISMC.h
@@ -15,7 +15,8 @@ public:
 	//~ Begin IPCGRuntimePrimitiveFactory interface
 	virtual int32 GetNumPrimitives() const override { return Components.Num(); }
 	virtual FPrimitiveSceneProxy* GetSceneProxy(int32 InPrimitiveIndex) const override;
-	virtual int32 GetNumInstances(int32 InPrimitiveIndex) const override;
+	virtual int32 GetNumInstances(int32 InPrimitiveIndex, int32 InCellID) const override;
+	virtual int32 GetNumInstancesTotal(int32 InPrimitiveIndex) const override;
 	virtual bool IsRenderStateCreated() const override;
 	//~ End IPCGRuntimePrimitiveFactory interface
 
@@ -26,13 +27,12 @@ public:
 	//~ End IPCGPrimitiveFactoryISMBase interface
 
 protected:
-	TArray<FPCGProceduralISMComponentDescriptor> Descriptors;
+	TArray<FPCGPrimitiveInfo> PrimitiveInfos;
 	FPCGPackedCustomData CustomPrimitiveData;
 	AActor* TargetActor = nullptr;
 
 	int32 NumPrimitivesProcessed = 0;
 
 	TArray<TWeakObjectPtr<UPrimitiveComponent>> Components;
-	TArray<int32> NumInstances;
 	TArray<FBox> MeshBounds;
 };
diff --git a/Engine/Plugins/PCG/Source/PCG/Private/Elements/Metadata/PCGMetadataPartition.cpp b/Engine/Plugins/PCG/Source/PCG/Private/Elements/Metadata/PCGMetadataPartition.cpp
--- a/Engine/Plugins/PCG/Source/PCG/Private/Elements/Metadata/PCGMetadataPartition.cpp
+++ b/Engine/Plugins/PCG/Source/PCG/Private/Elements/Metadata/PCGMetadataPartition.cpp
@@ -77,6 +77,7 @@ void UPCGMetadataPartitionSettings::CreateKernels(FPCGGPUCompilationContext& InO
 	UPCGCountUniqueAttributeValuesKernel* CountKernel = PCGKernelHelpers::CreateKernel<UPCGCountUniqueAttributeValuesKernel>(InOutContext, CreateParams, OutKernels, OutEdges);
 	CountKernel->SetAttributeName(PartitionAttributeSelectors[0].GetAttributeName());
 	CountKernel->SetEmitPerDataCounts(true);
+	CountKernel->SetMaxNumUniqueValuesSource(EPCGMaxNumUniqueValuesSource::StringKeyValues);
 
 	// Wire count kernel to copy partition kernel.
 	OutEdges.Emplace(FPCGPinReference(CountKernel, PCGPinConstants::DefaultOutputLabel), FPCGPinReference(PartitionKernel, PCGMetadataPartitionConstants::ElementCountsPinLabel));
diff --git a/Engine/Plugins/PCG/Source/PCG/Private/Elements/PCGSkinnedMeshSpawner.cpp b/Engine/Plugins/PCG/Source/PCG/Private/Elements/PCGSkinnedMeshSpawner.cpp
--- a/Engine/Plugins/PCG/Source/PCG/Private/Elements/PCGSkinnedMeshSpawner.cpp
+++ b/Engine/Plugins/PCG/Source/PCG/Private/Elements/PCGSkinnedMeshSpawner.cpp
@@ -54,6 +54,7 @@ void UPCGSkinnedMeshSpawnerSettings::CreateKernels(FPCGGPUCompilationContext& In
 		CountKernel->SetAttributeName(MeshSelectorParameters->MeshAttribute.GetName());
 		// We operate across all input data rather than spawning for each input data separately.
 		CountKernel->SetEmitPerDataCounts(false);
+		CountKernel->SetMaxNumUniqueValuesSource(EPCGMaxNumUniqueValuesSource::StringKeyValues);
 
 		OutEdges.Emplace(FPCGPinReference(CountKernel, PCGPinConstants::DefaultOutputLabel), FPCGPinReference(SpawnerKernel, PCGSkinnedMeshSpawnerConstants::InstanceCountsPinLabel));
 	}
diff --git a/Engine/Plugins/PCG/Source/PCG/Private/Elements/PCGSkinnedMeshSpawnerKernel.cpp b/Engine/Plugins/PCG/Source/PCG/Private/Elements/PCGSkinnedMeshSpawnerKernel.cpp
--- a/Engine/Plugins/PCG/Source/PCG/Private/Elements/PCGSkinnedMeshSpawnerKernel.cpp
+++ b/Engine/Plugins/PCG/Source/PCG/Private/Elements/PCGSkinnedMeshSpawnerKernel.cpp
@@ -78,7 +78,7 @@ void UPCGSkinnedMeshSpawnerKernel::CreateAdditionalOutputDataInterfaces(FPCGGPUC
 
 	UPCGInstanceDataInterface* InstanceDI = InOutContext.NewObject_AnyThread<UPCGInstanceDataInterface>(InObjectOuter);
 	InstanceDI->SetProducerKernel(this);
-	InstanceDI->InputPinProvidingData = PCGPinConstants::DefaultInputLabel;
+	InstanceDI->SetInputPinProvidingData(PCGPinConstants::DefaultInputLabel);
 	OutDataInterfaces.Add(InstanceDI);
 }
 #endif
diff --git a/Engine/Plugins/PCG/Source/PCG/Private/Elements/PCGStaticMeshSpawner.cpp b/Engine/Plugins/PCG/Source/PCG/Private/Elements/PCGStaticMeshSpawner.cpp
--- a/Engine/Plugins/PCG/Source/PCG/Private/Elements/PCGStaticMeshSpawner.cpp
+++ b/Engine/Plugins/PCG/Source/PCG/Private/Elements/PCGStaticMeshSpawner.cpp
@@ -31,18 +31,26 @@
 
 #define LOCTEXT_NAMESPACE "PCGStaticMeshSpawnerElement"
 
-static TAutoConsoleVariable<bool> CVarAllowISMReuse(
-	TEXT("pcg.ISM.AllowReuse"),
-	true,
-	TEXT("Controls whether ISMs can be reused and skipped when re-executing"));
-
-static TAutoConsoleVariable<bool> CVarUseRawDataForAnalysis(
-	TEXT("pcg.RuntimeGeneration.ISM.UseRawDataForGPUAnalysis"),
-	true,
-	TEXT("Read back raw data in GPU primitive instancing path rather than constructing an attribute set, to save CPU time. Read-only, must be set from ini."),
-	ECVF_ReadOnly // Structural change, could flush graph compilation cache, but trickier to flush compute graph instance pool in UPCGGraph currently.
+namespace PCGStaticMeshSpawner
+{
+	static TAutoConsoleVariable<bool> CVarAllowISMReuse(
+		TEXT("pcg.ISM.AllowReuse"),
+		true,
+		TEXT("Controls whether ISMs can be reused and skipped when re-executing"));
+
+	static TAutoConsoleVariable<bool> CVarUseRawDataForAnalysis(
+		TEXT("pcg.RuntimeGeneration.ISM.UseRawDataForGPUAnalysis"),
+		true,
+		TEXT("Read back raw data in GPU primitive instancing path rather than constructing an attribute set, to save CPU time. Read-only, must be set from ini."),
+		ECVF_ReadOnly // Structural change, could flush graph compilation cache, but trickier to flush compute graph instance pool in UPCGGraph currently.
 	);
 
+	static int32 SanitizeCullGridSize(int32 In)
+	{
+		return FMath::RoundUpToPowerOfTwo(FMath::Max(In, 64));
+	}
+};
+
 UPCGStaticMeshSpawnerSettings::UPCGStaticMeshSpawnerSettings(const FObjectInitializer &ObjectInitializer)
 {
 	MeshSelectorType = UPCGMeshSelectorWeighted::StaticClass();
@@ -62,7 +70,7 @@ void UPCGStaticMeshSpawnerSettings::CreateKernels(FPCGGPUCompilationContext& InO
 
 	// Only the by-primitive-data selector supports the raw buffer optimization, for which the analysis data contains table indices.
 	// For by-attribute spawning the analysis data contains upstream string keys which we can't currently remap to the downstream graph string keys.
-	const bool bUseRawBuffer = SelectorPrimitiveData && CVarUseRawDataForAnalysis.GetValueOnAnyThread();
+	const bool bUseRawBuffer = SelectorPrimitiveData && PCGStaticMeshSpawner::CVarUseRawDataForAnalysis.GetValueOnAnyThread();
 
 	UPCGStaticMeshSpawnerKernel* SpawnerKernel = nullptr;
 
@@ -87,6 +95,7 @@ void UPCGStaticMeshSpawnerSettings::CreateKernels(FPCGGPUCompilationContext& InO
 		SpawnerKernel = PCGKernelHelpers::CreateKernel<UPCGStaticMeshSpawnerKernel>(InOutContext, CreateParams, OutKernels, OutEdges);
 		SpawnerKernel->SetUseRawInstanceCountsBuffer(bUseRawBuffer);
 		SpawnerKernel->SetProduceOutputPoints(bProduceOutputPoints);
+		SpawnerKernel->SetTargetCullingGridSize(GetCullingCellSize());
 	}
 
 	// If doing by-attribute selection, add analysis kernel that will count how many instances of each mesh are present.
@@ -103,11 +112,17 @@ void UPCGStaticMeshSpawnerSettings::CreateKernels(FPCGGPUCompilationContext& InO
 		// We operate across all input data rather than spawning for each input data separately.
 		CountKernel->SetEmitPerDataCounts(false);
 		CountKernel->SetOutputRawBuffer(bUseRawBuffer);
+		CountKernel->SetBucketExtent(GetCullingCellSize());
 
-		// Rather than add a set of max reduction passes to compute this on the fly, we set the maximum number of values to the known static maximum value.
 		if (SelectorPrimitiveData)
 		{
-			CountKernel->SetMaxIntValue(PCGComputeConstants::MAX_PRIMITIVE_COMPONENTS_PER_SPAWNER - 1);
+			// The max number of unique primitives indices is the num elements in the input primitive data table.
+			CountKernel->SetMaxNumUniqueValuesSource(EPCGMaxNumUniqueValuesSource::UniqueValueTable);
+			OutEdges.Emplace(FPCGPinReference(PCGStaticMeshSpawner::PrimitiveTablePinLabel), FPCGPinReference(CountKernel, PCGCountUniqueAttributeValuesConstants::UniqueValueTablePinLabel));
+		}
+		else
+		{
+			CountKernel->SetMaxNumUniqueValuesSource(EPCGMaxNumUniqueValuesSource::StringKeyValues);
 		}
 
 		OutEdges.Emplace(FPCGPinReference(CountKernel, PCGPinConstants::DefaultOutputLabel), FPCGPinReference(SpawnerKernel, PCGStaticMeshSpawnerKernelConstants::InstanceCountsPinLabel));
@@ -146,6 +161,11 @@ EPCGChangeType UPCGStaticMeshSpawnerSettings::GetChangeTypeForProperty(const FNa
 }
 #endif
 
+int32 UPCGStaticMeshSpawnerSettings::GetCullingCellSize() const
+{
+	return bSetupCullingCells ? PCGStaticMeshSpawner::SanitizeCullGridSize(CullingCellSize) : 0;
+}
+
 TArray<FPCGPinProperties> UPCGStaticMeshSpawnerSettings::InputPinProperties() const
 {
 	// Note: If executing on the GPU, we need to prevent multiple connections on inputs, since it is not supported at this time.
@@ -208,7 +228,7 @@ bool FPCGStaticMeshSpawnerElement::PrepareDataInternal(FPCGContext* InContext) c
 	// Check if we can reuse existing resources on an "all-or-nothing" basis
 	bool& bSkippedDueToReuse = Context->bSkippedDueToReuse;
 
-	if (!Context->bReuseCheckDone && CVarAllowISMReuse.GetValueOnAnyThread())
+	if (!Context->bReuseCheckDone && PCGStaticMeshSpawner::CVarAllowISMReuse.GetValueOnAnyThread())
 	{
 		// Compute CRC if it has not been computed (it likely isn't, but this is to futureproof this)
 		if (!Context->DependenciesCrc.IsValid())
@@ -313,7 +333,7 @@ bool FPCGStaticMeshSpawnerElement::PrepareDataInternal(FPCGContext* InContext) c
 			// Prior to selection, if we have MISMC resources that have the same settings crc & the same data crc,
 			// We can skip the selection here, unless we need to output the data (again, could/should be cached instead).
 			FPCGCrc DataCrc;
-			if(!bSkippedDueToReuse && CVarAllowISMReuse.GetValueOnAnyThread())
+			if(!bSkippedDueToReuse && PCGStaticMeshSpawner::CVarAllowISMReuse.GetValueOnAnyThread())
 			{
 				check(Context->DependenciesCrc.IsValid()); // should have been done earlier
 
@@ -797,6 +817,8 @@ void UPCGStaticMeshSpawnerSettings::PostEditChangeProperty(FPropertyChangedEvent
 		}
 	}
 
+	CullingCellSize = PCGStaticMeshSpawner::SanitizeCullGridSize(CullingCellSize);
+
 	Super::PostEditChangeProperty(PropertyChangedEvent);
 }
 
diff --git a/Engine/Plugins/PCG/Source/PCG/Private/Elements/PCGStaticMeshSpawnerKernel.cpp b/Engine/Plugins/PCG/Source/PCG/Private/Elements/PCGStaticMeshSpawnerKernel.cpp
--- a/Engine/Plugins/PCG/Source/PCG/Private/Elements/PCGStaticMeshSpawnerKernel.cpp
+++ b/Engine/Plugins/PCG/Source/PCG/Private/Elements/PCGStaticMeshSpawnerKernel.cpp
@@ -196,7 +196,8 @@ void UPCGStaticMeshSpawnerKernel::CreateAdditionalOutputDataInterfaces(FPCGGPUCo
 
 	UPCGInstanceDataInterface* InstanceDI = InOutContext.NewObject_AnyThread<UPCGInstanceDataInterface>(InObjectOuter);
 	InstanceDI->SetProducerKernel(this);
-	InstanceDI->InputPinProvidingData = PCGPinConstants::DefaultInputLabel;
+	InstanceDI->SetInputPinProvidingData(PCGPinConstants::DefaultInputLabel);
+	InstanceDI->SetCullingCellExtent(GetTargetCullingGridSize());
 	OutDataInterfaces.Add(InstanceDI);
 }
 
diff --git a/Engine/Plugins/PCG/Source/PCG/Public/Compute/PCGDataBinding.h b/Engine/Plugins/PCG/Source/PCG/Public/Compute/PCGDataBinding.h
--- a/Engine/Plugins/PCG/Source/PCG/Public/Compute/PCGDataBinding.h
+++ b/Engine/Plugins/PCG/Source/PCG/Public/Compute/PCGDataBinding.h
@@ -176,6 +176,9 @@ public:
 
 	const FPCGDataCollection& GetOutputDataCollection() const;
 
+	/** Compute how many culling cells are needed to cover the current generation volume. */
+	FIntVector GetMaxNumCullingCells(int32 InCullingCellExtent) const;
+
 private:
 	/** Loop over all metadata attributes in all input data ensure all attributes are registered in attribute table. */
 	void AddInputDataAttributesToTable();
diff --git a/Engine/Plugins/PCG/Source/PCG/Public/Compute/PrimitiveFactories/IPCGRuntimePrimitiveFactory.h b/Engine/Plugins/PCG/Source/PCG/Public/Compute/PrimitiveFactories/IPCGRuntimePrimitiveFactory.h
--- a/Engine/Plugins/PCG/Source/PCG/Public/Compute/PrimitiveFactories/IPCGRuntimePrimitiveFactory.h
+++ b/Engine/Plugins/PCG/Source/PCG/Public/Compute/PrimitiveFactories/IPCGRuntimePrimitiveFactory.h
@@ -14,6 +14,10 @@ public:
 
 	virtual int32 GetNumPrimitives() const = 0;
 	virtual FPrimitiveSceneProxy* GetSceneProxy(int32 InPrimitiveIndex) const = 0;
-	virtual int32 GetNumInstances(int32 InPrimitiveIndex) const = 0;
+	PCG_API virtual int32 GetNumInstances(int32 InPrimitiveIndex, int32 InCellID) const;
+	virtual int32 GetNumInstancesTotal(int32 InPrimitiveIndex) const = 0;
 	virtual bool IsRenderStateCreated() const = 0;
+
+	UE_DEPRECATED(5.8, "Use version that takes cell ID instead")
+	virtual int32 GetNumInstances(int32 InPrimitiveIndex) const { return 0; }
 };
diff --git a/Engine/Plugins/PCG/Source/PCG/Public/Compute/PrimitiveFactories/PCGPrimitiveFactoryISMBase.h b/Engine/Plugins/PCG/Source/PCG/Public/Compute/PrimitiveFactories/PCGPrimitiveFactoryISMBase.h
--- a/Engine/Plugins/PCG/Source/PCG/Public/Compute/PrimitiveFactories/PCGPrimitiveFactoryISMBase.h
+++ b/Engine/Plugins/PCG/Source/PCG/Public/Compute/PrimitiveFactories/PCGPrimitiveFactoryISMBase.h
@@ -4,25 +4,65 @@
 
 #include "Compute/PrimitiveFactories/IPCGRuntimePrimitiveFactory.h"
 
+#include "Components/PCGProceduralISMComponentDescriptor.h"
 #include "InstanceDataPackers/PCGInstanceDataPackerBase.h"
 
+#include "InstanceDataSceneProxy.h"
+
 class AActor;
 struct FPCGContext;
-struct FPCGProceduralISMComponentDescriptor;
+
+struct FPCGInstanceRange
+{
+	FPCGInstanceRange(uint32 InNumInstances)
+		: NumInstances(InNumInstances)
+	{
+	}
+
+	FPCGInstanceRange(uint32 InNumInstances, int32 InCellID)
+		: NumInstances(InNumInstances)
+		, CellID(InCellID)
+	{
+	}
+
+	uint32 GetNumInstances() const { return NumInstances; }
+	int32 GetCellID() const { return CellID; }
+
+private:
+	uint32 NumInstances = 0;
+	int32 CellID = 0;
+};
+
+struct FPCGPrimitiveInfo
+{
+	FPCGProceduralISMComponentDescriptor Descriptor;
+	TArray<FPCGInstanceRange> InstanceRanges;
+	float CullingCellExtent = 0.0f;
+	FIntVector NumCullingCells = FIntVector::ZeroValue;
+};
 
 class IPCGPrimitiveFactoryISMBase : public IPCGRuntimePrimitiveFactory
 {
 public:
 	struct FParameters
 	{
-		TArray<FPCGProceduralISMComponentDescriptor> Descriptors;
+		TArray<FPCGPrimitiveInfo> PrimitiveInfos;
+
 		FPCGPackedCustomData CustomPrimitiveData;
+		
 		AActor* TargetActor = nullptr; // Optional, unless creating actor primitive components.
+
+		UE_DEPRECATED(5.8, "Use PrimitiveInfos instead")
+		TArray<FPCGProceduralISMComponentDescriptor> Descriptors;
 	};
 
 	virtual void Initialize(FParameters&& InParameters) = 0;
 	virtual bool Create(FPCGContext* InContext) = 0;
 	virtual FBox GetMeshBounds(int32 InPrimitiveIndex) const = 0;
+
+protected:
+	/** Produces spatial hash info which is an implicit cell in space plus instance counts. Also updates primitive bounds to tightly wrap populated cells. */
+	PCG_API virtual void PopulateHashes(const FPCGPrimitiveInfo& InPrimitiveInfo, TArray<FInstanceSceneDataBuffers::FCompressedSpatialHashItem>& OutHashes, FBox& OutBounds) const;
 };
 
 namespace PCGPrimitiveFactoryHelpers
diff --git a/Engine/Plugins/PCG/Source/PCG/Public/Elements/PCGStaticMeshSpawner.h b/Engine/Plugins/PCG/Source/PCG/Public/Elements/PCGStaticMeshSpawner.h
--- a/Engine/Plugins/PCG/Source/PCG/Public/Elements/PCGStaticMeshSpawner.h
+++ b/Engine/Plugins/PCG/Source/PCG/Public/Elements/PCGStaticMeshSpawner.h
@@ -42,6 +42,8 @@ public:
 #endif // WITH_EDITOR
 	virtual bool UseSeed() const override { return true; }
 
+	int32 GetCullingCellSize() const;
+
 protected:
 	UE_API virtual TArray<FPCGPinProperties> InputPinProperties() const override;
 	virtual TArray<FPCGPinProperties> OutputPinProperties() const override { return Super::DefaultPointOutputPinProperties(); }
@@ -126,6 +128,13 @@ public:
 	bool bWarnOnIdenticalSpawn = true;
 
 protected:
+	/** Experimental: Divide instances into cells to allow renderer to cull subranges of instances early on the CPU. */
+	UPROPERTY(EditAnywhere, Category = GPU, AdvancedDisplay, meta = (EditCondition = "bExecuteOnGPU", EditConditionHides))
+	bool bSetupCullingCells = false;
+
+	UPROPERTY(EditAnywhere, Category = GPU, AdvancedDisplay, meta = (EditCondition = "bExecuteOnGPU && bSetupCullingCells", EditConditionHides, ClampMin = "64", UIMin = "64", UIMax = "65536"))
+	int32 CullingCellSize = 2048;
+
 	UE_API void RefreshMeshSelector();
 	UE_API void RefreshInstancePacker();
 };
diff --git a/Engine/Plugins/PCG/Source/PCG/Public/Elements/PCGStaticMeshSpawnerKernel.h b/Engine/Plugins/PCG/Source/PCG/Public/Elements/PCGStaticMeshSpawnerKernel.h
--- a/Engine/Plugins/PCG/Source/PCG/Public/Elements/PCGStaticMeshSpawnerKernel.h
+++ b/Engine/Plugins/PCG/Source/PCG/Public/Elements/PCGStaticMeshSpawnerKernel.h
@@ -41,6 +41,9 @@ public:
 	bool GetProduceOutputPoints() const { return bProduceOutputPoints; }
 	void SetProduceOutputPoints(bool bInProduceOutputPoints) { bProduceOutputPoints = bInProduceOutputPoints; }
 
+	int32 GetTargetCullingGridSize() const { return TargetCullingGridSize; }
+	void SetTargetCullingGridSize(int32 InTargetCullingGridSize) { TargetCullingGridSize = InTargetCullingGridSize; }
+
 protected:
 #if WITH_EDITOR
 	virtual bool PerformStaticValidation() override;
@@ -53,4 +56,8 @@ protected:
 
 	UPROPERTY()
 	bool bProduceOutputPoints = true;
+
+	/** Optionally divide primitive instances into cells and pass to renderer to aid culling. */
+	UPROPERTY()
+	int32 TargetCullingGridSize = 0;
 };
diff --git a/Engine/Plugins/PCG/Source/PCGCompute/Internal/PCGSceneWriterCS.h b/Engine/Plugins/PCG/Source/PCGCompute/Internal/PCGSceneWriterCS.h
--- a/Engine/Plugins/PCG/Source/PCGCompute/Internal/PCGSceneWriterCS.h
+++ b/Engine/Plugins/PCG/Source/PCGCompute/Internal/PCGSceneWriterCS.h
@@ -26,6 +26,7 @@ public:
 		SHADER_PARAMETER(uint32, InCustomDataCount)
 		SHADER_PARAMETER(uint32, InPayloadDataFlags)
 		SHADER_PARAMETER(uint32, InSeed)
+		SHADER_PARAMETER(uint32, InNumCullingCellsTotal)
 		SHADER_PARAMETER_STRUCT_INCLUDE(FGPUSceneWriterParameters, GPUSceneWriterParameters)
 	END_SHADER_PARAMETER_STRUCT()
 
