[PCG] ActorTracker: handle actor deleted through undo/redo

#Synced-CL 50397505

diff --git a/Engine/Plugins/PCG/Source/PCGEditor/Internal/ChangeTracking/PCGActorTracker.h b/Engine/Plugins/PCG/Source/PCGEditor/Internal/ChangeTracking/PCGActorTracker.h
--- a/Engine/Plugins/PCG/Source/PCGEditor/Internal/ChangeTracking/PCGActorTracker.h
+++ b/Engine/Plugins/PCG/Source/PCGEditor/Internal/ChangeTracking/PCGActorTracker.h
@@ -12,6 +12,7 @@
 class FPCGTrackingManager;
 
 class AActor;
+class FTransactionObjectEvent;
 class ILevelInstanceInterface;
 class UObject;
 
@@ -56,6 +57,7 @@ private:
 	void OnActorDeleted(AActor* InActor);
 	void OnActorDeleted_Internal(const AActor* InActor, bool bShouldDirty);
 	void OnObjectModified(UObject* InObject);
+	void OnObjectTransacted(UObject* Object, const FTransactionObjectEvent& Event);
 	
 	static const FLazyName Name;
 
diff --git a/Engine/Plugins/PCG/Source/PCGEditor/Private/ChangeTracking/PCGActorTracker.cpp b/Engine/Plugins/PCG/Source/PCGEditor/Private/ChangeTracking/PCGActorTracker.cpp
--- a/Engine/Plugins/PCG/Source/PCGEditor/Private/ChangeTracking/PCGActorTracker.cpp
+++ b/Engine/Plugins/PCG/Source/PCGEditor/Private/ChangeTracking/PCGActorTracker.cpp
@@ -15,6 +15,7 @@
 #include "LevelInstance/LevelInstanceEditorPivotInterface.h"
 #include "LevelInstance/LevelInstanceInterface.h"
 #include "LevelInstance/LevelInstanceSubsystem.h"
+#include "Misc/TransactionObjectEvent.h"
 #include "PackedLevelActor/PackedLevelActor.h"
 #include "UObject/UObjectGlobals.h"
 #include "WorldPartition/WorldPartition.h"
@@ -78,6 +79,7 @@ FPCGActorTracker::FPCGActorTracker(FPCGTrackingManager* InOwner)
 	GEngine->OnLevelActorAdded().AddRaw(this, &FPCGActorTracker::OnActorAdded);
 	GEngine->OnLevelActorDeleted().AddRaw(this, &FPCGActorTracker::OnActorDeleted);
 	FCoreUObjectDelegates::OnObjectModified.AddRaw(this, &FPCGActorTracker::OnObjectModified);
+	FCoreUObjectDelegates::OnObjectTransacted.AddRaw(this, &FPCGActorTracker::OnObjectTransacted);
 	
 	UWorld* World = Owner->GetWorld();
 	check(World);
@@ -100,6 +102,7 @@ FPCGActorTracker::~FPCGActorTracker()
 	GEngine->OnLevelActorAdded().RemoveAll(this);
 	GEngine->OnLevelActorDeleted().RemoveAll(this);
 	FCoreUObjectDelegates::OnObjectModified.RemoveAll(this);
+	FCoreUObjectDelegates::OnObjectTransacted.RemoveAll(this);
 
 	UWorld* World = Owner->GetWorld();
 	check(World);
@@ -416,6 +419,26 @@ void FPCGActorTracker::OnActorDeleted_Internal(const AActor* InActor, bool bShou
 	OnActorChanged_Impl(InActor, nullptr, PCGActorTracker::ActorDeleted, /*InPreviousData=*/nullptr);
 }
 
+void FPCGActorTracker::OnObjectTransacted(UObject* Object, const FTransactionObjectEvent& Event)
+{
+	// We only care about deleted actors through undo redo
+	const bool bProcessChange = Event.GetEventType() == ETransactionObjectEventType::UndoRedo && Event.HasPendingKillChange() && !IsValid(Object);
+	if (!bProcessChange)
+	{
+		return;
+	}
+
+	const AActor* Actor = Cast<AActor>(Object);
+
+	// If this isn't an actor early out / or not owned by the world we are tracking
+	if (!Actor || (Actor->GetWorld() != Owner->GetWorld()))
+	{
+		return;
+	}
+
+	OnActorDeleted_Internal(Actor, /*bShouldDirty=*/ true);
+}
+
 void FPCGActorTracker::OnObjectModified(UObject* InObject)
 {
 	TRACE_CPUPROFILER_EVENT_SCOPE(FPCGActorTracker::OnObjectModified);
