[PCG] Blueprint: fix desync between TrackingManager and BP PCG Component tracking keys

#Synced-CL 50397505

diff --git a/Engine/Plugins/PCG/Source/PCG/Internal/PCGTrackingManager.h b/Engine/Plugins/PCG/Source/PCG/Internal/PCGTrackingManager.h
--- a/Engine/Plugins/PCG/Source/PCG/Internal/PCGTrackingManager.h
+++ b/Engine/Plugins/PCG/Source/PCG/Internal/PCGTrackingManager.h
@@ -194,10 +194,10 @@ private:
 	void RemapTracking(const IPCGGraphExecutionSource* InOldExecutionSource, IPCGGraphExecutionSource* InNewExecutionSource);
 
 	/** Unregister tracking when a execution source is removed */
-	void UnregisterTracking(IPCGGraphExecutionSource* InExecutionSource);
+	bool UnregisterTracking(const IPCGGraphExecutionSource* InExecutionSource);
 
 	/** Unregister tracking when a execution source is removed or keys are unregistered */
-	void UnregisterTracking(IPCGGraphExecutionSource* InExecutionSource, const TSet<FPCGSelectionKey>* OptionalKeysToUntrack);
+	bool UnregisterTracking(const IPCGGraphExecutionSource* InExecutionSource, const TSet<FPCGSelectionKey>* OptionalKeysToUntrack);
 
 	void OnObjectChanged(const UObject* InObject, const FPCGTrackingChangeID& InChangeID, const UObject* InOriginatingChangeObject = nullptr);
 
@@ -233,8 +233,8 @@ private:
 
 	/** Components to be registered/unregistered on the next frame. see RegisterOrUpdatePCGComponent/UnregisterPCGComponent for a better understanding on why it is needed. */
 	mutable PCG::FLock DelayedComponentsRegistrationLock;
-	TSet<UPCGComponent*> DelayedComponentsToUnregister;
-	TSet<UPCGComponent*> DelayedComponentsToRegister;
+	TSet<TWeakObjectPtr<UPCGComponent>> DelayedComponentsToUnregister;
+	TSet<TWeakObjectPtr<UPCGComponent>> DelayedComponentsToRegister;
 	
 	// Tracking actors
 	/** Keep a mapping between tracked keys and the execution sources that track them, and the tracking needs to be culled.*/
diff --git a/Engine/Plugins/PCG/Source/PCG/Private/PCGComponent.cpp b/Engine/Plugins/PCG/Source/PCG/Private/PCGComponent.cpp
--- a/Engine/Plugins/PCG/Source/PCG/Private/PCGComponent.cpp
+++ b/Engine/Plugins/PCG/Source/PCG/Private/PCGComponent.cpp
@@ -3781,14 +3781,19 @@ FPCGComponentInstanceData::FPCGComponentInstanceData(const UPCGComponent* InSour
 #endif
 {
 #if WITH_EDITOR
-	// Convert map to array for serialization
-	for (const auto& KeyValuePair : InSourceComponent->ChangeTracker.DynamicallyTrackedKeysToSettings)
+	auto CopyFromTrackingKeys = [](const FPCGSelectionKeyToSettingsMap& InKeysToSettingsMap, TArray<FPCGTrackedKeyInstanceData>& OutTrackingKeys)
 	{
-		for (const auto& SettingsAndCullingPair : KeyValuePair.Value)
+		for (const auto& KeyValuePair : InKeysToSettingsMap)
 		{
-			DynamicallyTrackedKeysToSettings.Add({ KeyValuePair.Key, SettingsAndCullingPair.Key, SettingsAndCullingPair.Value });
+			for (const auto& SettingsAndCullingPair : KeyValuePair.Value)
+			{
+				OutTrackingKeys.Add(FPCGTrackedKeyInstanceData { KeyValuePair.Key, SettingsAndCullingPair.Key, SettingsAndCullingPair.Value });
+			}
 		}
-	}
+	};
+	
+	CopyFromTrackingKeys(InSourceComponent->ChangeTracker.DynamicallyTrackedKeysToSettings, DynamicallyTrackedKeysToSettings);
+	CopyFromTrackingKeys(InSourceComponent->ChangeTracker.StaticallyTrackedKeysToSettings, StaticallyTrackedKeysToSettings);
 #endif
 }
 
@@ -3835,10 +3840,18 @@ void FPCGComponentInstanceData::ApplyToComponent(UActorComponent* Component, con
 		PCGComponent->bGeneratedOffline = bGeneratedOffline;
 		PCGComponent->ManagedResourceContainer.SetOwner(PCGComponent);
 
-		for (const FDynamicTrackedKeyInstanceData& DynamicTrackedKeyInstanceData : DynamicallyTrackedKeysToSettings)
+		auto CopyToTrackingKeys = [](const TArray<FPCGTrackedKeyInstanceData>& InTrackingKeys, FPCGSelectionKeyToSettingsMap& OutKeysToSettingsMap)
 		{
-			PCGComponent->ChangeTracker.DynamicallyTrackedKeysToSettings.FindOrAdd(DynamicTrackedKeyInstanceData.SelectionKey).Add({ DynamicTrackedKeyInstanceData.Settings, DynamicTrackedKeyInstanceData.bValue });
-		}
+			OutKeysToSettingsMap.Empty();
+			for(const FPCGTrackedKeyInstanceData& InstanceData : InTrackingKeys)
+			{
+				OutKeysToSettingsMap.FindOrAdd(InstanceData.SelectionKey).Add({ InstanceData.Settings , InstanceData.bValue });
+			}
+		};
+
+		CopyToTrackingKeys(StaticallyTrackedKeysToSettings, PCGComponent->ChangeTracker.StaticallyTrackedKeysToSettings);
+		CopyToTrackingKeys(DynamicallyTrackedKeysToSettings, PCGComponent->ChangeTracker.DynamicallyTrackedKeysToSettings);
+
 #endif
 
 		if (ConstructionSourceComponent)
diff --git a/Engine/Plugins/PCG/Source/PCG/Private/PCGTrackingManager.cpp b/Engine/Plugins/PCG/Source/PCG/Private/PCGTrackingManager.cpp
--- a/Engine/Plugins/PCG/Source/PCG/Private/PCGTrackingManager.cpp
+++ b/Engine/Plugins/PCG/Source/PCG/Private/PCGTrackingManager.cpp
@@ -127,16 +127,29 @@ void FPCGTrackingManager::Tick()
 #if WITH_EDITOR
 
 	// Process delayed Registrations
-	TSet<UPCGComponent*> ComponentsToUnregister;
-	TSet<UPCGComponent*> ComponentsToRegister;
+	TSet<TWeakObjectPtr<UPCGComponent>> ComponentsToUnregister;
+	TSet<TWeakObjectPtr<UPCGComponent>> ComponentsToRegister;
 	{
 		PCG::TScopeLock Lock(DelayedComponentsRegistrationLock);
 		ComponentsToUnregister = MoveTemp(DelayedComponentsToUnregister);
 		ComponentsToRegister = MoveTemp(DelayedComponentsToRegister);
 	}
 
-	Algo::ForEach(ComponentsToRegister, [this](UPCGComponent* Component) { RegisterOrUpdateExecutionSource(Component, Component->bGenerated, /*bForce=*/true); });
-	Algo::ForEach(ComponentsToUnregister, [this](UPCGComponent* Component) { UnregisterExecutionSource(Component, /*bForce=*/true); });
+	Algo::ForEach(ComponentsToRegister, [this](TWeakObjectPtr<UPCGComponent> WeakComponent) 
+	{
+		if(UPCGComponent* Component = WeakComponent.Get())
+		{
+			RegisterOrUpdateExecutionSource(Component, Component->bGenerated, /*bForce=*/true); 
+		}
+	});
+	
+	Algo::ForEach(ComponentsToUnregister, [this](TWeakObjectPtr<UPCGComponent> WeakComponent)
+	{ 
+		if(UPCGComponent* Component = WeakComponent.Get())
+		{
+			UnregisterExecutionSource(Component, /*bForce=*/true); 
+		}
+	});
 
 	for (TUniquePtr<IPCGChangeTracker>& ChangeTracker : ChangeTrackers)
 	{
@@ -398,11 +411,13 @@ bool FPCGTrackingManager::RemapExecutionSource(const IPCGGraphExecutionSource* I
 	if (OldComponent && NewComponent)
 	{
 #if WITH_EDITOR
-		// Remove it from the delayed
+		// This is called from within construction script to notify us that an execution source has been replaced with a new version.
+		// - We expect the old execution source might be in the delayed unregisters and we remove it
+		// - No need to touch delayed registers
 		{
 			PCG::TScopeLock Lock(DelayedComponentsRegistrationLock);
 			DelayedComponentsToUnregister.Remove(OldComponent);
-			DelayedComponentsToRegister.Remove(NewComponent);
+			ensure(DelayedComponentsToRegister.Contains(NewComponent));
 		}
 #endif
 
@@ -1065,20 +1080,6 @@ void FPCGTrackingManager::UpdateTracking(IPCGGraphExecutionSource* InExecutionSo
 
 void FPCGTrackingManager::RemapTracking(const IPCGGraphExecutionSource* InOldExecutionSource, IPCGGraphExecutionSource* InNewExecutionSource)
 {
-	auto ReplaceInMap = [InOldExecutionSource, InNewExecutionSource](TMap<FPCGSelectionKey, TSet<IPCGGraphExecutionSource*>>& InMap)
-	{
-		for (TPair<FPCGSelectionKey, TSet<IPCGGraphExecutionSource*>>& It : InMap)
-		{
-			if (It.Value.Remove(InOldExecutionSource) > 0)
-			{
-				It.Value.Add(InNewExecutionSource);
-			}
-		}
-	};
-
-	ReplaceInMap(CulledTrackedKeysToExecutionSourcesMap);
-	ReplaceInMap(AlwaysTrackedKeysToExecutionSourcesMap);
-
 	const FObjectKey OldKey = Cast<UObject>(InOldExecutionSource);
 	const FObjectKey NewKey = Cast<UObject>(InNewExecutionSource);
 
@@ -1095,38 +1096,53 @@ void FPCGTrackingManager::RemapTracking(const IPCGGraphExecutionSource* InOldExe
 	{
 		It.Value.RemoveSwap(OldKey);
 	}
+
+	// Unregister tracking and register it so that we don't leak keys that might have changed
+	const bool bFoundExecutionSource = UnregisterTracking(InOldExecutionSource);
+	if (bFoundExecutionSource)
+	{
+		RegisterTracking(InNewExecutionSource);
+	}
 }
 
-void FPCGTrackingManager::UnregisterTracking(IPCGGraphExecutionSource* InExecutionSource, const TSet<FPCGSelectionKey>* OptionalKeysToUntrack)
+bool FPCGTrackingManager::UnregisterTracking(const IPCGGraphExecutionSource* InExecutionSource, const TSet<FPCGSelectionKey>* OptionalKeysToUntrack)
 {
 	if (!InExecutionSource)
 	{
-		return;
+		return false;
 	}
 
+	bool bExecutionSourceFound = false;
+
 	TSet<FPCGSelectionKey> KeysToRemove;
-	auto RemoveAllFromMap = [InExecutionSource, &KeysToRemove](TMap<FPCGSelectionKey, TSet<IPCGGraphExecutionSource*>>& InMap)
+	auto RemoveAllFromMap = [InExecutionSource, &KeysToRemove, &bExecutionSourceFound](TMap<FPCGSelectionKey, TSet<IPCGGraphExecutionSource*>>& InMap)
 	{
 		for (auto& It : InMap)
 		{
-			It.Value.Remove(InExecutionSource);
-			if (It.Value.IsEmpty())
+			if (It.Value.Remove(InExecutionSource) > 0)
 			{
-				KeysToRemove.Add(It.Key);
+				bExecutionSourceFound = true;
+				if (It.Value.IsEmpty())
+				{
+					KeysToRemove.Add(It.Key);
+				}
 			}
 		}
 	};
 
-	auto RemoveKeysFromMap = [InExecutionSource, &KeysToRemove](TMap<FPCGSelectionKey, TSet<IPCGGraphExecutionSource*>>& InMap, const TSet<FPCGSelectionKey>& SetToIterateOn)
+	auto RemoveKeysFromMap = [InExecutionSource, &KeysToRemove, &bExecutionSourceFound](TMap<FPCGSelectionKey, TSet<IPCGGraphExecutionSource*>>& InMap, const TSet<FPCGSelectionKey>& SetToIterateOn)
 	{
 		for (const FPCGSelectionKey& KeyIt : SetToIterateOn)
 		{
 			if (TSet<IPCGGraphExecutionSource*>* ExecutionSourceSetPtr = InMap.Find(KeyIt))
 			{
-				ExecutionSourceSetPtr->Remove(InExecutionSource);
-				if (ExecutionSourceSetPtr->IsEmpty())
+				if (ExecutionSourceSetPtr->Remove(InExecutionSource) > 0)
 				{
-					KeysToRemove.Add(KeyIt);
+					bExecutionSourceFound = true;
+					if (ExecutionSourceSetPtr->IsEmpty())
+					{
+						KeysToRemove.Add(KeyIt);
+					}
 				}
 			}
 		}
@@ -1148,16 +1164,18 @@ void FPCGTrackingManager::UnregisterTracking(IPCGGraphExecutionSource* InExecuti
 		CulledTrackedKeysToExecutionSourcesMap.Remove(Key);
 		AlwaysTrackedKeysToExecutionSourcesMap.Remove(Key);
 	}
+
+	return bExecutionSourceFound;
 }
 
-void FPCGTrackingManager::UnregisterTracking(IPCGGraphExecutionSource* InExecutionSource)
+bool FPCGTrackingManager::UnregisterTracking(const IPCGGraphExecutionSource* InExecutionSource)
 {
 	if (!InExecutionSource)
 	{
-		return;
+		return false;
 	}
 
-	UnregisterTracking(InExecutionSource, nullptr);
+	return UnregisterTracking(InExecutionSource, nullptr);
 }
 
 bool FPCGTrackingManager::IsKeyTracked(const FPCGSelectionKey& InKey) const
diff --git a/Engine/Plugins/PCG/Source/PCG/Public/PCGComponent.h b/Engine/Plugins/PCG/Source/PCG/Public/PCGComponent.h
--- a/Engine/Plugins/PCG/Source/PCG/Public/PCGComponent.h
+++ b/Engine/Plugins/PCG/Source/PCG/Public/PCGComponent.h
@@ -846,7 +846,22 @@ private:
 };
 
 USTRUCT()
-struct FDynamicTrackedKeyInstanceData
+struct UE_DEPRECATED(5.8, "Use FPCGTrackedKeyInstanceData instead") FDynamicTrackedKeyInstanceData
+{
+	GENERATED_BODY();
+
+	UPROPERTY()
+	FPCGSelectionKey SelectionKey;
+
+	UPROPERTY()
+	TSoftObjectPtr<const UPCGSettings> Settings;
+
+	UPROPERTY()
+	bool bValue = false;
+};
+
+USTRUCT()
+struct FPCGTrackedKeyInstanceData
 {
 	GENERATED_BODY();
 
@@ -898,10 +913,13 @@ protected:
 	UPROPERTY()
 	bool bGeneratedOffline = false;
 
-	// UPCGComponent::DynamicallyTrackedKeysToSettings is not a property so we need to capture/apply it ourselves.
+	// UPCGComponent::DynamicallyTrackedKeysToSettings/StaticallyTrackedKeysToSettings are not properties so we need to capture/apply them ourselves.
 	// To avoid having to write custom serialization code the map is converted into an array for serialization.
 	UPROPERTY()
-	TArray<FDynamicTrackedKeyInstanceData> DynamicallyTrackedKeysToSettings;
+	TArray<FPCGTrackedKeyInstanceData> DynamicallyTrackedKeysToSettings;
+
+	UPROPERTY()
+	TArray<FPCGTrackedKeyInstanceData> StaticallyTrackedKeysToSettings;
 #endif
 };
 
