[PCG] Fixed crash in Determinism Test Framework when running a test on UPCGUserParametersData

#Synced-CL 50160642

diff --git a/Engine/Plugins/PCG/Source/PCG/Private/Tests/Determinism/PCGDeterminismTestsCommon.cpp b/Engine/Plugins/PCG/Source/PCG/Private/Tests/Determinism/PCGDeterminismTestsCommon.cpp
--- a/Engine/Plugins/PCG/Source/PCG/Private/Tests/Determinism/PCGDeterminismTestsCommon.cpp
+++ b/Engine/Plugins/PCG/Source/PCG/Private/Tests/Determinism/PCGDeterminismTestsCommon.cpp
@@ -1433,19 +1433,14 @@ namespace PCGDeterminismTests
 
 	bool DataTypeIsComparable(FPCGDataTypeIdentifier DataType)
 	{
-		// Data types that don't need to be compared
-		if (DataType == EPCGDataType::None || DataType == EPCGDataType::Other || DataType == EPCGDataType::Settings)
-		{
-			return false;
-		}
-
-		// Comparable data types
-		if (!!(DataType & (EPCGDataType::Param | EPCGDataType::Spatial)))
+		// Comparable data types (Param and Spatial)
+		const FPCGDataTypeIdentifier ValidTypes = FPCGDataTypeIdentifier::Construct(FPCGDataTypeInfoParam::AsId(), FPCGDataTypeInfoSpatial::AsId());
+		if (DataType.Intersects(ValidTypes))
 		{
 			return true;
 		}
 
-		UE_LOG(LogPCG, Warning, TEXT("Unknown data comparison type: %s"), *DataType.ToString());
+		UE_LOG(LogPCG, Warning, TEXT("Unknown determinism data comparison type: %s"), *DataType.ToString());
 		return false;
 	}
 
@@ -1579,29 +1574,35 @@ namespace PCGDeterminismTests
 			return ComparisonIsUnimplemented;
 		}
 
-		if (DataType == EPCGDataType::Param)
+		if (DataType.IsSameType(FPCGDataTypeInfoParam::AsId()))
 		{
 			return ParamDataIsIdentical;
 		}
 
-		switch (DeterminismLevel)
+		if (DataType.IsChildOf(FPCGDataTypeInfoSpatial::AsId()))
 		{
-		case EDeterminismLevel::OrderOrthogonal:
-			return SpatialDataIsOrthogonal;
-		case EDeterminismLevel::OrderConsistent:
-			return SpatialDataIsConsistent;
-		case EDeterminismLevel::OrderIndependent:
-			return SpatialDataIsIdentical;
-		default:
-			UE_LOG(LogPCG, Warning, TEXT("Comparable PCG DataType has no 'data comparison function'."));
-			return ComparisonIsUnimplemented;
+			switch (DeterminismLevel)
+			{
+				case EDeterminismLevel::OrderOrthogonal:
+					return SpatialDataIsOrthogonal;
+				case EDeterminismLevel::OrderConsistent:
+					return SpatialDataIsConsistent;
+				case EDeterminismLevel::OrderIndependent:
+					return SpatialDataIsIdentical;
+				default:
+					UE_LOG(LogPCG, Warning, TEXT("Comparable PCG DataType has no 'data comparison function'."));
+					return ComparisonIsUnimplemented;
+			}
 		}
+
+		UE_LOG(LogPCG, Warning, TEXT("Can not compare non-spatial or non-param data: %s"), *DataType.ToString());
+		return ComparisonIsUnimplemented;
 	}
 
 	TFunction<bool(const UPCGData*, const UPCGData*)> GetDataCollectionCompareFunction(FPCGDataTypeIdentifier DataType)
 	{
 		// TODO: Currently only spatial data is comparable. Other types to be added.
-		if (!DataTypeIsComparable(DataType) || !(DataType & EPCGDataType::Spatial))
+		if (!DataTypeIsComparable(DataType) || !DataType.IsChildOf(FPCGDataTypeInfoSpatial::AsId()))
 		{
 			// Should never reach here
 			UE_LOG(LogPCG, Warning, TEXT("Attempting to compare incomparable data."));
diff --git a/Engine/Plugins/PCG/Source/PCG/Public/Data/PCGUserParametersData.h b/Engine/Plugins/PCG/Source/PCG/Public/Data/PCGUserParametersData.h
--- a/Engine/Plugins/PCG/Source/PCG/Public/Data/PCGUserParametersData.h
+++ b/Engine/Plugins/PCG/Source/PCG/Public/Data/PCGUserParametersData.h
@@ -23,8 +23,9 @@ class UPCGUserParametersData : public UPCGData
 	GENERATED_BODY()
 
 public:
+	PCG_ASSIGN_DEFAULT_TYPE_INFO(FPCGDataTypeInfoOther)
+
 	// ~Begin UPCGData interface
-	virtual EPCGDataType GetDataType() const override { return EPCGDataType::Other; }
 	virtual void AddToCrc(FArchiveCrc32& Ar, bool bFullDataCrc) const override { AddUIDToCrc(Ar); }
 	// ~End UPCGData interface
 
