[PCG] Fix crash in editor worlds when a non-runtime-gen PCG component has generated FastGeo primitives.

Since we manage the FastGeo using a managed resource, we must look for FastGeo resources and release them on UnRegister.

#rnx

#Synced-CL 50393138

diff --git a/Engine/Plugins/PCG/Source/PCG/Private/PCGComponent.cpp b/Engine/Plugins/PCG/Source/PCG/Private/PCGComponent.cpp
--- a/Engine/Plugins/PCG/Source/PCG/Private/PCGComponent.cpp
+++ b/Engine/Plugins/PCG/Source/PCG/Private/PCGComponent.cpp
@@ -1727,39 +1727,43 @@ void UPCGComponent::OnUnregister()
 	// We shouldn't cleanup resources in OnUnregister in most persistent cases.
 	// This specific case is to handle World Partition unloading of actors where our actor might have a PCG Component which generated some Preview Resources.
 	// In this case we do want to release those resources and the actors they might hold on to. If we don't those actors become orphaned and we might end up with duplicates when reloading the component.
-	if (bUnregisteredThroughLoading && PCGComponent::CVarDeletePreviewResourcesWhenUnloading.GetValueOnAnyThread() && CurrentEditingMode == EPCGEditorDirtyMode::Preview)
+	TSet<TSoftObjectPtr<AActor>> ActorsToDelete;
 	{
-		TSet<TSoftObjectPtr<AActor>> ActorsToDelete;
+		UE::TScopeLock ResourcesLock(GeneratedResourcesLock);
+		if (ensure(ManagedResourceContainer.AreResourcesAccessible()))
 		{
-			UE::TScopeLock ResourcesLock(GeneratedResourcesLock);
-			if (ensure(ManagedResourceContainer.AreResourcesAccessible()))
+			const bool bCleanupPersistentResources = bUnregisteredThroughLoading && PCGComponent::CVarDeletePreviewResourcesWhenUnloading.GetValueOnAnyThread() && CurrentEditingMode == EPCGEditorDirtyMode::Preview;
+
+			for (UPCGManagedResource* ManagedResource : ManagedResourceContainer.GetConstManagedResourcesNoLock())
 			{
-				for (UPCGManagedResource* ManagedResource : ManagedResourceContainer.GetConstManagedResourcesNoLock())
+				if (ManagedResource)
 				{
 					// Make sure resource is indeed a transient/preview resource
-					if (ManagedResource &&
+					const bool bReleasePersistent = bCleanupPersistentResources &&
 						!ManagedResource->IsMarkedTransientOnLoad() &&
-						!ManagedResourceContainer.GetConstLoadedPreviewResourcesNoLock().Contains(ManagedResource))
+						!ManagedResourceContainer.GetConstLoadedPreviewResourcesNoLock().Contains(ManagedResource);
+
+					if (bReleasePersistent || ManagedResource->ReleaseOnTeardown())
 					{
 						ManagedResource->Release(/*bHardRelease=*/true, ActorsToDelete);
 					}
 				}
 			}
 		}
+	}
 
-		if (!ActorsToDelete.IsEmpty())
+	if (!ActorsToDelete.IsEmpty())
+	{
+		// Dispatch delete of actors to next tick
+		::ExecuteOnGameThread(UE_SOURCE_LOCATION, [WeakWorld = TWeakObjectPtr<UWorld>(GetWorld()), InActorsToDelete = MoveTemp(ActorsToDelete)]()
 		{
-			// Dispatch delete of actors to next tick
-			::ExecuteOnGameThread(UE_SOURCE_LOCATION, [WeakWorld = TWeakObjectPtr<UWorld>(GetWorld()), InActorsToDelete = MoveTemp(ActorsToDelete)]()
-			{
-				LLM_SCOPE_BYTAG(PCG);
+			LLM_SCOPE_BYTAG(PCG);
 
-				if (UWorld* World = WeakWorld.Get())
-				{
-					UPCGActorHelpers::DeleteActors(World, InActorsToDelete.Array());
-				}
-			});
-		}
+			if (UWorld* World = WeakWorld.Get())
+			{
+				UPCGActorHelpers::DeleteActors(World, InActorsToDelete.Array());
+			}
+		});
 	}
 #endif // WITH_EDITOR
 
