[PCG] Fixed Data View Json conversion bug, which was discounting default domains

#rnx
#Synced-CL 50160642

diff --git a/Engine/Plugins/PCG/Source/PCG/Private/Data/DataView/PCGDataViewNativeJsonConverters.cpp b/Engine/Plugins/PCG/Source/PCG/Private/Data/DataView/PCGDataViewNativeJsonConverters.cpp
--- a/Engine/Plugins/PCG/Source/PCG/Private/Data/DataView/PCGDataViewNativeJsonConverters.cpp
+++ b/Engine/Plugins/PCG/Source/PCG/Private/Data/DataView/PCGDataViewNativeJsonConverters.cpp
@@ -37,12 +37,14 @@ TValueOrError<FInstancedStruct, FText> UPCGDataViewJsonConverter::SerializeToTar
 
 		const FPCGDataViewJsonParameters& JsonParameters = Parameters.Get<FPCGDataViewJsonParameters>();
 
-		auto GetDomainAttributes = [&Selection, ViewedData, Layout = JsonParameters.AttributeLayout, &RootJsonObject](FPCGMetadataDomainID DomainID)
+		const FPCGMetadataDomainID DefaultDomainID = ViewedData->GetDefaultMetadataDomainID();
+		auto GetDomainAttributes = [&Selection, ViewedData, DefaultDomainID, Layout = JsonParameters.AttributeLayout, &RootJsonObject](FPCGMetadataDomainID DomainID)
 		{
-			const TArray<FPCGAttributePropertySelector> DomainAttributes = Selection.FilterByPredicate([ViewedData, &DomainID, &RootJsonObject](const FPCGAttributePropertySelector& Selector)
+			const bool bIsDefaultDomain = (DefaultDomainID == DomainID);
+			const TArray<FPCGAttributePropertySelector> DomainAttributes = Selection.FilterByPredicate([ViewedData, &DomainID, bIsDefaultDomain, &RootJsonObject](const FPCGAttributePropertySelector& Selector)
 			{
 				const FPCGMetadataDomainID SelectorDomainID = ViewedData->GetMetadataDomainIDFromSelector(Selector);
-				return (SelectorDomainID == DomainID);
+				return (SelectorDomainID == DomainID) || (bIsDefaultDomain && SelectorDomainID.IsDefault());
 			});
 
 			if (DomainAttributes.IsEmpty())
@@ -60,7 +62,7 @@ TValueOrError<FInstancedStruct, FText> UPCGDataViewJsonConverter::SerializeToTar
 			}
 		};
 
-		for (const FPCGMetadataDomainID DomainID : InDataView.GetAllMetadataDomainIDs())
+		for (const FPCGMetadataDomainID DomainID : ViewedData->GetAllSupportedMetadataDomainIDs())
 		{
 			GetDomainAttributes(DomainID);
 		}
