[PCG] Fixed bug where Convert to Data View was improperly restricting pins to spatial data

#rnx

#Synced-CL 50160642

diff --git a/Engine/Plugins/PCG/Source/PCG/Private/Elements/PCGConvertToDataView.cpp b/Engine/Plugins/PCG/Source/PCG/Private/Elements/PCGConvertToDataView.cpp
--- a/Engine/Plugins/PCG/Source/PCG/Private/Elements/PCGConvertToDataView.cpp
+++ b/Engine/Plugins/PCG/Source/PCG/Private/Elements/PCGConvertToDataView.cpp
@@ -33,7 +33,12 @@ bool UPCGConvertToDataViewSettings::ShouldDrawNodeCompact() const
 TArray<FPCGPinProperties> UPCGConvertToDataViewSettings::InputPinProperties() const
 {
 	TArray<FPCGPinProperties> PinProperties;
-	PinProperties.Emplace_GetRef(PCGPinConstants::DefaultInputLabel, FPCGDataTypeInfoSpatial::AsId()).SetRequiredPin();
+	PinProperties.Emplace_GetRef(
+		PCGPinConstants::DefaultInputLabel,
+		FPCGDataTypeIdentifier::Construct(
+			FPCGDataTypeInfoSpatial::AsId(),
+			FPCGDataTypeInfoParam::AsId()
+			)).SetRequiredPin();
 
 	return PinProperties;
 }
