[PCG] Fix broken enum multi select test.

#rnx

#Synced-CL 50412500

diff --git a/Engine/Plugins/PCG/Source/PCG/Private/Tests/Elements/PCGMultiSelectTest.cpp b/Engine/Plugins/PCG/Source/PCG/Private/Tests/Elements/PCGMultiSelectTest.cpp
--- a/Engine/Plugins/PCG/Source/PCG/Private/Tests/Elements/PCGMultiSelectTest.cpp
+++ b/Engine/Plugins/PCG/Source/PCG/Private/Tests/Elements/PCGMultiSelectTest.cpp
@@ -43,7 +43,7 @@ protected:
 
 		// Add a tagged input to each of the valid pin paths
 		TArray<FPCGPinProperties> SelectElementPinProperties = TestSettings->AllInputPinProperties();
-		for (int i = 0; i < InputNum; ++i)
+		for (int i = 0; i < SelectElementPinProperties.Num(); ++i)
 		{
 			if (!SelectElementPinProperties[i].IsAdvancedPin() && SelectElementPinProperties[i].Label != DefaultPathPinLabel)
 			{
