[PCG] Fix non-deterministic hash in FPCGGenerationRadii

#rnx

#Synced-CL 50398172

diff --git a/Engine/Plugins/PCG/Source/PCG/Private/PCGCommon.cpp b/Engine/Plugins/PCG/Source/PCG/Private/PCGCommon.cpp
--- a/Engine/Plugins/PCG/Source/PCG/Private/PCGCommon.cpp
+++ b/Engine/Plugins/PCG/Source/PCG/Private/PCGCommon.cpp
@@ -205,6 +205,11 @@ namespace PCGHiGenGrid
 	}
 }
 
+void FPCGRuntimeGenerationRadii::PostSerialize(const FArchive& Ar)
+{
+	ComputeHash();
+}
+
 double FPCGRuntimeGenerationRadii::GetGenerationRadiusFromGrid(EPCGHiGenGrid Grid) const
 {
 	const int32 QualityLevel = UPCGSubsystem::GetPCGQualityLevel();
@@ -365,16 +370,24 @@ void FPCGRuntimeGenerationRadii::ApplyDeprecation()
 }
 #endif
 
-uint32 GetTypeHash(const FPCGRuntimeGenerationRadii& InGenerationRadii)
+void FPCGRuntimeGenerationRadii::ComputeHash()
 {
-	uint32 Hash = 0;
+	Hash = 0;
 
 	auto HashQualityFloat = [](const FPerQualityLevelFloat& InQualityFloat)
 	{
 		uint32 Hash = 0;
 
-		for (const auto& [QualityLevel, Value] : InQualityFloat.PerQuality)
+		TArray<int32, TInlineAllocator<4>> QualityLevels;
+		InQualityFloat.PerQuality.GetKeys(QualityLevels);
+
+		// Sort quality levels so that the hash will be stable.
+		QualityLevels.Sort();
+
+		for (const int32 QualityLevel : QualityLevels)
 		{
+			const float Value = InQualityFloat.GetValue(QualityLevel);
+
 			Hash = HashCombine(Hash, HashCombine(QualityLevel, Value));
 		}
 
@@ -383,20 +396,23 @@ uint32 GetTypeHash(const FPCGRuntimeGenerationRadii& InGenerationRadii)
 		return Hash;
 	};
 
-	Hash = HashCombine(Hash, HashQualityFloat(InGenerationRadii.RadiusUnbounded));
-	Hash = HashCombine(Hash, HashQualityFloat(InGenerationRadii.Radius400));
-	Hash = HashCombine(Hash, HashQualityFloat(InGenerationRadii.Radius800));
-	Hash = HashCombine(Hash, HashQualityFloat(InGenerationRadii.Radius1600));
-	Hash = HashCombine(Hash, HashQualityFloat(InGenerationRadii.Radius3200));
-	Hash = HashCombine(Hash, HashQualityFloat(InGenerationRadii.Radius6400));
-	Hash = HashCombine(Hash, HashQualityFloat(InGenerationRadii.Radius12800));
-	Hash = HashCombine(Hash, HashQualityFloat(InGenerationRadii.Radius25600));
-	Hash = HashCombine(Hash, HashQualityFloat(InGenerationRadii.Radius51200));
-	Hash = HashCombine(Hash, HashQualityFloat(InGenerationRadii.Radius102400));
-	Hash = HashCombine(Hash, HashQualityFloat(InGenerationRadii.Radius204800));
-	Hash = HashCombine(Hash, HashQualityFloat(InGenerationRadii.CleanupRadiusScalar));
-
-	return Hash;
+	Hash = HashCombine(Hash, HashQualityFloat(RadiusUnbounded));
+	Hash = HashCombine(Hash, HashQualityFloat(Radius400));
+	Hash = HashCombine(Hash, HashQualityFloat(Radius800));
+	Hash = HashCombine(Hash, HashQualityFloat(Radius1600));
+	Hash = HashCombine(Hash, HashQualityFloat(Radius3200));
+	Hash = HashCombine(Hash, HashQualityFloat(Radius6400));
+	Hash = HashCombine(Hash, HashQualityFloat(Radius12800));
+	Hash = HashCombine(Hash, HashQualityFloat(Radius25600));
+	Hash = HashCombine(Hash, HashQualityFloat(Radius51200));
+	Hash = HashCombine(Hash, HashQualityFloat(Radius102400));
+	Hash = HashCombine(Hash, HashQualityFloat(Radius204800));
+	Hash = HashCombine(Hash, HashQualityFloat(CleanupRadiusScalar));
+}
+
+uint32 GetTypeHash(const FPCGRuntimeGenerationRadii& InGenerationRadii)
+{
+	return InGenerationRadii.Hash;
 }
 
 namespace PCGDelegates
diff --git a/Engine/Plugins/PCG/Source/PCG/Private/PCGComponent.cpp b/Engine/Plugins/PCG/Source/PCG/Private/PCGComponent.cpp
--- a/Engine/Plugins/PCG/Source/PCG/Private/PCGComponent.cpp
+++ b/Engine/Plugins/PCG/Source/PCG/Private/PCGComponent.cpp
@@ -2267,11 +2267,17 @@ void UPCGComponent::PostEditChangeProperty(FPropertyChangedEvent& PropertyChange
 
 	const FName MemberName = PropertyChangedEvent.MemberProperty->GetFName();
 
-	if (MemberName == GET_MEMBER_NAME_CHECKED(UPCGComponent, GenerationRadii) || MemberName == GET_MEMBER_NAME_CHECKED(UPCGComponent, bOverrideGenerationRadii))
+	if (MemberName == GET_MEMBER_NAME_CHECKED(UPCGComponent, bOverrideGenerationRadii))
 	{
 		// RuntimeGen will automatically pick up any changes to generation radii, we don't need to do any work here.
 		return;
 	}
+	else if (MemberName == GET_MEMBER_NAME_CHECKED(UPCGComponent, GenerationRadii))
+	{
+		// RuntimeGen will automatically pick up any changes to generation radii, we don't need to do any work here.
+		GenerationRadii.ComputeHash();
+		return;
+	}
 
 	// Important note: all property changes already go through the OnObjectPropertyChanged, and will be dirtied here.
 	// So where only a Refresh is needed, it goes through the "capture all" else case.
diff --git a/Engine/Plugins/PCG/Source/PCG/Private/PCGGraph.cpp b/Engine/Plugins/PCG/Source/PCG/Private/PCGGraph.cpp
--- a/Engine/Plugins/PCG/Source/PCG/Private/PCGGraph.cpp
+++ b/Engine/Plugins/PCG/Source/PCG/Private/PCGGraph.cpp
@@ -2154,14 +2154,22 @@ void UPCGGraph::PostEditChangeChainProperty(FPropertyChangedChainEvent& Property
 		return;
 	}
 
-	const FName PropertyName = PropertyChangedEvent.GetPropertyName();
+	const FEditPropertyChain::TDoubleLinkedListNode* MemberPropertyNode = PropertyChangedEvent.PropertyChain.GetHead();
+	const FProperty* MemberProperty = MemberPropertyNode ? MemberPropertyNode->GetValue() : nullptr;
+
+	if (!MemberProperty)
+	{
+		return;
+	}
+
+	const FName MemberPropertyName = MemberProperty->GetFName();
 
-	if (PropertyName == GET_MEMBER_NAME_CHECKED(UPCGGraph, bLandscapeUsesMetadata) ||
-		PropertyName == GET_MEMBER_NAME_CHECKED(UPCGGraph, bHasDefaultConstructedInputs))
+	if (MemberPropertyName == GET_MEMBER_NAME_CHECKED(UPCGGraph, bLandscapeUsesMetadata) ||
+		MemberPropertyName == GET_MEMBER_NAME_CHECKED(UPCGGraph, bHasDefaultConstructedInputs))
 	{
 		NotifyGraphChanged(EPCGChangeType::Input);
 	}
-	else if (PropertyName == GET_MEMBER_NAME_CHECKED(UPCGGraph, UserParameters))
+	else if (MemberPropertyName == GET_MEMBER_NAME_CHECKED(UPCGGraph, UserParameters))
 	{
 		EPCGGraphParameterEvent ChangeType = EPCGGraphParameterEvent::None;
 		const int32 NumberOfUserParametersPreEdit = PreviousPropertyBag ? PreviousPropertyBag->GetPropertyDescs().Num() : 0;
@@ -2229,25 +2237,29 @@ void UPCGGraph::PostEditChangeChainProperty(FPropertyChangedChainEvent& Property
 			OnGraphParametersChanged(ChangeType, ChangedPropertyName);
 		}
 	}
-	else if (PropertyName == GET_MEMBER_NAME_CHECKED(UPCGGraph, HiGenGridSize)
-		|| PropertyName == GET_MEMBER_NAME_CHECKED(UPCGGraph, bUseHierarchicalGeneration)
-		|| PropertyName == GET_MEMBER_NAME_CHECKED(UPCGGraph, bUse2DGrid)
-		|| PropertyName == GET_MEMBER_NAME_CHECKED(UPCGGraph, HiGenExponential)
-		|| PropertyName == GET_MEMBER_NAME_CHECKED(UPCGGraph, bUseActorComponentlessGeneration))
+	else if (MemberPropertyName == GET_MEMBER_NAME_CHECKED(UPCGGraph, HiGenGridSize)
+		|| MemberPropertyName == GET_MEMBER_NAME_CHECKED(UPCGGraph, bUseHierarchicalGeneration)
+		|| MemberPropertyName == GET_MEMBER_NAME_CHECKED(UPCGGraph, bUse2DGrid)
+		|| MemberPropertyName == GET_MEMBER_NAME_CHECKED(UPCGGraph, HiGenExponential)
+		|| MemberPropertyName == GET_MEMBER_NAME_CHECKED(UPCGGraph, bUseActorComponentlessGeneration))
 	{
 		// The higen settings change the structure of the graph (presence or absence of links between grid levels).
 		NotifyGraphChanged(EPCGChangeType::Structural | EPCGChangeType::GenerationGrid);
 	}
-	else if (PropertyName == GET_MEMBER_NAME_CHECKED(UPCGGraphInterface, Title)
-		|| PropertyName == GET_MEMBER_NAME_CHECKED(UPCGGraphInterface, bOverrideTitle)
-		|| PropertyName == GET_MEMBER_NAME_CHECKED(UPCGGraphInterface, Color)
-		|| PropertyName == GET_MEMBER_NAME_CHECKED(UPCGGraphInterface, bOverrideColor)
-		|| PropertyName == GET_MEMBER_NAME_CHECKED(UPCGGraph, bIsStandaloneGraph)
-		|| PropertyName == GET_MEMBER_NAME_CHECKED(UPCGGraph, Category)
-		|| PropertyName == GET_MEMBER_NAME_CHECKED(UPCGGraph, Description))
+	else if (MemberPropertyName == GET_MEMBER_NAME_CHECKED(UPCGGraphInterface, Title)
+		|| MemberPropertyName == GET_MEMBER_NAME_CHECKED(UPCGGraphInterface, bOverrideTitle)
+		|| MemberPropertyName == GET_MEMBER_NAME_CHECKED(UPCGGraphInterface, Color)
+		|| MemberPropertyName == GET_MEMBER_NAME_CHECKED(UPCGGraphInterface, bOverrideColor)
+		|| MemberPropertyName == GET_MEMBER_NAME_CHECKED(UPCGGraph, bIsStandaloneGraph)
+		|| MemberPropertyName == GET_MEMBER_NAME_CHECKED(UPCGGraph, Category)
+		|| MemberPropertyName == GET_MEMBER_NAME_CHECKED(UPCGGraph, Description))
 	{
 		NotifyGraphChanged(EPCGChangeType::Cosmetic | EPCGChangeType::GraphCustomization);
 	}
+	else if (MemberPropertyName == GET_MEMBER_NAME_CHECKED(UPCGGraph, GenerationRadii))
+	{
+		GenerationRadii.ComputeHash();
+	}
 	else if (VerifyIfGraphCustomizationChanged(PropertyChangedEvent))
 	{
 		NotifyGraphChanged(EPCGChangeType::GraphCustomization);
diff --git a/Engine/Plugins/PCG/Source/PCG/Public/PCGCommon.h b/Engine/Plugins/PCG/Source/PCG/Public/PCGCommon.h
--- a/Engine/Plugins/PCG/Source/PCG/Public/PCGCommon.h
+++ b/Engine/Plugins/PCG/Source/PCG/Public/PCGCommon.h
@@ -587,6 +587,8 @@ struct FPCGRuntimeGenerationRadii
 	GENERATED_BODY()
 
 public:
+	UE_API void PostSerialize(const FArchive& Ar);
+
 	/** Get the runtime generation radius for the given grid size. */
 	UE_API double GetGenerationRadiusFromGrid(EPCGHiGenGrid Grid) const;
 
@@ -599,6 +601,9 @@ public:
 	UE_API void ApplyDeprecation();
 #endif
 
+	/** Recompute the cached hash value. Call this anytime the values are changed. */
+	UE_API void ComputeHash();
+
 	friend UE_API uint32 GetTypeHash(const FPCGRuntimeGenerationRadii& In);
 
 	static constexpr double DefaultGenerationRadiusMultiplier = 2.0;
@@ -643,6 +648,9 @@ protected:
 	UPROPERTY(BlueprintReadWrite, EditAnywhere, Category = "Runtime Generation", meta = (UIMin = "1.0", ClampMin = "1.0"))
 	FPerQualityLevelFloat CleanupRadiusScalar = DefaultCleanupRadiusMultiplier;
 
+private:
+	uint32 Hash = 0;
+
 	// Deprecated section
 public:
 	UPROPERTY(meta=(DeprecatedProperty, DeprecationMessage="Use RadiusUnbounded instead"))
@@ -682,6 +690,15 @@ public:
 	double CleanupRadiusMultiplier = 0.0;
 };
 
+template<>
+struct TStructOpsTypeTraits<FPCGRuntimeGenerationRadii> : public TStructOpsTypeTraitsBase2<FPCGRuntimeGenerationRadii>
+{
+	enum
+	{
+		WithPostSerialize = true,
+	};
+};
+
 struct FInstancedPropertyBag;
 
 namespace PCGDelegates
