[PCG] FPCGGraphCustomization: Remove unused Graph UPROPERTY as this was causing non-default value struct

#rnx

#Synced-CL 50423609

diff --git a/Engine/Plugins/PCG/Source/PCG/Private/PCGGraph.cpp b/Engine/Plugins/PCG/Source/PCG/Private/PCGGraph.cpp
--- a/Engine/Plugins/PCG/Source/PCG/Private/PCGGraph.cpp
+++ b/Engine/Plugins/PCG/Source/PCG/Private/PCGGraph.cpp
@@ -514,10 +514,6 @@ UPCGGraph::UPCGGraph(const FObjectInitializer& ObjectInitializer)
 	// Force the user parameters to have an empty property bag. It is necessary to catch the first
 	// add property into the undo/redo history.
 	UserParameters.MigrateToNewBagStruct(UPropertyBag::GetOrCreateFromDescs({}));
-
-#if WITH_EDITOR
-	GraphCustomization.Graph = this;
-#endif
 }
 
 void UPCGGraph::PostLoad()
@@ -626,8 +622,6 @@ void UPCGGraph::PostLoad()
 	}
 
 	GenerationRadii.ApplyDeprecation();
-
-	GraphCustomization.Graph = this;
 #endif // WITH_EDITOR
 }
 
diff --git a/Engine/Plugins/PCG/Source/PCG/Public/Editor/PCGGraphCustomization.h b/Engine/Plugins/PCG/Source/PCG/Public/Editor/PCGGraphCustomization.h
--- a/Engine/Plugins/PCG/Source/PCG/Public/Editor/PCGGraphCustomization.h
+++ b/Engine/Plugins/PCG/Source/PCG/Public/Editor/PCGGraphCustomization.h
@@ -28,10 +28,6 @@ struct FPCGGraphEditorCustomization
 {
 	GENERATED_BODY()
 
-	// Pointer to graph so we can have hierarchical behavior
-	UPROPERTY(Transient)
-	TObjectPtr<const UPCGGraph> Graph;
-
 	UPROPERTY(EditAnywhere, Category = Filtering)
 	bool bFilterNodesByCategory = false;
 
