Added a UE::CUEnum concept for testing if a type is a UEnum.

#Synced-CL 50486641

diff --git a/Engine/Plugins/PCG/Source/PCG/Public/PCGGraph.h b/Engine/Plugins/PCG/Source/PCG/Public/PCGGraph.h
--- a/Engine/Plugins/PCG/Source/PCG/Public/PCGGraph.h
+++ b/Engine/Plugins/PCG/Source/PCG/Public/PCGGraph.h
@@ -5,6 +5,7 @@
 #include "PCGCommon.h"
 #include "PCGNode.h"
 #include "PCGSettings.h"
+#include "Concepts/UEnum.h"
 #include "Data/PCGToolData.h"
 #include "Editor/PCGGraphCustomization.h"
 #include "Graph/PCGStackContext.h"
@@ -239,7 +240,7 @@ public:
 		const FInstancedPropertyBag* UserParameters = GetUserParametersStruct();
 		check(UserParameters);
 
-		if constexpr (std::is_enum_v<T> && StaticEnum<T>())
+		if constexpr (UE::CUEnum<T>)
 		{
 			return FPCGGraphParameterExtension::GetGraphParameter<T>(*UserParameters, PropertyName, StaticEnum<T>());
 		}
@@ -258,7 +259,7 @@ public:
 		check(UserParameters);
 
 		EPropertyBagResult Result;
-		if constexpr (std::is_enum_v<T> && StaticEnum<T>())
+		if constexpr (UE::CUEnum<T>)
 		{
 			Result = FPCGGraphParameterExtension::SetGraphParameter(*UserParameters, PropertyName, Value, StaticEnum<T>());
 		}
diff --git a/Engine/Source/Runtime/Core/Public/Concepts/UEnum.h b/Engine/Source/Runtime/Core/Public/Concepts/UEnum.h
new file mode 100644
--- /dev/null
+++ b/Engine/Source/Runtime/Core/Public/Concepts/UEnum.h
@@ -0,0 +1,28 @@
+// Copyright Epic Games, Inc. All Rights Reserved.
+
+#pragma once
+
+#include "Concepts/CompleteType.h"
+#include "Concepts/Enum.h"
+#include "Concepts/SameAs.h"
+#include "UObject/CoreReflectedTypeAccessors.h"
+
+#include <type_traits>
+
+namespace UE
+{
+	/**
+	* Concept which describes a UEnum (i.e. an enum marked up with UENUM(), not a UObject derived from UEnum).
+	*/
+	template <typename T>
+	concept CUEnum =
+		UE::CEnum<T> &&
+		//  We try to guard against incomplete types here, but this will succeed for forward declared enums with an underlying type, so
+		// it doesn't really work.  Until a solution is found, it's possible for this concept to be memoized in an incorrect state, so
+		// care must be taken.
+		UE::CCompleteType<T> &&
+		requires (T* Ptr)
+		{
+			{ StaticEnum<std::remove_cv_t<T>>() } -> UE::CSameAs<UEnum*>;
+		};
+}
diff --git a/Engine/Source/Runtime/Core/Public/UObject/CoreReflectedTypeAccessors.h b/Engine/Source/Runtime/Core/Public/UObject/CoreReflectedTypeAccessors.h
--- a/Engine/Source/Runtime/Core/Public/UObject/CoreReflectedTypeAccessors.h
+++ b/Engine/Source/Runtime/Core/Public/UObject/CoreReflectedTypeAccessors.h
@@ -8,4 +8,4 @@ class UEnum;
 C++ templated Static(Enum) retrieval function prototypes.
 -----------------------------------------------------------------------------*/
 
-template<typename EnumType> UEnum* StaticEnum();
+template <typename EnumType> UEnum* StaticEnum() = delete;
