[PCG] Fixed bug preventing node overrides from appearing on Save Data View node

#rnx

#Synced-CL 50411743

diff --git a/Engine/Plugins/PCG/Source/PCG/Private/Elements/IO/PCGSaveDataViewElement.cpp b/Engine/Plugins/PCG/Source/PCG/Private/Elements/IO/PCGSaveDataViewElement.cpp
--- a/Engine/Plugins/PCG/Source/PCG/Private/Elements/IO/PCGSaveDataViewElement.cpp
+++ b/Engine/Plugins/PCG/Source/PCG/Private/Elements/IO/PCGSaveDataViewElement.cpp
@@ -88,10 +88,10 @@ TArray<FPCGSettingsOverridableParam> UPCGSaveDataViewSettings::GatherOverridable
 	if (ConverterParameters.IsValid())
 	{
 		const PCGSettingsHelpers::FPCGGetAllOverridableParamsConfig Config;
-		return PCGSettingsHelpers::GetAllOverridableParams(ConverterParameters.GetScriptStruct(), Config);
+		OverridableParams.Append(PCGSettingsHelpers::GetAllOverridableParams(ConverterParameters.GetScriptStruct(), Config));
 	}
 
-	return {};
+	return OverridableParams;
 }
 #endif // WITH_EDITOR
 
