[PCG] Remove overly zealous ensure

#rnx
#Synced-CL 50483771

diff --git a/Engine/Plugins/PCG/Source/PCG/Private/PCGTrackingManager.cpp b/Engine/Plugins/PCG/Source/PCG/Private/PCGTrackingManager.cpp
--- a/Engine/Plugins/PCG/Source/PCG/Private/PCGTrackingManager.cpp
+++ b/Engine/Plugins/PCG/Source/PCG/Private/PCGTrackingManager.cpp
@@ -417,7 +417,6 @@ bool FPCGTrackingManager::RemapExecutionSource(const IPCGGraphExecutionSource* I
 		{
 			PCG::TScopeLock Lock(DelayedComponentsRegistrationLock);
 			DelayedComponentsToUnregister.Remove(OldComponent);
-			ensure(DelayedComponentsToRegister.Contains(NewComponent));
 		}
 #endif
 
