[PCG] Fix stale RG execution source entries in debug object tree

#rnx
#Synced-CL 50486779

#Synced-CL 50398172

diff --git a/Engine/Plugins/PCG/Source/PCG/Private/PCGComponent.cpp b/Engine/Plugins/PCG/Source/PCG/Private/PCGComponent.cpp
--- a/Engine/Plugins/PCG/Source/PCG/Private/PCGComponent.cpp
+++ b/Engine/Plugins/PCG/Source/PCG/Private/PCGComponent.cpp
@@ -1715,7 +1715,7 @@ void UPCGComponent::OnUnregister()
 			Subsystem->CancelGeneration(this);
 		}
 
-		Subsystem->OnPCGComponentUnregistered.Broadcast(IsValidChecked(this) ? this : nullptr);
+		Subsystem->GetOnPCGSourceUnregistered().Broadcast(IsValidChecked(this) ? this : nullptr);
 
 		// Always unregister here if we are being unloaded (we are not being re-registered)
 		if (bUnregisteredThroughLoading)
diff --git a/Engine/Plugins/PCG/Source/PCG/Private/PCGModule.cpp b/Engine/Plugins/PCG/Source/PCG/Private/PCGModule.cpp
--- a/Engine/Plugins/PCG/Source/PCG/Private/PCGModule.cpp
+++ b/Engine/Plugins/PCG/Source/PCG/Private/PCGModule.cpp
@@ -20,6 +20,7 @@
 #include "Data/DataView/PCGDataViewNativePropertySelectors.h"
 #include "Helpers/PCGActorHelpers.h"
 #include "Helpers/PCGWorldQueryHelpers.h"
+#include "RuntimeGen/PCGRuntimeGenExecutionSource.h"
 
 #include "Engine/Engine.h"
 #include "Interfaces/IPluginManager.h"
@@ -79,35 +80,36 @@ namespace PCGModule
 
 		virtual TArray<FPCGGraphExecutionSourceDescriptor> GatherExecutionSources() const override
 		{
-			TArray<FPCGGraphExecutionSourceDescriptor> ExecutionSources;
+			const TStaticArray<const UClass*, 3> ClassToGather =
+			{
+				UPCGComponent::StaticClass(),
+				UPCGDefaultExecutionSource::StaticClass(),
+				UPCGRuntimeGenExecutionSource::StaticClass()
+			};
 
 			TArray<UObject*> PCGSources;
-			GetObjectsOfClass(UPCGComponent::StaticClass(), PCGSources, /*bIncludeDerivedClasses=*/ true);
 
-			ExecutionSources.Reserve(PCGSources.Num());
-
-			for (UObject* PCGSource : PCGSources)
+			for (const UClass* Class : ClassToGather)
 			{
-				UPCGComponent* PCGComponent = CastChecked<UPCGComponent>(PCGSource, ECastCheckedType::NullAllowed);
-				if (IsValid(PCGComponent) && PCGComponent->IsRegistered() && PCGComponent->GetGraph())
-				{
-					FPCGGraphExecutionSourceDescriptor Descriptor = { /*ExecutionSource=*/ PCGComponent };
-#if WITH_EDITOR
-					Descriptor.bShowInDebugger = true;
-#endif
-					ExecutionSources.Add(MoveTemp(Descriptor));
-				}
+				GetObjectsOfClass(Class, PCGSources, /*bIncludeDerivedClasses=*/true);
 			}
 
-			PCGSources.Reset();
-			GetObjectsOfClass(UPCGDefaultExecutionSource::StaticClass(), PCGSources, /*bIncludeDerivedClasses=*/ true);
+			TArray<FPCGGraphExecutionSourceDescriptor> ExecutionSources;
 
 			for (UObject* PCGSource : PCGSources)
 			{
-				UPCGDefaultExecutionSource* PCGDefaultExecutionSource = CastChecked<UPCGDefaultExecutionSource>(PCGSource, ECastCheckedType::NullAllowed);
-				if (IsValid(PCGDefaultExecutionSource) && PCGDefaultExecutionSource->GetExecutionState().GetGraph())
+				IPCGGraphExecutionSource* ExecutionSource = CastChecked<IPCGGraphExecutionSource>(PCGSource, ECastCheckedType::NullAllowed);
+
+				bool bIsValid = IsValid(PCGSource) && ExecutionSource->GetExecutionState().GetGraph();
+
+				if (UPCGComponent* Component = Cast<UPCGComponent>(ExecutionSource))
+				{
+					bIsValid &= Component->IsRegistered();
+				}
+
+				if (bIsValid)
 				{
-					FPCGGraphExecutionSourceDescriptor Descriptor = { /*ExecutionSource=*/ PCGDefaultExecutionSource };
+					FPCGGraphExecutionSourceDescriptor Descriptor = { /*ExecutionSource=*/ ExecutionSource};
 #if WITH_EDITOR
 					Descriptor.bShowInDebugger = true;
 #endif
diff --git a/Engine/Plugins/PCG/Source/PCG/Private/RuntimeGen/PCGRuntimeGenExecutionSource.cpp b/Engine/Plugins/PCG/Source/PCG/Private/RuntimeGen/PCGRuntimeGenExecutionSource.cpp
--- a/Engine/Plugins/PCG/Source/PCG/Private/RuntimeGen/PCGRuntimeGenExecutionSource.cpp
+++ b/Engine/Plugins/PCG/Source/PCG/Private/RuntimeGen/PCGRuntimeGenExecutionSource.cpp
@@ -48,6 +48,9 @@ void UPCGRuntimeGenExecutionSource::Unregister()
 {
 	if (UPCGSubsystem* Subsystem = Cast<UPCGSubsystem>(ExecutionState.GetSubsystem()))
 	{
+#if WITH_EDITOR
+		Subsystem->GetOnPCGSourceUnregistered().Broadcast(this);
+#endif
 		Subsystem->UnregisterExecutionSource(this);
 	}
 }
@@ -567,6 +570,8 @@ void UPCGRuntimeGenExecutionSource::DestroyLocalSource(UPCGRuntimeGenExecutionSo
 		CleanupParams.bImmediate = true;
 		
 		InLocalSource->GetExecutionState().Cleanup(CleanupParams);
+		InLocalSource->Unregister();
+		InLocalSource->MarkAsGarbage();
 	}
 
 	LocalSources.Remove(InLocalSource);
diff --git a/Engine/Plugins/PCG/Source/PCG/Private/RuntimeGen/PCGRuntimeGenScheduler.cpp b/Engine/Plugins/PCG/Source/PCG/Private/RuntimeGen/PCGRuntimeGenScheduler.cpp
--- a/Engine/Plugins/PCG/Source/PCG/Private/RuntimeGen/PCGRuntimeGenScheduler.cpp
+++ b/Engine/Plugins/PCG/Source/PCG/Private/RuntimeGen/PCGRuntimeGenScheduler.cpp
@@ -1887,6 +1887,50 @@ void FPCGRuntimeGenScheduler::CleanupDelayedRefreshSources()
 {
 	TRACE_CPUPROFILER_EVENT_SCOPE(FPCGRuntimeGenScheduler::CleanupDelayedRefreshSources);
 
+	// Check that each refreshed local source is still intersecting its original source.
+	// If it is not, it would be leaked instead of refreshed, so we should force a full cleanup.
+	for (const FGridGenerationKey& GenerationKey : GeneratedSourcesToRemove)
+	{
+		if (!GenerationKey.IsValid())
+		{
+			continue;
+		}
+
+		const FPCGGridDescriptor GridDescriptor = GenerationKey.GetGridDescriptor();
+		const EPCGHiGenGrid Grid = GridDescriptor.GetHiGenGrid();
+		IPCGGraphExecutionSource* OriginalSource = GenerationKey.GetOriginalSource();
+		const FIntVector& GridCoords = GenerationKey.GetGridCoords();
+
+		// The unbounded grid level will always lie inside the original source, so we can skip it.
+		if (Grid == EPCGHiGenGrid::Unbounded)
+		{
+			if (OriginalSource && !OriginalSource->GetExecutionState().IsActive())
+			{
+				CleanupSource(GenerationKey, OriginalSource);
+			}
+
+			continue;
+		}
+
+		IPCGGraphExecutionSource* LocalSource = OriginalSource ? OriginalSource->GetExecutionState().GetLocalSource(GridDescriptor, GridCoords) : nullptr;
+
+		if (LocalSource)
+		{
+			const FBox OriginalBounds = OriginalSource->GetExecutionState().GetBounds();
+			const FBox LocalBounds = LocalSource->GetExecutionState().GetBounds();
+
+			if (!OriginalBounds.Intersect(LocalBounds) || !OriginalSource->GetExecutionState().IsActive())
+			{
+				CleanupSource(GenerationKey, LocalSource);
+			}
+		}
+		else
+		{
+			// If the local source could not be recovered, just clean up.
+			CleanupSource(GenerationKey, /*LocalSource=*/nullptr);
+		}
+	}
+
 	// Remove any generation keys that have been registered for deferred removal.
 	if (!GeneratedSourcesToRemove.IsEmpty())
 	{
@@ -2362,6 +2406,7 @@ void FPCGRuntimeGenScheduler::UnregisterOriginalRuntimeGenExecutionSource(UPCGCo
 
 			(*FoundRuntimeGenSource)->Cleanup(CleanupParams);
 			(*FoundRuntimeGenSource)->Unregister();
+			(*FoundRuntimeGenSource)->MarkAsGarbage();
 		}
 
 		OriginalExecutionSources.Remove(InOriginalComponent);
diff --git a/Engine/Plugins/PCG/Source/PCG/Public/Subsystems/IPCGBaseSubsystem.h b/Engine/Plugins/PCG/Source/PCG/Public/Subsystems/IPCGBaseSubsystem.h
--- a/Engine/Plugins/PCG/Source/PCG/Public/Subsystems/IPCGBaseSubsystem.h
+++ b/Engine/Plugins/PCG/Source/PCG/Public/Subsystems/IPCGBaseSubsystem.h
@@ -30,6 +30,7 @@ class UWorld;
 
 #if WITH_EDITOR
 DECLARE_MULTICAST_DELEGATE_ThreeParams(FPCGOnPCGSourceGenerationDone, IPCGBaseSubsystem*, IPCGGraphExecutionSource*, EPCGGenerationStatus);
+DECLARE_MULTICAST_DELEGATE_OneParam(FPCGOnPCGSourceUnregistered, IPCGGraphExecutionSource*);
 #endif // WITH_EDITOR
 
 UINTERFACE()
@@ -114,6 +115,7 @@ public:
 	PCG_API virtual void OnPCGSourceGenerationDone(IPCGGraphExecutionSource* InExecutionSource, EPCGGenerationStatus InStatus);
 
 	FPCGOnPCGSourceGenerationDone& GetOnPCGSourceGenerationDone() { return OnPCGSourceGenerationDoneDelegate; }
+	FPCGOnPCGSourceUnregistered& GetOnPCGSourceUnregistered() { return OnPCGSourceUnregisteredDelegate; }
 #endif // WITH_EDITOR
 	
 protected:
@@ -133,5 +135,6 @@ protected:
 	PCG_API void SetDisableClearResults(bool bInDisableClearResults);
 
 	FPCGOnPCGSourceGenerationDone OnPCGSourceGenerationDoneDelegate;
+	FPCGOnPCGSourceUnregistered OnPCGSourceUnregisteredDelegate;
 #endif
 };
diff --git a/Engine/Plugins/PCG/Source/PCG/Public/Subsystems/PCGSubsystem.h b/Engine/Plugins/PCG/Source/PCG/Public/Subsystems/PCGSubsystem.h
--- a/Engine/Plugins/PCG/Source/PCG/Public/Subsystems/PCGSubsystem.h
+++ b/Engine/Plugins/PCG/Source/PCG/Public/Subsystems/PCGSubsystem.h
@@ -52,11 +52,12 @@ class UWorld;
 
 DECLARE_MULTICAST_DELEGATE_OneParam(FPCGOnComponentGenerationCompleteOrCancelled, UPCGSubsystem*);
 
-DECLARE_MULTICAST_DELEGATE_OneParam(FPCGOnPCGComponentUnregistered, UPCGComponent*);
-
-/** Deprecated - use FPCGOnPCGSourceGenerationDone */
+/** Deprecated 5.7 - use FPCGOnPCGSourceGenerationDone */
 DECLARE_MULTICAST_DELEGATE_ThreeParams(FPCGOnPCGComponentGenerationDone, UPCGSubsystem*, UPCGComponent*, EPCGGenerationStatus);
 
+/** Deprecated 5.8 - use FPCGOnPCGSourceUnregistered */
+DECLARE_MULTICAST_DELEGATE_OneParam(FPCGOnPCGComponentUnregistered, UPCGComponent*);
+
 // Event triggered when GenerateAllPCGComponents is done
 DECLARE_MULTICAST_DELEGATE(FPCGOnAllComponentsGenerated);
 
@@ -398,8 +399,6 @@ public:
 	UE_DEPRECATED(5.8, "No longer used")
 	void NotifyLandscapeEditModeExited() {}
 
-	FPCGOnPCGComponentUnregistered OnPCGComponentUnregistered;
-
 	FPCGOnAllComponentsGenerated OnAllComponentsGenerated;
 	FPCGOnAllComponentsCleanedup OnAllComponentsCleanedup;
 	FPCGOnAllComponentsClearedLink OnAllComponentsClearedLink;
@@ -407,6 +406,9 @@ public:
 	UE_DEPRECATED(5.7, "Deprecated in favor of OnPCGSourceGenerationDone, will not be notified anymore")
 	FPCGOnPCGComponentGenerationDone OnPCGComponentGenerationDone;
 
+	UE_DEPRECATED(5.8, "Deprecated in favor of OnPCGSourceUnregistered, will not be notified anymore")
+	FPCGOnPCGComponentUnregistered OnPCGComponentUnregistered;
+
 	PCG_API void CreateMissingPartitionActors();
 
 	void OnPreviewPlatformChanged();
diff --git a/Engine/Plugins/PCG/Source/PCGEditor/Private/PCGEditor.cpp b/Engine/Plugins/PCG/Source/PCGEditor/Private/PCGEditor.cpp
--- a/Engine/Plugins/PCG/Source/PCGEditor/Private/PCGEditor.cpp
+++ b/Engine/Plugins/PCG/Source/PCGEditor/Private/PCGEditor.cpp
@@ -3683,10 +3683,10 @@ void FPCGEditor::JumpToDefinition(const UClass* Class) const
 	}
 }
 
-void FPCGEditor::OnComponentUnregistered(UPCGComponent* Component)
+void FPCGEditor::OnSourceUnregistered(IPCGGraphExecutionSource* Source)
 {
-	// Refresh the debug object tree to avoid stale entries from components that have been unregistered.
-	if (!Component || Component->GetGraph() == MainGraph)
+	// Refresh the debug object tree to avoid stale entries from sources that have been unregistered.
+	if (!Source || Source->GetExecutionState().GetGraph() == MainGraph)
 	{
 		if (DebugObjectTreeWidget.IsValid())
 		{
@@ -3696,7 +3696,7 @@ void FPCGEditor::OnComponentUnregistered(UPCGComponent* Component)
 
 	if (IPCGEditorModule* PCGEditorModule = IPCGEditorModule::Get())
 	{
-		PCGEditorModule->GetNodeVisualLogsMutable().ClearLogs(Component);
+		PCGEditorModule->GetNodeVisualLogsMutable().ClearLogs(Source);
 	}
 }
 
@@ -3769,7 +3769,7 @@ void FPCGEditor::RegisterDelegatesForWorld(UWorld* World)
 
 	if (UPCGSubsystem* Subsystem = UPCGSubsystem::GetInstance(World))
 	{
-		Subsystem->OnPCGComponentUnregistered.AddRaw(this, &FPCGEditor::OnComponentUnregistered);
+		Subsystem->GetOnPCGSourceUnregistered().AddRaw(this, &FPCGEditor::OnSourceUnregistered);
 		Subsystem->GetOnPCGSourceGenerationDone().AddRaw(this, &FPCGEditor::OnSourceGenerationDone);
 	}
 }
@@ -3778,7 +3778,7 @@ void FPCGEditor::UnregisterDelegatesForWorld(UWorld* World)
 {
 	if (UPCGSubsystem* Subsystem = UPCGSubsystem::GetInstance(World))
 	{
-		Subsystem->OnPCGComponentUnregistered.RemoveAll(this);
+		Subsystem->GetOnPCGSourceUnregistered().RemoveAll(this);
 		Subsystem->GetOnPCGSourceGenerationDone().RemoveAll(this);
 	}
 }
diff --git a/Engine/Plugins/PCG/Source/PCGEditor/Public/PCGEditor.h b/Engine/Plugins/PCG/Source/PCGEditor/Public/PCGEditor.h
--- a/Engine/Plugins/PCG/Source/PCGEditor/Public/PCGEditor.h
+++ b/Engine/Plugins/PCG/Source/PCGEditor/Public/PCGEditor.h
@@ -570,8 +570,8 @@ private:
 	 */
 	void JumpToDefinition(const UClass* Class) const;
 
-	/** Called when a PCG component unregisters. */
-	void OnComponentUnregistered(UPCGComponent* Component);
+	/** Called when a PCG execution source unregisters. */
+	void OnSourceUnregistered(IPCGGraphExecutionSource* Source);
 
 	/** Called when a component finishes executing. Useful for updating debugging tools/UIs. */
 	void OnSourceGenerationDone(IPCGBaseSubsystem* Subsystem, IPCGGraphExecutionSource* Source, EPCGGenerationStatus Status);
