[PCG] Support generic attributes for accessors, for existing types

* Use CRTP to factorize the logic for accessors with typed attributes and generic attributes.
* Still temporary, `IsGeneric` is not expected to stay at the moment.

#rnx

#Synced-CL 50489647

diff --git a/Engine/Plugins/PCG/Source/PCG/Private/Metadata/Accessors/PCGAttributeAccessorHelpers.cpp b/Engine/Plugins/PCG/Source/PCG/Private/Metadata/Accessors/PCGAttributeAccessorHelpers.cpp
--- a/Engine/Plugins/PCG/Source/PCG/Private/Metadata/Accessors/PCGAttributeAccessorHelpers.cpp
+++ b/Engine/Plugins/PCG/Source/PCG/Private/Metadata/Accessors/PCGAttributeAccessorHelpers.cpp
@@ -366,7 +366,14 @@ TUniquePtr<const IPCGAttributeAccessor> PCGAttributeAccessorHelpers::CreateConst
 	{
 		auto MakeConstAccessor = [InAttribute, InMetadata]<typename T>(T) -> TUniquePtr<const IPCGAttributeAccessor>
 		{
-			return MakeUnique<const FPCGAttributeAccessor<T>>(static_cast<const FPCGMetadataAttribute<T>*>(InAttribute), InMetadata);
+			if (InAttribute->IsGeneric())
+			{
+				return MakeUnique<const FPCGAttributeGenericAccessor<T>>(static_cast<const FPCGMetadataAttributeGeneric*>(InAttribute), InMetadata);
+			}
+			else
+			{
+				return MakeUnique<const FPCGAttributeAccessor<T>>(static_cast<const FPCGMetadataAttribute<T>*>(InAttribute), InMetadata);
+			}
 		};
 		
 		return PCGMetadataAttribute::CallbackWithRightType(InAttribute->GetTypeId(), MakeConstAccessor);
@@ -443,7 +450,14 @@ TUniquePtr<IPCGAttributeAccessor> PCGAttributeAccessorHelpers::CreateAccessor(FP
 	{
 		auto MakeAccessor = [InAttribute, InMetadata]<typename T>(T) -> TUniquePtr<IPCGAttributeAccessor>
 		{
-			return MakeUnique<FPCGAttributeAccessor<T>>(static_cast<FPCGMetadataAttribute<T>*>(InAttribute), InMetadata);
+			if (InAttribute->IsGeneric())
+			{
+				return MakeUnique<FPCGAttributeGenericAccessor<T>>(static_cast<FPCGMetadataAttributeGeneric*>(InAttribute), InMetadata);
+			}
+			else
+			{
+				return MakeUnique<FPCGAttributeAccessor<T>>(static_cast<FPCGMetadataAttribute<T>*>(InAttribute), InMetadata);
+			}
 		};
 		
 		return PCGMetadataAttribute::CallbackWithRightType(InAttribute->GetTypeId(), MakeAccessor);
diff --git a/Engine/Plugins/PCG/Source/PCG/Public/Metadata/Accessors/PCGAttributeAccessor.h b/Engine/Plugins/PCG/Source/PCG/Public/Metadata/Accessors/PCGAttributeAccessor.h
--- a/Engine/Plugins/PCG/Source/PCG/Public/Metadata/Accessors/PCGAttributeAccessor.h
+++ b/Engine/Plugins/PCG/Source/PCG/Public/Metadata/Accessors/PCGAttributeAccessor.h
@@ -5,22 +5,26 @@
 #include "IPCGAttributeAccessorTpl.h"
 
 #include "Metadata/PCGMetadata.h"
+#include "Metadata/PCGMetadataAttributeGeneric.h"
 #include "Metadata/PCGMetadataAttributeTpl.h"
 #include "Metadata/PCGMetadataDomain.h"
 
+#define LOCTEXT_NAMESPACE "PCGAttributeAccessor"
+
 /**
 * Templated accessor class for attributes. Will wrap around a typed attribute.
 * Key supported: MetadataEntryKey and Points
 */
-template <typename T>
-class FPCGAttributeAccessor : public IPCGAttributeAccessorT<FPCGAttributeAccessor<T>>
+template <typename T, typename UnderlyingAccessor>
+class FPCGAttributeAccessorBase : public IPCGAttributeAccessorT<FPCGAttributeAccessorBase<T, UnderlyingAccessor>>
 {
 public:
 	using Type = T;
-	using Super = IPCGAttributeAccessorT<FPCGAttributeAccessor<T>>;
-
+	using Super = IPCGAttributeAccessorT<FPCGAttributeAccessorBase<T, UnderlyingAccessor>>;
+	
+protected:
 	// Can't write if metadata is null
-	FPCGAttributeAccessor(FPCGMetadataAttribute<T>* InAttribute, FPCGMetadataDomain* InMetadataDomain, bool bForceReadOnly = false)
+	FPCGAttributeAccessorBase(FPCGMetadataAttributeBase* InAttribute, FPCGMetadataDomain* InMetadataDomain, bool bForceReadOnly = false)
 		: Super(/*bInReadOnly=*/ InMetadataDomain == nullptr || bForceReadOnly)
 		, Attribute(InAttribute)
 		, MetadataDomain(InMetadataDomain)
@@ -28,24 +32,15 @@ public:
 		check(InAttribute && (!InMetadataDomain || InAttribute->GetMetadataDomain() == InMetadataDomain));
 	}
 
-	FPCGAttributeAccessor(const FPCGMetadataAttribute<T>* InAttribute, const FPCGMetadataDomain* InMetadataDomain, bool bForceReadOnly = false)
+	FPCGAttributeAccessorBase(const FPCGMetadataAttributeBase* InAttribute, const FPCGMetadataDomain* InMetadataDomain, bool bForceReadOnly = false)
 		: Super(/*bInReadOnly=*/ true)
-		, Attribute(const_cast<FPCGMetadataAttribute<T>*>(InAttribute))
+		, Attribute(const_cast<FPCGMetadataAttributeBase*>(InAttribute))
 		, MetadataDomain(const_cast<FPCGMetadataDomain*>(InMetadataDomain))
 	{
 		check(InAttribute && (!InMetadataDomain || InAttribute->GetMetadataDomain() == InMetadataDomain));
 	}
-
-	UE_DEPRECATED(5.6, "Use the FPCGMetadataDomain version")
-	FPCGAttributeAccessor(FPCGMetadataAttribute<T>* InAttribute, UPCGMetadata* InMetadata, bool bForceReadOnly = false)
-		: FPCGAttributeAccessor(InAttribute, InMetadata->GetDefaultMetadataDomain(), bForceReadOnly)
-	{}
-
-	UE_DEPRECATED(5.6, "Use the FPCGMetadataDomain version")
-	FPCGAttributeAccessor(FPCGMetadataAttribute<T>* InAttribute, const UPCGMetadata* InMetadata, bool bForceReadOnly = false)
-		: FPCGAttributeAccessor(InAttribute, InMetadata->GetConstDefaultMetadataDomain(), bForceReadOnly)
-	{}
-
+	
+public:
 	bool GetRangeImpl(TArrayView<T> OutValues, int32 Index, const IPCGAttributeAccessorKeys& Keys) const
 	{
 		TArray<const PCGMetadataEntryKey*, TInlineAllocator<256>> EntryKeyPtrs;
@@ -61,7 +56,7 @@ public:
 		EntryKeys.Reserve(EntryKeyPtrs.Num());
 		Algo::Transform(EntryKeyPtrs, EntryKeys, [](const PCGMetadataEntryKey* KeyPtr) { return *KeyPtr; });
 
-		Attribute->GetValuesFromItemKeys(TArrayView<PCGMetadataEntryKey>(EntryKeys), OutValues);
+		static_cast<const UnderlyingAccessor&>(*this).GetValuesFromItemKeys(TArrayView<PCGMetadataEntryKey>(EntryKeys), OutValues);
 
 		return true;
 	}
@@ -76,7 +71,7 @@ public:
 			if (EntryKey && *EntryKey == PCGInvalidEntryKey)
 			{
 				check(!InValues.IsEmpty())
-				Attribute->SetDefaultValue(InValues[0]);
+				static_cast<UnderlyingAccessor&>(*this).SetDefaultValue(InValues[0]);
 				return true;
 			}
 		}
@@ -96,13 +91,13 @@ public:
 
 			{
 				//TRACE_CPUPROFILER_EVENT_SCOPE(FPCGAttributeAccessor::SetRangeImpl::SetValues);
-				Attribute->SetValues(EntryKeys, InValues);
+				static_cast<UnderlyingAccessor&>(*this).SetValues(EntryKeys, InValues);
 			}
 		}
 		else
 		{
 			//TRACE_CPUPROFILER_EVENT_SCOPE(FPCGAttributeAccessor::SetRangeImpl::SetValues);
-			Attribute->SetValues_TryLockless(EntryKeysView, InValues, Start);
+			static_cast<UnderlyingAccessor&>(*this).SetValues_TryLockless(EntryKeys, InValues, Start);
 		}
 
 		return true;
@@ -113,11 +108,17 @@ public:
 	bool Prepare(IPCGAttributeAccessorKeys& Keys, int32 Count, const bool bCanReuseEntryKeys, int32 Index, bool bPreallocateValues, TArrayView<PCGMetadataEntryKey*>* OutEntryKeys = nullptr, int32* OutStartIndex = nullptr)
 	{
 		//TRACE_CPUPROFILER_EVENT_SCOPE(FPCGAttributeAccessor::Prepare);
+		
+		if (!MetadataDomain || Super::IsReadOnly())
+		{
+			PCGLog::LogWarningOnGraph(LOCTEXT("ReadOnlyInvalidPrepare", "Try to prepare an attribute accessor with a read only accessor. This is invalid."));
+			return false;		
+		}
 
 		const bool bSupportsMultiEntries = MetadataDomain->SupportsMultiEntries();
 		if (!bSupportsMultiEntries && Count > 1)
 		{
-			UE_LOG(LogPCG, Warning, TEXT("Try to prepare multiple values to set in an attribute from a domain that doesn't support multi entries (%ls)"), *MetadataDomain->GetDomainID().DebugName.ToString());
+			PCGLog::LogWarningOnGraph(FText::Format(LOCTEXT("InvalidDomain", "Try to prepare multiple values to set in an attribute from a domain that doesn't support multi entries ({0})"), FText::FromString(MetadataDomain->GetDomainID().DebugName.ToString())));
 			return false;
 		}
 
@@ -168,7 +169,7 @@ public:
 
 		if (bPreallocateValues && ensure(!bWasPrepared))
 		{
-			const int32 StartIndex = Attribute->PreallocateValues(EntryKeysView, /*bLockless=*/false);
+			const int32 StartIndex = static_cast<UnderlyingAccessor&>(*this).PreallocateValues(EntryKeysView, /*bLockless=*/false);
 			if (OutStartIndex)
 			{
 				*OutStartIndex = StartIndex;
@@ -185,12 +186,110 @@ public:
 			return;
 		}
 
-		Attribute->Prepare(Count);
+		static_cast<UnderlyingAccessor&>(*this).PrepareAttribute(Count);
 		bWasPrepared = Prepare(Keys, Count, bCanReuseEntryKeys, /*Index=*/0, /*bPreallocateValues=*/true);
 	}
 
-private:
-	FPCGMetadataAttribute<T>* Attribute = nullptr;
+protected:
+	FPCGMetadataAttributeBase* Attribute = nullptr;
 	FPCGMetadataDomain* MetadataDomain = nullptr;
 	bool bWasPrepared = false;
 };
+
+template <typename T>
+class FPCGAttributeAccessor : public FPCGAttributeAccessorBase<T, FPCGAttributeAccessor<T>>
+{
+public:
+	using Super = FPCGAttributeAccessorBase<T, FPCGAttributeAccessor<T>>;
+	
+	// Can't write if metadata is null
+	FPCGAttributeAccessor(FPCGMetadataAttribute<T>* InAttribute, FPCGMetadataDomain* InMetadataDomain, bool bForceReadOnly = false)
+		: Super(InAttribute, InMetadataDomain, bForceReadOnly)
+	{
+	}
+
+	FPCGAttributeAccessor(const FPCGMetadataAttribute<T>* InAttribute, const FPCGMetadataDomain* InMetadataDomain, bool bForceReadOnly = false)
+		: Super(InAttribute, InMetadataDomain, bForceReadOnly)
+	{
+	}
+	
+	void GetValuesFromItemKeys(TArrayView<PCGMetadataEntryKey> EntryKeys, TArrayView<T> OutValues) const
+	{
+		static_cast<const FPCGMetadataAttribute<T>*>(this->Attribute)->GetValuesFromItemKeys(EntryKeys, OutValues);
+	}
+
+	void SetDefaultValue(const T& Value)
+	{
+		static_cast<FPCGMetadataAttribute<T>*>(this->Attribute)->SetDefaultValue(Value);
+	}
+
+	void SetValues(TArrayView<const PCGMetadataEntryKey * const> ItemKeys, const TArrayView<const T>& InValues)
+	{
+		static_cast<FPCGMetadataAttribute<T>*>(this->Attribute)->SetValues(ItemKeys, InValues);
+	}
+
+	void SetValues_TryLockless(TArrayView<PCGMetadataEntryKey*> ItemKeys, const TArrayView<const T>& InValues, int32 StartIndex)
+	{
+		static_cast<FPCGMetadataAttribute<T>*>(this->Attribute)->SetValues_TryLockless(ItemKeys, InValues, StartIndex);
+	}
+
+	int32 PreallocateValues(TArrayView<PCGMetadataEntryKey*> EntryKeys, bool bLockless)
+	{
+		return static_cast<FPCGMetadataAttribute<T>*>(this->Attribute)->PreallocateValues(EntryKeys, bLockless);
+	}
+
+	void PrepareAttribute(int32 Count)
+	{
+		static_cast<FPCGMetadataAttribute<T>*>(this->Attribute)->Prepare(Count);
+	}
+};
+
+template <typename T>
+class FPCGAttributeGenericAccessor : public FPCGAttributeAccessorBase<T, FPCGAttributeGenericAccessor<T>>
+{
+public:
+	using Super = FPCGAttributeAccessorBase<T, FPCGAttributeGenericAccessor<T>>;
+	
+	// Can't write if metadata is null
+	FPCGAttributeGenericAccessor(FPCGMetadataAttributeGeneric* InAttribute, FPCGMetadataDomain* InMetadataDomain, bool bForceReadOnly = false)
+		: Super(InAttribute, InMetadataDomain, bForceReadOnly)
+	{
+	}
+
+	FPCGAttributeGenericAccessor(const FPCGMetadataAttributeGeneric* InAttribute, const FPCGMetadataDomain* InMetadataDomain, bool bForceReadOnly = false)
+		: Super(InAttribute, InMetadataDomain, bForceReadOnly)
+	{
+	}
+	
+	void GetValuesFromItemKeys(TArrayView<PCGMetadataEntryKey> EntryKeys, TArrayView<T> OutValues) const
+	{
+		static_cast<const FPCGMetadataAttributeGeneric*>(this->Attribute)->GetValuesFromItemKeys(EntryKeys, OutValues);
+	}
+
+	void SetDefaultValue(const T& Value)
+	{
+		static_cast<FPCGMetadataAttributeGeneric*>(this->Attribute)->SetDefaultValue(Value);
+	}
+
+	void SetValues(TArrayView<const PCGMetadataEntryKey * const> ItemKeys, const TArrayView<const T>& InValues)
+	{
+		static_cast<FPCGMetadataAttributeGeneric*>(this->Attribute)->SetValues(ItemKeys, InValues);
+	}
+
+	void SetValues_TryLockless(TArrayView<PCGMetadataEntryKey*> ItemKeys, const TArrayView<const T>& InValues, int32 StartIndex)
+	{
+		static_cast<FPCGMetadataAttributeGeneric*>(this->Attribute)->SetValues_TryLockless(ItemKeys, InValues, StartIndex);
+	}
+
+	int32 PreallocateValues(TArrayView<PCGMetadataEntryKey*> EntryKeys, bool bLockless)
+	{
+		return static_cast<FPCGMetadataAttributeGeneric*>(this->Attribute)->PreallocateValues(EntryKeys, bLockless);
+	}
+
+	void PrepareAttribute(int32 Count)
+	{
+		static_cast<FPCGMetadataAttributeGeneric*>(this->Attribute)->Prepare(Count);
+	}
+};
+
+#undef LOCTEXT_NAMESPACE
\ No newline at end of file
