[PCG] Fix change type for input/output settings to force a recompilation when a label changes.

* Task inputs are using pin labels, without recompilation, those labels are not up to date anymore.

#rnx

#Synced-CL 50489647

diff --git a/Engine/Plugins/PCG/Source/PCG/Private/PCGInputOutputSettings.cpp b/Engine/Plugins/PCG/Source/PCG/Private/PCGInputOutputSettings.cpp
--- a/Engine/Plugins/PCG/Source/PCG/Private/PCGInputOutputSettings.cpp
+++ b/Engine/Plugins/PCG/Source/PCG/Private/PCGInputOutputSettings.cpp
@@ -220,17 +220,27 @@ EPCGChangeType UPCGGraphInputOutputSettings::GetChangeTypeForProperty(FPropertyC
 {
 	EPCGChangeType ChangeType = Super::GetChangeTypeForProperty(PropertyChangedEvent);
 
-	if (bIsInput && PropertyChangedEvent.Property)
+	const FName PropertyName = PropertyChangedEvent.GetPropertyName();
+	const FName MemberPropertyName = PropertyChangedEvent.GetMemberPropertyName();
+
+	if (bIsInput)
 	{
 		// If pins were removed or the required-ness of a pin has changed, this needs to be a structural change.
-		if (PropertyChangedEvent.Property->GetFName() == GET_MEMBER_NAME_CHECKED(UPCGGraphInputOutputSettings, Pins) ||
-			(PropertyChangedEvent.MemberProperty->GetFName() == GET_MEMBER_NAME_CHECKED(UPCGGraphInputOutputSettings, Pins) && 
-			PropertyChangedEvent.Property->GetFName() == GET_MEMBER_NAME_CHECKED(FPCGPinProperties, PinStatus)))
+		if (PropertyName == GET_MEMBER_NAME_CHECKED(UPCGGraphInputOutputSettings, Pins) ||
+			(MemberPropertyName == GET_MEMBER_NAME_CHECKED(UPCGGraphInputOutputSettings, Pins) &&
+			PropertyName == GET_MEMBER_NAME_CHECKED(FPCGPinProperties, PinStatus)))
 		{
 			ChangeType |= EPCGChangeType::Structural;
 		}
 	}
 
+	// If a label changed, we need to recompile, as the tasks inputs are gathered by pin label.
+	if (MemberPropertyName == GET_MEMBER_NAME_CHECKED(UPCGGraphInputOutputSettings, Pins) &&
+		PropertyName == GET_MEMBER_NAME_CHECKED(FPCGPinProperties, Label))
+	{
+		ChangeType |= EPCGChangeType::Structural;
+	}
+
 	return ChangeType;
 }
 #endif // WITH_EDITOR
