PCG Property Bag: Hierarchy Editor Support added.

#Synced-CL 50532063

#Synced-CL 50532063

diff --git a/Engine/Plugins/PCG/Source/PCG/Private/PCGGraph.cpp b/Engine/Plugins/PCG/Source/PCG/Private/PCGGraph.cpp
--- a/Engine/Plugins/PCG/Source/PCG/Private/PCGGraph.cpp
+++ b/Engine/Plugins/PCG/Source/PCG/Private/PCGGraph.cpp
@@ -25,6 +25,8 @@
 #include "AssetRegistry/AssetData.h"
 #include "Logging/StructuredLog.h"
 #include "Serialization/ArchiveReplaceObjectRef.h"
+#include "UObject/FortniteMainBranchObjectVersion.h"
+#include "UObject/ObjectEditorOptionalSupport.h"
 
 #if WITH_EDITOR
 #include "CoreGlobals.h"
@@ -184,6 +186,11 @@ bool UPCGGraphInterface::IsStandaloneGraph() const
 	return GetGraph() && GetGraph()->bIsStandaloneGraph;
 }
 
+UObject* UPCGGraphInterface::GetUserParameterHierarchyRoot() const
+{
+	return nullptr;
+}
+
 bool UPCGGraphInterface::IsStandaloneGraphAsset(const FAssetData& InAssetData)
 {
 	if (UPCGGraphInterface* GraphInterface = Cast<UPCGGraphInterface>(InAssetData.FastGetAsset(/*bLoad=*/false)))
@@ -514,6 +521,13 @@ UPCGGraph::UPCGGraph(const FObjectInitializer& ObjectInitializer)
 	// Force the user parameters to have an empty property bag. It is necessary to catch the first
 	// add property into the undo/redo history.
 	UserParameters.MigrateToNewBagStruct(UPropertyBag::GetOrCreateFromDescs({}));
+
+#if WITH_EDITOR
+	if (UClass* UserParametersHierarchyRootClass = StaticLoadClass(UObject::StaticClass(), nullptr, TEXT("/Script/PCGEditor.PCGPropertyBagHierarchyRoot")))
+	{
+		UserParameterHierarchyRoot = ObjectInitializer.CreateEditorOnlyDefaultSubobject(this, TEXT("UserParametersHierarchyRoot"), UserParametersHierarchyRootClass, false);
+	}
+#endif
 }
 
 void UPCGGraph::PostLoad()
@@ -622,6 +636,15 @@ void UPCGGraph::PostLoad()
 	}
 
 	GenerationRadii.ApplyDeprecation();
+	
+	int32 CustomVersion = GetLinkerCustomVersion(FFortniteMainBranchObjectVersion::GUID);
+	if (CustomVersion < FFortniteMainBranchObjectVersion::PCG_PropertyBagHierarchySupport)
+	{
+		if (IPCGEditorModule* PCGEditorModule = IPCGEditorModule::Get())
+		{
+			PCGEditorModule->TransferPropertyBagMetadataIntoHierarchy(this);
+		}
+	}
 #endif // WITH_EDITOR
 }
 
@@ -661,6 +684,15 @@ bool UPCGGraph::IsEditorOnly() const
 	return bIsCurrentlyEditorOnly;
 }
 
+void UPCGGraph::Serialize(FArchive& Ar)
+{
+	LLM_SCOPE_BYTAG(PCG);
+	
+	Super::Serialize(Ar);
+	
+	Ar.UsingCustomVersion(FFortniteMainBranchObjectVersion::GUID);
+}
+
 bool UPCGGraph::IsEditorOnly_Internal() const
 {
 	return Super::IsEditorOnly() || bIsEditorOnly;
@@ -2889,6 +2921,11 @@ TOptional<FLinearColor> UPCGGraphInstance::GetColorOverride() const
 {
 	return (!bOverrideColor && Graph) ? Graph->GetColorOverride() : UPCGGraphInterface::GetColorOverride();
 }
+
+UObject* UPCGGraphInstance::GetUserParameterHierarchyRoot() const
+{
+	return Graph ? Graph->GetUserParameterHierarchyRoot() : nullptr;
+}
 #endif // WITH_EDITOR
 
 void UPCGGraphInstance::OnGraphParametersChanged(UPCGGraphInterface* InGraph, EPCGGraphParameterEvent InChangeType, FName InChangedPropertyName)
diff --git a/Engine/Plugins/PCG/Source/PCG/Public/Editor/IPCGEditorModule.h b/Engine/Plugins/PCG/Source/PCG/Public/Editor/IPCGEditorModule.h
--- a/Engine/Plugins/PCG/Source/PCG/Public/Editor/IPCGEditorModule.h
+++ b/Engine/Plugins/PCG/Source/PCG/Public/Editor/IPCGEditorModule.h
@@ -9,6 +9,7 @@
 #include "Internationalization/Text.h"
 #include "Modules/ModuleInterface.h"
 
+class UPropertyBagHierarchyRoot;
 class IPCGBaseSubsystem;
 struct FPCGStackContext;
 class APCGWorldActor;
@@ -36,6 +37,7 @@ public:
 	virtual bool CanSelectPartitionActors() const = 0;
 	virtual void SetPermissionMode(EPCGEditorPermissionMode InPermissionMode) = 0;
 	virtual EPCGEditorPermissionMode GetPermissionMode() const = 0;
+	virtual void TransferPropertyBagMetadataIntoHierarchy(UPCGGraph* Graph) const = 0;
 
 	DECLARE_EVENT_OneParam(IPCGEditorModule, FPermissionModeChanged, EPCGEditorPermissionMode);
 	virtual FPermissionModeChanged& OnPermissionModeChanged() = 0;
diff --git a/Engine/Plugins/PCG/Source/PCG/Public/PCGGraph.h b/Engine/Plugins/PCG/Source/PCG/Public/PCGGraph.h
--- a/Engine/Plugins/PCG/Source/PCG/Public/PCGGraph.h
+++ b/Engine/Plugins/PCG/Source/PCG/Public/PCGGraph.h
@@ -187,6 +187,8 @@ public:
 	UFUNCTION()
 	UE_API bool IsStandaloneGraph() const;
 
+	UFUNCTION()
+	UE_API virtual UObject* GetUserParameterHierarchyRoot() const; 
 protected:
 	/** By default export to library is visible only for graph that are assets, but can be enabled/disabled if needed. */
 	UFUNCTION()
@@ -325,6 +327,7 @@ public:
 	/** ~Begin UObject interface */
 	UE_API virtual void PostLoad() override;
 	UE_API virtual bool IsEditorOnly() const override;
+	UE_API virtual void Serialize(FArchive& Ar) override;
 
 #if WITH_EDITOR
 	UE_API virtual void PreSave(FObjectPreSaveContext ObjectSaveContext) override;
@@ -636,6 +639,11 @@ protected:
 
 	// Editor graph created from PCG Editor but owned by this, reference is collected using AddReferencedObjects
 	TObjectPtr<UPCGEditorGraph> PCGEditorGraph = nullptr;
+	
+	UPROPERTY(meta=(PCGNoHash))
+	TObjectPtr<UObject> UserParameterHierarchyRoot;
+	
+	virtual UObject* GetUserParameterHierarchyRoot() const override { return UserParameterHierarchyRoot; }
 #endif // WITH_EDITORONLY_DATA
 
 	// Parameters
@@ -835,6 +843,8 @@ public:
 
 	UFUNCTION()
 	bool ShowToolDataOverrides() const { return IsAsset(); }
+	
+	UE_API virtual UObject* GetUserParameterHierarchyRoot() const override;
 #endif // WITH_EDITOR
 
 	/** Because subgraph nodes hold graph instances, delegation of the selection can't be done elsewhere, so this function gets called when the dialog is opened and will forward to the graph customization if filtering is applied. */
@@ -844,7 +854,7 @@ public:
 	UPROPERTY(EditAnywhere, BlueprintReadOnly, Category = Instance, AssetRegistrySearchable, meta = (GetAssetFilter = "GraphAssetFilter"))
 	TObjectPtr<UPCGGraphInterface> Graph;
 
-	UPROPERTY(EditAnywhere, Category = Instance, meta = (NoResetToDefault))
+	UPROPERTY(EditAnywhere, Category = Instance, meta = (NoResetToDefault, PropertyBagSchemaClass="/Script/PCGEditor.PCGGraphParametersSchema"))
 	FPCGOverrideInstancedPropertyBag ParametersOverrides;
 
 	virtual const FInstancedPropertyBag* GetUserParametersStruct() const override { return &ParametersOverrides.Parameters; }
diff --git a/Engine/Plugins/PCG/Source/PCGEditor/PCGEditor.Build.cs b/Engine/Plugins/PCG/Source/PCGEditor/PCGEditor.Build.cs
--- a/Engine/Plugins/PCG/Source/PCGEditor/PCGEditor.Build.cs
+++ b/Engine/Plugins/PCG/Source/PCGEditor/PCGEditor.Build.cs
@@ -16,7 +16,9 @@ namespace UnrealBuildTool.Rules
 					"Engine",
 					"CoreUObject",
 					"PlacementMode", 
-					"PCG"
+					"PCG",
+					"StructUtilsEditor",
+					"DataHierarchyEditor"
 				});
 
 			if (Target.WithAutomationTests)
@@ -62,8 +64,7 @@ namespace UnrealBuildTool.Rules
 					"Slate",
 					"SlateCore",
 					"SourceControl",
-					"StatusBar",
-					"StructUtilsEditor",
+					"StatusBar", 
 					"SubobjectEditor",
 					"SubobjectDataInterface",
 					"ToolMenus",
diff --git a/Engine/Plugins/PCG/Source/PCGEditor/Private/Details/PCGInstancedPropertyBagOverrideDetails.cpp b/Engine/Plugins/PCG/Source/PCGEditor/Private/Details/PCGInstancedPropertyBagOverrideDetails.cpp
--- a/Engine/Plugins/PCG/Source/PCGEditor/Private/Details/PCGInstancedPropertyBagOverrideDetails.cpp
+++ b/Engine/Plugins/PCG/Source/PCGEditor/Private/Details/PCGInstancedPropertyBagOverrideDetails.cpp
@@ -7,8 +7,10 @@
 
 #include "DetailLayoutBuilder.h"
 #include "IDetailChildrenBuilder.h"
+#include "PropertyBagDetails.h"
 #include "PropertyCustomizationHelpers.h"
 #include "ScopedTransaction.h"
+#include "StructUtilsMetadata.h"
 #include "Widgets/SBoxPanel.h"
 #include "Widgets/Layout/SBox.h"
 #include "Widgets/Images/SImage.h"
@@ -66,12 +68,24 @@ void FPCGOverrideInstancedPropertyBagDetails::CustomizeHeader(TSharedRef<IProper
 void FPCGOverrideInstancedPropertyBagDetails::CustomizeChildren(TSharedRef<IPropertyHandle> InPropertyHandle, IDetailChildrenBuilder& InChildrenBuilder, IPropertyTypeCustomizationUtils& InCustomizationUtils)
 {
 	TSharedPtr<IPropertyHandle> InstancedPropertyBagHandle = PropertyHandle.IsValid() ? PropertyHandle->GetChildHandle(TEXT("Parameters")) : nullptr;
-	const TSharedRef<FPCGOverrideInstancedPropertyBagDataDetails> InstanceDetails = MakeShareable(new FPCGOverrideInstancedPropertyBagDataDetails(InstancedPropertyBagHandle, InCustomizationUtils.GetPropertyUtilities()));
+	
+	// We forward the schema metadata
+	if (PropertyHandle.IsValid() && PropertyHandle->HasMetaData(UE::StructUtils::Metadata::PropertyBagSchemaClassName) && InstancedPropertyBagHandle.IsValid())
+	{
+		InstancedPropertyBagHandle->SetInstanceMetaData(UE::StructUtils::Metadata::PropertyBagSchemaClassName, *PropertyHandle->GetMetaData(UE::StructUtils::Metadata::PropertyBagSchemaClassName));
+	}
+	
+	const TSharedRef<FPCGOverrideInstancedPropertyBagDataDetails> InstanceDetails = MakeShared<FPCGOverrideInstancedPropertyBagDataDetails>(InstancedPropertyBagHandle, InCustomizationUtils.GetPropertyUtilities());
 	InChildrenBuilder.AddCustomBuilder(InstanceDetails);
 }
 
 FPCGOverrideInstancedPropertyBagDataDetails::FPCGOverrideInstancedPropertyBagDataDetails(TSharedPtr<IPropertyHandle> InStructProperty, const TSharedPtr<IPropertyUtilities>& InPropUtils)
-	: FInstancedStructDataDetails(InStructProperty.IsValid() ? InStructProperty->GetChildHandle(TEXT("Value")) : nullptr)
+	: FPropertyBagInstanceDataDetails(
+		FConstructParams{
+			.BagStructProperty = InStructProperty.IsValid() ? InStructProperty : nullptr,
+			.PropUtils = InPropUtils,
+			.ChildRowFeatures = EPropertyBagChildRowFeatures::Fixed
+		})
 {
 	if (InStructProperty.IsValid())
 	{
diff --git a/Engine/Plugins/PCG/Source/PCGEditor/Private/Details/PCGInstancedPropertyBagOverrideDetails.h b/Engine/Plugins/PCG/Source/PCGEditor/Private/Details/PCGInstancedPropertyBagOverrideDetails.h
--- a/Engine/Plugins/PCG/Source/PCGEditor/Private/Details/PCGInstancedPropertyBagOverrideDetails.h
+++ b/Engine/Plugins/PCG/Source/PCGEditor/Private/Details/PCGInstancedPropertyBagOverrideDetails.h
@@ -7,6 +7,7 @@
 #include "IDetailCustomization.h"
 #include "IPropertyTypeCustomization.h"
 #include "InstancedStructDetails.h"
+#include "PropertyBagDetails.h"
 #include "Input/Reply.h"
 #include "Layout/Visibility.h"
 #include "Templates/SharedPointer.h"
@@ -29,7 +30,7 @@ private:
 	TSharedPtr<IPropertyHandle> PropertyHandle;
 };
 
-class FPCGOverrideInstancedPropertyBagDataDetails : public FInstancedStructDataDetails
+class FPCGOverrideInstancedPropertyBagDataDetails : public FPropertyBagInstanceDataDetails
 {
 public:
 	FPCGOverrideInstancedPropertyBagDataDetails(TSharedPtr<IPropertyHandle> InStructProperty, const TSharedPtr<IPropertyUtilities>& InPropUtils);
diff --git a/Engine/Plugins/PCG/Source/PCGEditor/Private/PCGEditorModule.cpp b/Engine/Plugins/PCG/Source/PCGEditor/Private/PCGEditorModule.cpp
--- a/Engine/Plugins/PCG/Source/PCGEditor/Private/PCGEditorModule.cpp
+++ b/Engine/Plugins/PCG/Source/PCGEditor/Private/PCGEditorModule.cpp
@@ -91,6 +91,8 @@
 #include "EditorMode/Tools/PCGInteractiveToolSettings.h"
 #include "Framework/MultiBox/MultiBoxBuilder.h"
 #include "Framework/Notifications/NotificationManager.h"
+#include "HierarchyEditor/PropertyBagHierarchyViewModel.h"
+#include "Schema/PCGGraphParametersSchema.h"
 #include "UObject/ObjectSaveContext.h"
 #include "Widgets/SPCGNodeSourceTextBox.h"
 
@@ -410,6 +412,16 @@ void FPCGEditorModule::SetPermissionMode(EPCGEditorPermissionMode InPermissionMo
 	PermissionModeChangedEvent.Broadcast(PermissionMode);
 }
 
+void FPCGEditorModule::TransferPropertyBagMetadataIntoHierarchy(UPCGGraph* Graph) const
+{
+	if (!Graph || !Graph->UserParameterHierarchyRoot || !Graph->UserParameterHierarchyRoot->IsA<UPropertyBagHierarchyRoot>())
+	{
+		return;
+	}
+	
+	GetDefault<UPCGGraphParametersSchema>()->TransferPropertyBagMetadataIntoHierarchy(Graph->UserParameters, *Cast<UPropertyBagHierarchyRoot>(Graph->UserParameterHierarchyRoot));
+}
+
 void FPCGEditorModule::RegisterDetailsCustomizations()
 {
 	FPropertyEditorModule& PropertyEditor = FModuleManager::LoadModuleChecked<FPropertyEditorModule>("PropertyEditor");
diff --git a/Engine/Plugins/PCG/Source/PCGEditor/Private/Schema/PCGGraphParametersSchema.cpp b/Engine/Plugins/PCG/Source/PCGEditor/Private/Schema/PCGGraphParametersSchema.cpp
--- a/Engine/Plugins/PCG/Source/PCGEditor/Private/Schema/PCGGraphParametersSchema.cpp
+++ b/Engine/Plugins/PCG/Source/PCGEditor/Private/Schema/PCGGraphParametersSchema.cpp
@@ -2,7 +2,41 @@
 
 #include "Schema/PCGGraphParametersSchema.h"
 
+#include "PCGGraph.h"
+#include "HierarchyEditor/PropertyBagHierarchyViewModel.h"
+#include "StructUtilsMetadata.h"
+
 bool UPCGGraphParametersSchema::SupportsPinTypeContainer(TWeakPtr<const FEdGraphSchemaAction> SchemaAction, const FEdGraphPinType& PinType, const EPinContainerType& ContainerType) const
 {
 	return ContainerType == EPinContainerType::None || ContainerType == EPinContainerType::Array || ContainerType == EPinContainerType::Set;
 }
+
+UPropertyBagHierarchyRoot* UPCGGraphParametersSchema::GetHierarchyRoot(const TArray<UObject*>& ObjectsWithProperty) const
+{
+	if (ObjectsWithProperty.Num() != 1)
+	{
+		return nullptr;
+	}
+	
+	if (UPCGGraphInterface* PCGGraphInterface = Cast<UPCGGraphInterface>(ObjectsWithProperty[0]))
+	{
+		return Cast<UPropertyBagHierarchyRoot>(PCGGraphInterface->GetUserParameterHierarchyRoot());
+	}
+	
+	return nullptr;
+}
+
+TArray<UScriptStruct*> UPCGGraphParametersSchema::GetHierarchyPropertyMetaDataTypes(const FPropertyBagPropertyDesc& Desc) const
+{
+	return Super::GetHierarchyPropertyMetaDataTypes(Desc);
+}
+
+TSubclassOf<UPropertyBagHierarchyCategory> UPCGGraphParametersSchema::GetHierarchyCategoryType() const
+{
+	return UPCGPropertyBagHierarchyCategory::StaticClass();
+}
+
+TSubclassOf<UPropertyBagHierarchySection> UPCGGraphParametersSchema::GetHierarchySectionType() const
+{
+	return UPCGPropertyBagHierarchySection::StaticClass();
+}
diff --git a/Engine/Plugins/PCG/Source/PCGEditor/Public/PCGEditorModule.h b/Engine/Plugins/PCG/Source/PCGEditor/Public/PCGEditorModule.h
--- a/Engine/Plugins/PCG/Source/PCGEditor/Public/PCGEditorModule.h
+++ b/Engine/Plugins/PCG/Source/PCGEditor/Public/PCGEditorModule.h
@@ -40,6 +40,7 @@ public:
 	virtual bool CanSelectPartitionActors() const override;
 	virtual void SetPermissionMode(EPCGEditorPermissionMode InPermissionMode) override;
 	virtual EPCGEditorPermissionMode GetPermissionMode() const { return PermissionMode; }
+	virtual void TransferPropertyBagMetadataIntoHierarchy(UPCGGraph* Graph) const override;
 	
 	DECLARE_DERIVED_EVENT(FPCGEditorModule, IPCGEditorModule::FPermissionModeChanged, FPermissionModeChanged);
 	virtual FPermissionModeChanged& OnPermissionModeChanged() override { return PermissionModeChangedEvent; }
diff --git a/Engine/Plugins/PCG/Source/PCGEditor/Public/Schema/PCGGraphParametersSchema.h b/Engine/Plugins/PCG/Source/PCGEditor/Public/Schema/PCGGraphParametersSchema.h
--- a/Engine/Plugins/PCG/Source/PCGEditor/Public/Schema/PCGGraphParametersSchema.h
+++ b/Engine/Plugins/PCG/Source/PCGEditor/Public/Schema/PCGGraphParametersSchema.h
@@ -3,9 +3,34 @@
 #pragma once
 
 #include "PropertyBagDetails.h"
+#include "HierarchyEditor/PropertyBagHierarchyViewModel.h"
 
 #include "PCGGraphParametersSchema.generated.h"
 
+UCLASS(MinimalAPI)
+class UPCGPropertyBagHierarchyRoot : public UPropertyBagHierarchyRoot
+{
+	GENERATED_BODY()
+	
+	// Add future properties for PCG root here
+};
+
+UCLASS(MinimalAPI)
+class UPCGPropertyBagHierarchyCategory : public UPropertyBagHierarchyCategory
+{
+	GENERATED_BODY()
+	
+	// Add future properties for PCG categories here
+};
+
+UCLASS(MinimalAPI)
+class UPCGPropertyBagHierarchySection : public UPropertyBagHierarchySection
+{
+	GENERATED_BODY()
+	
+	// Add future properties for PCG sections here
+};
+
 /**
  * Specific property bag schema to allow only arrays and set as containers.
  */
@@ -14,5 +39,13 @@ class UPCGGraphParametersSchema : public UPropertyBagSchema
 {
 	GENERATED_BODY()
 public:
-	PCGEDITOR_API virtual bool SupportsPinTypeContainer(TWeakPtr<const FEdGraphSchemaAction> SchemaAction, const FEdGraphPinType& PinType, const EPinContainerType& ContainerType) const override;
+	virtual bool SupportsPinTypeContainer(TWeakPtr<const FEdGraphSchemaAction> SchemaAction, const FEdGraphPinType& PinType, const EPinContainerType& ContainerType) const override;
+	
+	virtual UPropertyBagHierarchyRoot* GetHierarchyRoot(const TArray<UObject*>& ObjectsWithProperty) const override;
+	
+	virtual TArray<UScriptStruct*> GetHierarchyPropertyMetaDataTypes(const FPropertyBagPropertyDesc& Desc) const override;
+	
+	virtual TSubclassOf<UPropertyBagHierarchyCategory> GetHierarchyCategoryType() const override;
+	
+	virtual TSubclassOf<UPropertyBagHierarchySection> GetHierarchySectionType() const override;
 };
\ No newline at end of file
diff --git a/Engine/Source/Runtime/Core/Public/UObject/FortniteMainBranchObjectVersions.inl b/Engine/Source/Runtime/Core/Public/UObject/FortniteMainBranchObjectVersions.inl
--- a/Engine/Source/Runtime/Core/Public/UObject/FortniteMainBranchObjectVersions.inl
+++ b/Engine/Source/Runtime/Core/Public/UObject/FortniteMainBranchObjectVersions.inl
@@ -763,5 +763,8 @@ FORTNITE_MAIN_VERSION(RigVMPinDefaultValueObjectPtr, 249)
 // AnimSequenceTracks now have timestamps
 FORTNITE_MAIN_VERSION(AnimSequenceTrackTimestamps, 250)
 
+// PCG Graph's Instanced Property Bags now support Hierarchy Editor, post loading existing data into new format
+FORTNITE_MAIN_VERSION(PCG_PropertyBagHierarchySupport, 251)
+
 // -----<new versions can be added above this line>-------------------------------------------------
 #undef FORTNITE_MAIN_VERSION
