[PCG] Hi-gen grid scaling

Allows fitting a generation grid to other grids in the world.

Changes:
* Mass move from EPCGHiGenGrid to uint32 grid size (but enum still used for authoring)
* Expose a grid scale multiplier instead of the current exponent

[RN] [PCG] Generalize grid scale exponent to be any multiplier. Allows scaling of generation grid to conform to other grids in level such as landscape components.

#Synced-CL 50496977

#Synced-CL 50546465

diff --git a/Engine/Plugins/PCG/Source/PCG/Internal/RuntimeGen/PCGRuntimeGenExecutionSource.h b/Engine/Plugins/PCG/Source/PCG/Internal/RuntimeGen/PCGRuntimeGenExecutionSource.h
--- a/Engine/Plugins/PCG/Source/PCG/Internal/RuntimeGen/PCGRuntimeGenExecutionSource.h
+++ b/Engine/Plugins/PCG/Source/PCG/Internal/RuntimeGen/PCGRuntimeGenExecutionSource.h
@@ -85,14 +85,13 @@ public:
 	bool IsPartitioned() const;
 	bool IsLocalSource() const;
 	bool Use2DGrid() const;
-	EPCGHiGenGrid GetGenerationGrid() const;
 	uint32 GetGenerationGridSize() const;
 	FIntVector GetGenerationGridCoords() const;
 	bool IsManagedByRuntimeGenSystem() const;
 	const UPCGSchedulingPolicyBase* GetRuntimeGenSchedulingPolicy() const;
 	const FPCGRuntimeGenerationRadii& GetGenerationRadii() const;
-	double GetGenerationRadiusFromGrid(EPCGHiGenGrid InGrid) const;
-	double GetCleanupRadiusFromGrid(EPCGHiGenGrid InGrid) const;
+	double GetGenerationRadiusFromGrid(uint32 InGrid) const;
+	double GetCleanupRadiusFromGrid(uint32 InGrid) const;
 	/** End mirror of FPCGRuntimeGenExecutionState */
 
 	/** Creates a local execution source that is automatically setup and registered with this original execution source. Nullptr if this is not a partitioned original source. */
diff --git a/Engine/Plugins/PCG/Source/PCG/Internal/RuntimeGen/PCGRuntimeGenExecutionState.h b/Engine/Plugins/PCG/Source/PCG/Internal/RuntimeGen/PCGRuntimeGenExecutionState.h
--- a/Engine/Plugins/PCG/Source/PCG/Internal/RuntimeGen/PCGRuntimeGenExecutionState.h
+++ b/Engine/Plugins/PCG/Source/PCG/Internal/RuntimeGen/PCGRuntimeGenExecutionState.h
@@ -59,14 +59,13 @@ public:
 	virtual bool IsPartitioned() const override;
 	virtual bool IsLocalSource() const override;
 	virtual bool Use2DGrid() const override;
-	virtual EPCGHiGenGrid GetGenerationGrid() const override;
 	virtual uint32 GetGenerationGridSize() const override;
 	virtual FIntVector GetGenerationGridCoords() const override;
 	virtual bool IsManagedByRuntimeGenSystem() const override;
 	virtual const UPCGSchedulingPolicyBase* GetRuntimeGenSchedulingPolicy() const override;
 	virtual const FPCGRuntimeGenerationRadii& GetGenerationRadii() const override;
-	virtual double GetGenerationRadiusFromGrid(EPCGHiGenGrid InGrid) const override;
-	virtual double GetCleanupRadiusFromGrid(EPCGHiGenGrid InGrid) const override;
+	virtual double GetGenerationRadiusFromGrid(uint32 InGrid) const override;
+	virtual double GetCleanupRadiusFromGrid(uint32 InGrid) const override;
 	/** ~End IPCGGraphExecutionState interface */
 
 	void SetSource(UPCGRuntimeGenExecutionSource* InSource);
diff --git a/Engine/Plugins/PCG/Source/PCG/Internal/Utils/PCGGraphExecutionLogging.h b/Engine/Plugins/PCG/Source/PCG/Internal/Utils/PCGGraphExecutionLogging.h
--- a/Engine/Plugins/PCG/Source/PCG/Internal/Utils/PCGGraphExecutionLogging.h
+++ b/Engine/Plugins/PCG/Source/PCG/Internal/Utils/PCGGraphExecutionLogging.h
@@ -47,8 +47,8 @@ namespace PCGGraphExecutionLogging
 	void LogTaskCullingBeginLoop(FPCGTaskId PinTaskId, uint64 PinIndex, const TArray<FPCGPinId>& PinIdsToDeactivate);
 	void LogTaskCullingUpdatedPinDeps(FPCGTaskId TaskId, const FPCGPinDependencyExpression& PinDependency, bool bDependencyExpressionBecameFalse);
 
-	void LogGridLinkageTaskExecuteStore(const FPCGContext* InContext, EPCGHiGenGrid InGenerationGrid, int32 InFromLOD, int32 InToLOD, const FString& InResourcePath, int32 InDataItemCount);
-	void LogGridLinkageTaskExecuteRetrieve(const FPCGContext* InContext, EPCGHiGenGrid InGenerationGrid, int32 InFromLOD, int32 InToLOD, const FString& InResourcePath);
+	void LogGridLinkageTaskExecuteStore(const FPCGContext* InContext, uint32 InGenerationGrid, uint32 InFromGridSize, uint32 InToGridSize, const FString& InResourcePath, int32 InDataItemCount);
+	void LogGridLinkageTaskExecuteRetrieve(const FPCGContext* InContext, uint32 InGenerationGrid, uint32 InFromGridSize, uint32 InToGridSize, const FString& InResourcePath);
 	void LogGridLinkageTaskExecuteRetrieveSuccess(const FPCGContext* InContext, const IPCGGraphExecutionSource* InExecutionSource, const FString& InResourcePath, int32 InDataItemCount);
 	void LogGridLinkageTaskExecuteRetrieveScheduleGraph(const FPCGContext* InContext, const IPCGGraphExecutionSource* InScheduledSource, const FString& InResourcePath);
 	void LogGridLinkageTaskExecuteRetrieveWaitOnScheduledGraph(const FPCGContext* InContext, const IPCGGraphExecutionSource* InWaitOnSource, const FString& InResourcePath);
diff --git a/Engine/Plugins/PCG/Source/PCG/Private/Elements/PCGDataFromActor.cpp b/Engine/Plugins/PCG/Source/PCG/Private/Elements/PCGDataFromActor.cpp
--- a/Engine/Plugins/PCG/Source/PCG/Private/Elements/PCGDataFromActor.cpp
+++ b/Engine/Plugins/PCG/Source/PCG/Private/Elements/PCGDataFromActor.cpp
@@ -91,7 +91,8 @@ namespace PCGDataFromActorHelpers
 		{
 			auto AddComponent = [&LocalComponents, bGetAllGrids, AllowedGrids](UPCGComponent* LocalComponent)
 			{
-				if (bGetAllGrids || (AllowedGrids & (int32)LocalComponent->GetGenerationGrid()))
+				const EPCGHiGenGrid UnscaledGrid = LocalComponent->GetUnscaledGenerationGrid();
+				if (bGetAllGrids || (AllowedGrids & (int32)UnscaledGrid))
 				{
 					LocalComponents.Add(LocalComponent);
 				}
@@ -1015,7 +1016,9 @@ void FPCGDataFromActorElement::ProcessActor(FPCGContext* InContext, const UPCGDa
 					}
 #endif
 
-					DuplicatedTaggedData.Tags.Add(PCGDataFromActorConstants::PCGComponentDataGridSizeTagPrefix + FString::FromInt(PCGHiGenGrid::GridToGridSize(Component->GetGenerationGrid())));
+					// Use stable/canonical grid size for tags.
+					const EPCGHiGenGrid UnscaledGrid = Component->GetUnscaledGenerationGrid();
+					DuplicatedTaggedData.Tags.Add(PCGDataFromActorConstants::PCGComponentDataGridSizeTagPrefix + FString::FromInt(PCGHiGenGrid::GridToGridSize(UnscaledGrid)));
 
 					// Add some logging if we have un-expected output pin
 					if (DuplicatedTaggedData.Pin != PCGPinConstants::DefaultOutputLabel && !Settings->ExpectedPins.IsEmpty() && !Settings->ExpectedPins.Contains(TaggedData.Pin))
diff --git a/Engine/Plugins/PCG/Source/PCG/Private/Elements/PCGHiGenGridSize.cpp b/Engine/Plugins/PCG/Source/PCG/Private/Elements/PCGHiGenGridSize.cpp
--- a/Engine/Plugins/PCG/Source/PCG/Private/Elements/PCGHiGenGridSize.cpp
+++ b/Engine/Plugins/PCG/Source/PCG/Private/Elements/PCGHiGenGridSize.cpp
@@ -27,22 +27,21 @@ namespace PCGHiGenGridSizeConstants
 	const int32 ChangeGridSizePreconfiguredIndex = 1;
 }
 
-EPCGHiGenGrid UPCGHiGenGridSizeSettings::GetGrid() const
+uint32 UPCGHiGenGridSizeSettings::GetScaledGridSize() const
 {
-	if (GetOuter() && GetOuter()->GetOuter() && GetOuter()->GetOuter()->IsA<UPCGGraph>())
+	if (HiGenGridSize == EPCGHiGenGrid::Unbounded)
 	{
-		UPCGGraph* Graph = Cast<UPCGGraph>(GetOuter()->GetOuter());
-		return static_cast<EPCGHiGenGrid>(static_cast<uint32>(HiGenGridSize) << Graph->GetGridExponential());
+		return PCGHiGenGrid::UnboundedGridSize();
 	}
-	else
+
+	uint32 GridSize = PCGHiGenGrid::GridToGridSize(HiGenGridSize);
+
+	if (UPCGGraph* Graph = GetOuter() ? GetOuter()->GetTypedOuter<UPCGGraph>() : nullptr)
 	{
-		return HiGenGridSize;
+		GridSize = FMath::RoundToInt(GridSize * Graph->GetGridSizeMultiplier());
 	}
-}
 
-uint32 UPCGHiGenGridSizeSettings::GetGridSize() const
-{
-	return (HiGenGridSize == EPCGHiGenGrid::Unbounded) ? PCGHiGenGrid::UnboundedGridSize() : PCGHiGenGrid::GridToGridSize(GetGrid());
+	return GridSize;
 }
 
 bool UPCGHiGenGridSizeSettings::ShouldOutputCellVolumeData() const
@@ -156,7 +155,17 @@ FString UPCGHiGenGridSizeSettings::GetAdditionalTitleInformation() const
 	const UEnum* EnumPtr = StaticEnum<EPCGHiGenGrid>();
 	if (ensure(EnumPtr))
 	{
-		return EnumPtr->GetDisplayNameTextByValue(static_cast<int64>(HiGenGridSize)).ToString();
+		FString AdditionalTitle = EnumPtr->GetDisplayNameTextByValue(static_cast<int64>(HiGenGridSize)).ToString();
+
+		const UPCGNode* Node = GetTypedOuter<UPCGNode>();
+		const UPCGGraph* Graph = Node ? Node->GetGraph() : nullptr;
+
+		if (Graph && Graph->GetGridSizeMultiplier() != 1.0)
+		{
+			AdditionalTitle = FString::Format(TEXT("{0} (Scaled: {1})"), { AdditionalTitle, GetScaledGridSize()});
+		}
+
+		return AdditionalTitle;
 	}
 	else
 	{
@@ -247,7 +256,7 @@ bool FPCGHiGenGridSizeElement::ExecuteInternal(FPCGContext* Context) const
 	check(Settings);
 
 	const UPCGGraph* Graph = Context->Node ? Context->Node->GetGraph() : nullptr;
-	if (Graph && Graph->IsHierarchicalGenerationEnabled() && (Graph->GetDefaultGrid() < Settings->GetGrid()))
+	if (Graph && Graph->IsHierarchicalGenerationEnabled() && (Graph->GetDefaultGridSize() < Settings->GetScaledGridSize()))
 	{
 		PCGE_LOG(Warning, GraphAndLog, LOCTEXT("GridSizeLargerThanGraphGridSize", "Grid size is larger than graph default grid size and will be automatically clamped."));
 	}
diff --git a/Engine/Plugins/PCG/Source/PCG/Private/Graph/PCGGraphCompiler.cpp b/Engine/Plugins/PCG/Source/PCG/Private/Graph/PCGGraphCompiler.cpp
--- a/Engine/Plugins/PCG/Source/PCG/Private/Graph/PCGGraphCompiler.cpp
+++ b/Engine/Plugins/PCG/Source/PCG/Private/Graph/PCGGraphCompiler.cpp
@@ -480,11 +480,11 @@ UPCGComputeGraph* FPCGGraphCompiler::GetComputeGraph(const UPCGGraph* InGraph, u
 }
 
 void FPCGGraphCompiler::ResolveGridSizes(
-	EPCGHiGenGrid GenerationGrid,
+	uint32 GenerationGrid,
 	const TArray<FPCGGraphTask>& CompiledTasks,
 	const FPCGStackContext& StackContext,
-	EPCGHiGenGrid GenerationDefaultGrid,
-	TArray<EPCGHiGenGrid>& InOutTaskGenerationGrid)
+	uint32 GenerationDefaultGrid,
+	TArray<uint32>& InOutTaskGenerationGrid)
 {
 	TRACE_CPUPROFILER_EVENT_SCOPE(FPCGGraphCompiler::ResolveGridSizes);
 
@@ -506,8 +506,8 @@ void FPCGGraphCompiler::ResolveGridSizes(
 
 void FPCGGraphCompiler::CreateGridLinkages(
 	UPCGGraph* InGraph,
-	EPCGHiGenGrid InGenerationGrid,
-	TArray<EPCGHiGenGrid>& InOutTaskGenerationGrid,
+	uint32 InGenerationGrid,
+	TArray<uint32>& InOutTaskGenerationGrid,
 	TArray<FPCGGraphTask>& InOutCompiledTasks,
 	const FPCGStackContext& InStackContext,
 	bool bIsCooking)
@@ -527,7 +527,7 @@ void FPCGGraphCompiler::CreateGridLinkages(
 	const int32 NumCompiledTasksBefore = InOutCompiledTasks.Num();
 	for (FPCGTaskId TaskId = 0; TaskId < NumCompiledTasksBefore; ++TaskId)
 	{
-		const EPCGHiGenGrid GraphGenerationGrid = InOutTaskGenerationGrid[TaskId];
+		const uint32 GraphGenerationGrid = InOutTaskGenerationGrid[TaskId];
 
 		for (int32 InputIndex = 0; InputIndex < InOutCompiledTasks[TaskId].Inputs.Num(); ++InputIndex)
 		{
@@ -548,10 +548,10 @@ void FPCGGraphCompiler::CreateGridLinkages(
 				continue;
 			}
 
-			const EPCGHiGenGrid InputGraphGenerationGrid = InOutTaskGenerationGrid[TaskInput.TaskId];
+			const uint32 InputGraphGenerationGrid = InOutTaskGenerationGrid[TaskInput.TaskId];
 			// Register linkage task if grid sizes don't match - either way! This allows us to generate an execution-time error if going from
 			// small grid to large grid.
-			if (InputGraphGenerationGrid != EPCGHiGenGrid::Uninitialized && InputGraphGenerationGrid > GraphGenerationGrid)
+			if (InputGraphGenerationGrid != PCGHiGenGrid::UninitializedGridSize() && InputGraphGenerationGrid > GraphGenerationGrid)
 			{
 				// Build a string identifier for the data
 				FString ResourceKey;
@@ -579,8 +579,8 @@ void FPCGGraphCompiler::CreateGridLinkages(
 
 				LinkTask.Inputs.Emplace(TaskInput.TaskId, UpstreamPinProps, DownstreamPinProps);
 
-				const EPCGHiGenGrid FromGrid = InOutTaskGenerationGrid[TaskInput.TaskId];
-				const EPCGHiGenGrid ToGrid = InOutTaskGenerationGrid[TaskId];
+				const uint32 FromGrid = InOutTaskGenerationGrid[TaskInput.TaskId];
+				const uint32 ToGrid = InOutTaskGenerationGrid[TaskId];
 
 				// This lambda runs at execution time and attempts to retrieve the data from a larger grid. Capture by value is intentional.
 				auto GridLinkageOperation = [FromGrid, ToGrid, ResourceKey, InGenerationGrid, UpstreamPinLabel=UpstreamPinProps.Label](FPCGContext* InContext)
@@ -624,21 +624,22 @@ void FPCGGraphCompiler::CreateGridLinkages(
 				// Set the task input's UpstreamPin so that it has the correct type.
 				InOutCompiledTasks[TaskId].Inputs[InputIndex].UpstreamPin = UpstreamPinProps;
 
-				// The link needs to execute at both FROM grid size (store) and TO grid size (retrieve).
-				InOutTaskGenerationGrid.Add(FromGrid | ToGrid);
+				// The link executes at both FROM grid size (store) and TO grid size (retrieve). Just store from grid size here - the task culling lambda checks for
+				// grind links and checks both grid sizes.
+				InOutTaskGenerationGrid.Add(FromGrid);
 			}
 		}
 	}
 }
 
-EPCGHiGenGrid FPCGGraphCompiler::CalculateGridRecursive(
+uint32 FPCGGraphCompiler::CalculateGridRecursive(
 	FPCGTaskId InTaskId,
-	EPCGHiGenGrid GenerationDefaultGrid,
+	uint32 GenerationDefaultGrid,
 	const FPCGStackContext& InStackContext,
 	const TArray<FPCGGraphTask>& InCompiledTasks,
-	TArray<EPCGHiGenGrid>& InOutTaskGenerationGrid)
+	TArray<uint32>& InOutTaskGenerationGrid)
 {
-	if (InOutTaskGenerationGrid[InTaskId] != EPCGHiGenGrid::Uninitialized)
+	if (InOutTaskGenerationGrid[InTaskId] != PCGHiGenGrid::UninitializedGridSize())
 	{
 		return InOutTaskGenerationGrid[InTaskId];
 	}
@@ -651,12 +652,12 @@ EPCGHiGenGrid FPCGGraphCompiler::CalculateGridRecursive(
 	const UPCGSettings* Settings = Node ? Node->GetSettings() : nullptr;
 	const UPCGHiGenGridSizeSettings* GridSizeSettings = Cast<UPCGHiGenGridSizeSettings>(Settings);
 
-	EPCGHiGenGrid Grid = EPCGHiGenGrid::Uninitialized;
+	uint32 GridSize = PCGHiGenGrid::UninitializedGridSize();
 
 	// Grid Size nodes in the top graph set the execution grid level.
 	if (GridSizeSettings && GridSizeSettings->bEnabled && bTopLevelGraph)
 	{
-		Grid = FMath::Min(GenerationDefaultGrid, GridSizeSettings->GetGrid());
+		GridSize = FMath::Min(GenerationDefaultGrid, GridSizeSettings->GetScaledGridSize());
 	}
 	else
 	{
@@ -665,36 +666,36 @@ EPCGHiGenGrid FPCGGraphCompiler::CalculateGridRecursive(
 			if (bTopLevelGraph)
 			{
 				// Tasks with no inputs in top graph get prescribed the default generation grid.
-				Grid = GenerationDefaultGrid;
+				GridSize = GenerationDefaultGrid;
 			}
 			else
 			{
 				// Tasks with no inputs in a subgraph should execute on the same grid as the subgraph node task.
 				check(InCompiledTasks[InTaskId].ParentId != InvalidPCGTaskId);
-				Grid = CalculateGridRecursive(InCompiledTasks[InTaskId].ParentId, GenerationDefaultGrid, InStackContext, InCompiledTasks, InOutTaskGenerationGrid);
+				GridSize = CalculateGridRecursive(InCompiledTasks[InTaskId].ParentId, GenerationDefaultGrid, InStackContext, InCompiledTasks, InOutTaskGenerationGrid);
 			}
 		}
 		else
 		{
 			// This task has inputs. Grid of this task is minimum of all input grids. We can link in data from a larger grid, but
 			// not from a finer grid (this goes against hierarchy).
-			Grid = EPCGHiGenGrid::Unbounded;
+			GridSize = PCGHiGenGrid::UnboundedGridSize();
 			for (FPCGGraphTaskInput InputTask : InCompiledTasks[InTaskId].Inputs)
 			{
-				const EPCGHiGenGrid InputGrid = CalculateGridRecursive(InputTask.TaskId, GenerationDefaultGrid, InStackContext, InCompiledTasks, InOutTaskGenerationGrid);
-				if (PCGHiGenGrid::IsValidGrid(InputGrid))
+				const uint32 InputGridSize = CalculateGridRecursive(InputTask.TaskId, GenerationDefaultGrid, InStackContext, InCompiledTasks, InOutTaskGenerationGrid);
+				if (PCGHiGenGrid::IsValidGridSize(InputGridSize))
 				{
-					Grid = FMath::Min(InputGrid, Grid);
+					GridSize = FMath::Min(InputGridSize, GridSize);
 				}
 			}
 		}
 	}
 
-	ensure(Grid != EPCGHiGenGrid::Uninitialized);
+	ensure(GridSize != PCGHiGenGrid::UninitializedGridSize());
 
-	InOutTaskGenerationGrid[InTaskId] = Grid;
+	InOutTaskGenerationGrid[InTaskId] = GridSize;
 
-	return Grid;
+	return GridSize;
 }
 
 bool FPCGGraphCompiler::CalculateStaticallyActiveRecursive(FPCGTaskId InTaskId, const TArray<FPCGGraphTask>& InCompiledTasks, TMap<FPCGTaskId, bool>& InOutTaskIdToActiveFlag)
@@ -1388,22 +1389,35 @@ void FPCGGraphCompiler::CompileTopGraph(UPCGGraph* InGraph, uint32 GenerationGri
 		// TODO - we could add an else branch for higen disabled that culls the grid size nodes.
 		if (InGraph->IsHierarchicalGenerationEnabled() && GenerationGridSize != PCGHiGenGrid::UninitializedGridSize())
 		{
-			const EPCGHiGenGrid GenerationGrid = PCGHiGenGrid::GridSizeToGrid(GenerationGridSize);
-			const EPCGHiGenGrid DefaultGrid = PCGHiGenGrid::GridSizeToGrid(InGraph->GetDefaultGridSize());
+			const uint32 DefaultGrid = InGraph->GetDefaultGridSize();
 
 			// Propagate grid size nodes through the graph to determine which grid size each task should execute on.
-			TArray<EPCGHiGenGrid> TaskGenerationGrid;
+			TArray<uint32> TaskGenerationGrid;
 			TaskGenerationGrid.SetNumZeroed(CompiledTasks.Num());
-			ResolveGridSizes(GenerationGrid, CompiledTasks, StackContext, DefaultGrid, TaskGenerationGrid);
+			ResolveGridSizes(GenerationGridSize, CompiledTasks, StackContext, DefaultGrid, TaskGenerationGrid);
 
 			// Create linkage tasks for edges that cross from large grid to small grid tasks.
-			CreateGridLinkages(InGraph, GenerationGrid, TaskGenerationGrid, CompiledTasks, StackContext, bIsCooking);
+			CreateGridLinkages(InGraph, GenerationGridSize, TaskGenerationGrid, CompiledTasks, StackContext, bIsCooking);
 
 			// Cull any task that should not execute on the current grid.
-			CullTasks(CompiledTasks, /*bAddPassthroughWires=*/false, [GenerationGrid, &TaskGenerationGrid](const FPCGGraphTask& InTask)
+			CullTasks(CompiledTasks, /*bAddPassthroughWires=*/false, [GenerationGridSize, &TaskGenerationGrid](const FPCGGraphTask& InTask)
 			{
-				const EPCGHiGenGrid TaskGrid = TaskGenerationGrid[InTask.NodeId];
-				return TaskGrid != EPCGHiGenGrid::Uninitialized && !(TaskGrid & GenerationGrid);
+				const uint32 TaskGridSize = TaskGenerationGrid[InTask.NodeId];
+				if (TaskGridSize == PCGHiGenGrid::UninitializedGridSize())
+				{
+					// Don't cull if no grid size set.
+					return false;
+				}
+
+				if (InTask.Element && InTask.Element->IsGridLinkage())
+				{
+					PCGGraphExecutor::FPCGGridLinkageElement* LinkElement = static_cast<PCGGraphExecutor::FPCGGridLinkageElement*>(InTask.Element.Get());
+					return LinkElement->GetFromGridSize() != GenerationGridSize && LinkElement->GetToGridSize() != GenerationGridSize;
+				}
+				else
+				{
+					return TaskGridSize != GenerationGridSize;
+				}
 			});
 		}
 	}
diff --git a/Engine/Plugins/PCG/Source/PCG/Private/Graph/PCGGraphCompiler.h b/Engine/Plugins/PCG/Source/PCG/Private/Graph/PCGGraphCompiler.h
--- a/Engine/Plugins/PCG/Source/PCG/Private/Graph/PCGGraphCompiler.h
+++ b/Engine/Plugins/PCG/Source/PCG/Private/Graph/PCGGraphCompiler.h
@@ -108,25 +108,25 @@ private:
 
 	/** Propagates grid sizes through a graph's compiled tasks. */
 	static void ResolveGridSizes(
-		EPCGHiGenGrid GenerationGrid,
+		uint32 GenerationGridSize,
 		const TArray<FPCGGraphTask>& CompiledTasks,
 		const FPCGStackContext& StackContext,
-		EPCGHiGenGrid GenerationDefaultGrid,
-		TArray<EPCGHiGenGrid>& InOutTaskGenerationGrid);
+		uint32 ScaledDefaultGridSize,
+		TArray<uint32>& InOutTaskGenerationGrid);
 
 	/** Returns the execution grid for the given task. */
-	static EPCGHiGenGrid CalculateGridRecursive(
+	static uint32 CalculateGridRecursive(
 		FPCGTaskId InTaskId,
-		EPCGHiGenGrid GenerationDefaultGrid,
+		uint32 InScaledDefaultGridSize,
 		const FPCGStackContext& InStackContext,
 		const TArray<FPCGGraphTask>& InCompiledTasks,
-		TArray<EPCGHiGenGrid>& InOutTaskGenerationGrid);
+		TArray<uint32>& InOutTaskGenerationGridSize);
 
 	/** Create linkage tasks for edges that cross from large grid to small grid tasks. */
 	static void CreateGridLinkages(
 		UPCGGraph* InGraph,
-		EPCGHiGenGrid InGenerationGrid,
-		TArray<EPCGHiGenGrid>& TaskGenerationGrid,
+		uint32 InGenerationGrid,
+		TArray<uint32>& TaskGenerationGrid,
 		TArray<FPCGGraphTask>& InOutCompiledTasks,
 		const FPCGStackContext& InStackContext,
 		bool bIsCooking);
diff --git a/Engine/Plugins/PCG/Source/PCG/Private/Graph/PCGGraphExecutor.cpp b/Engine/Plugins/PCG/Source/PCG/Private/Graph/PCGGraphExecutor.cpp
--- a/Engine/Plugins/PCG/Source/PCG/Private/Graph/PCGGraphExecutor.cpp
+++ b/Engine/Plugins/PCG/Source/PCG/Private/Graph/PCGGraphExecutor.cpp
@@ -3458,9 +3458,9 @@ namespace PCGGraphExecutor
 	}
 
 	bool ExecuteGridLinkage(
-		EPCGHiGenGrid InGenerationGrid,
-		EPCGHiGenGrid InFromGrid,
-		EPCGHiGenGrid InToGrid,
+		uint32 InGenerationGrid,
+		uint32 InFromGrid,
+		uint32 InToGrid,
 		const FString& InResourceKey,
 		FName InUpstreamPinLabel,
 		FPCGGridLinkageContext* InContext)
@@ -3479,32 +3479,29 @@ namespace PCGGraphExecutor
 
 		// Non-hierarchical generation - no linkage required - data should just pass through.
 		if (!ExecutionState.GetGraph()->IsHierarchicalGenerationEnabled()
-			|| !ensure(PCGHiGenGrid::IsValidGrid(InFromGrid) || InFromGrid == EPCGHiGenGrid::Unbounded))
+			|| !ensure(PCGHiGenGrid::IsValidGridSize(InFromGrid) || InFromGrid == PCGHiGenGrid::UnboundedGridSize()))
 		{
 			InContext->OutputData = InContext->InputData;
 			return true;
 		}
 
-		const uint32 FromGridSize = PCGHiGenGrid::IsValidGrid(InFromGrid) ? PCGHiGenGrid::GridToGridSize(InFromGrid) : PCGHiGenGrid::UnboundedGridSize();
-		const uint32 ToGridSize = PCGHiGenGrid::IsValidGrid(InToGrid) ? PCGHiGenGrid::GridToGridSize(InToGrid) : PCGHiGenGrid::UnboundedGridSize();
-
 		// Never allow a large grid to read data from small grid - this violates hierarchy.
-		if (FromGridSize < ToGridSize)
+		if (InFromGrid < InToGrid)
 		{
 #if WITH_EDITOR
 			if (IPCGEditorModule* PCGEditorModule = IPCGEditorModule::Get())
 			{
 				// Using the low level logging call because we have only a node pointer for the downstream node. Note that InContext
 				// is the context for the linkage element/task, which is not represented on the graph and cannot receive graph warnings/errors.
-				if (ToGridSize == PCGHiGenGrid::UnboundedGridSize())
+				if (InToGrid == PCGHiGenGrid::UnboundedGridSize())
 				{
 					PCGEditorModule->GetNodeVisualLogsMutable().Log(
 						*InContext->GetStack(),
 						ELogVerbosity::Error,
 						FText::Format(
 							NSLOCTEXT("PCGGraphCompiler", "InvalidLinkageToUnbounded", "Could not read data across grid levels - cannot read from grid size {0} to Unbounded domain."),
-							FromGridSize,
-							ToGridSize));
+							InFromGrid,
+							InToGrid));
 				}
 				else
 				{
@@ -3513,8 +3510,8 @@ namespace PCGGraphExecutor
 						ELogVerbosity::Error,
 						FText::Format(
 							NSLOCTEXT("PCGGraphCompiler", "InvalidLinkageInvalidGridSizes", "Could not read data across grid levels - origin grid size {0} must be greater than destination grid size {1}. Graph default grid size may need increasing."),
-							FromGridSize,
-							ToGridSize));
+							InFromGrid,
+							InToGrid));
 				}
 			}
 #endif
@@ -3522,24 +3519,24 @@ namespace PCGGraphExecutor
 			return true;
 		}
 
-		if (!!(InFromGrid & InGenerationGrid) && FromGridSize != ToGridSize)
+		if (InFromGrid == InGenerationGrid && InFromGrid != InToGrid)
 		{
 			FPCGDataCollection Data;
 
 			// Grid linkages have a single virtual input pin with the label matching the upstream pin. A grid linkage is created for each edge.
 			Data.TaggedData = InContext->InputData.GetInputsByPin(InUpstreamPinLabel);
 
-			PCGGraphExecutionLogging::LogGridLinkageTaskExecuteStore(InContext, InGenerationGrid, FromGridSize, ToGridSize, InResourceKey, Data.TaggedData.Num());
+			PCGGraphExecutionLogging::LogGridLinkageTaskExecuteStore(InContext, InGenerationGrid, InFromGrid, InToGrid, InResourceKey, Data.TaggedData.Num());
 			
 			ExecutionState.StoreOutputDataForPin(InResourceKey, Data);
 		}
-		else if (InToGrid == InGenerationGrid && FromGridSize != ToGridSize)
+		else if (InToGrid == InGenerationGrid && InFromGrid != InToGrid)
 		{
-			PCGGraphExecutionLogging::LogGridLinkageTaskExecuteRetrieve(InContext, InGenerationGrid, FromGridSize, ToGridSize, InResourceKey);
+			PCGGraphExecutionLogging::LogGridLinkageTaskExecuteRetrieve(InContext, InGenerationGrid, InFromGrid, InToGrid, InResourceKey);
 
 			IPCGGraphExecutionSource* ExecSourceWithData = nullptr;
 
-			if (FromGridSize == PCGHiGenGrid::UnboundedGridSize())
+			if (InFromGrid == PCGHiGenGrid::UnboundedGridSize())
 			{
 				ExecSourceWithData = ExecutionState.GetOriginalSource();
 			}
@@ -3551,9 +3548,9 @@ namespace PCGGraphExecutor
 					return false;
 				}
 
-				const FPCGGridDescriptor GridDescriptor = OriginalSource->GetExecutionState().GetGridDescriptor(FromGridSize);
+				const FPCGGridDescriptor GridDescriptor = OriginalSource->GetExecutionState().GetGridDescriptor(InFromGrid);
 
-				// Get grid coords using the parent grid (FromGridSize).
+				// Get grid coords using the parent grid (InFromGrid).
 				const FIntVector CellCoords = UPCGActorHelpers::GetCellCoord(ExecutionState.GetTransform().GetLocation(), GridDescriptor.GetGridSize(), GridDescriptor.Is2DGrid());
 
 				ExecSourceWithData = ExecutionState.GetLocalSource(GridDescriptor, CellCoords);
diff --git a/Engine/Plugins/PCG/Source/PCG/Private/Graph/PCGGraphExecutor.h b/Engine/Plugins/PCG/Source/PCG/Private/Graph/PCGGraphExecutor.h
--- a/Engine/Plugins/PCG/Source/PCG/Private/Graph/PCGGraphExecutor.h
+++ b/Engine/Plugins/PCG/Source/PCG/Private/Graph/PCGGraphExecutor.h
@@ -534,7 +534,7 @@ namespace PCGGraphExecutor
 	class FPCGGridLinkageElement : public FPCGGenericElement
 	{
 	public:
-		FPCGGridLinkageElement(TFunction<bool(FPCGContext*)> InOperation, const FContextAllocator& InContextAllocator, EPCGHiGenGrid InFromGrid, EPCGHiGenGrid InToGrid, EPCGHiGenGrid InGenerationGrid, const FString& InResourceKey, const UPCGPin* InUpstreamPin)
+		FPCGGridLinkageElement(TFunction<bool(FPCGContext*)> InOperation, const FContextAllocator& InContextAllocator, uint32 InFromGrid, uint32 InToGrid, uint32 InGenerationGrid, const FString& InResourceKey, const UPCGPin* InUpstreamPin)
 			: FPCGGenericElement(InOperation, InContextAllocator)
 			, FromGrid(InFromGrid)
 			, ToGrid(InToGrid)
@@ -557,6 +557,8 @@ namespace PCGGraphExecutor
 
 	public:
 		const UPCGPin* GetUpstreamPin() const { return UpstreamPin.Get(); }
+		uint32 GetFromGridSize() const { return FromGrid; }
+		uint32 GetToGridSize() const { return ToGrid; }
 
 #if WITH_EDITOR
 		/** Return true if the grid sizes & path match. */
@@ -565,11 +567,11 @@ namespace PCGGraphExecutor
 
 	private:
 		// These values are stored here so that we can compare two grid linkage elements for equivalence.
-		EPCGHiGenGrid FromGrid = EPCGHiGenGrid::Uninitialized;
-		EPCGHiGenGrid ToGrid = EPCGHiGenGrid::Uninitialized;
+		uint32 FromGrid = PCGHiGenGrid::UninitializedGridSize();
+		uint32 ToGrid = PCGHiGenGrid::UninitializedGridSize();
 
 		// This tells us which side of the From/To relationship this grid linkage is on.
-		EPCGHiGenGrid GenerationGrid = EPCGHiGenGrid::Uninitialized; 
+		uint32 GenerationGrid = PCGHiGenGrid::UninitializedGridSize();
 
 		FString ResourceKey;
 		TWeakObjectPtr<const UPCGPin> UpstreamPin = nullptr;
@@ -577,9 +579,9 @@ namespace PCGGraphExecutor
 
 	/** Compares InFromGrid and InToGrid and performs data storage/retrieval as necessary to marshal data across execution grids. */
 	bool ExecuteGridLinkage(
-		EPCGHiGenGrid InGenerationGrid,
-		EPCGHiGenGrid InFromGrid,
-		EPCGHiGenGrid InToGrid,
+		uint32 InGenerationGridSize,
+		uint32 InFromGridSize,
+		uint32 InToGridSize,
 		const FString& InResourceKey,
 		FName InUpstreamPinLabel,
 		FPCGGridLinkageContext* InContext);
@@ -598,13 +600,13 @@ protected:
 
 public:
 	UPROPERTY()
-	EPCGHiGenGrid FromGrid = EPCGHiGenGrid::Uninitialized;
+	uint32 FromGrid = PCGHiGenGrid::UninitializedGridSize();
 
 	UPROPERTY()
-	EPCGHiGenGrid ToGrid = EPCGHiGenGrid::Uninitialized;
+	uint32 ToGrid = PCGHiGenGrid::UninitializedGridSize();
 
 	UPROPERTY()
-	EPCGHiGenGrid GenerationGrid = EPCGHiGenGrid::Uninitialized;
+	uint32 GenerationGrid = PCGHiGenGrid::UninitializedGridSize();
 
 	UPROPERTY()
 	FString ResourceKey;
diff --git a/Engine/Plugins/PCG/Source/PCG/Private/Helpers/PCGHelpers.cpp b/Engine/Plugins/PCG/Source/PCG/Private/Helpers/PCGHelpers.cpp
--- a/Engine/Plugins/PCG/Source/PCG/Private/Helpers/PCGHelpers.cpp
+++ b/Engine/Plugins/PCG/Source/PCG/Private/Helpers/PCGHelpers.cpp
@@ -817,7 +817,7 @@ namespace PCGHelpers
 			if (ComponentInner && World->IsGameWorld())
 			{
 				FString Text;
-				if (ComponentInner->GetGenerationGrid() == EPCGHiGenGrid::Unbounded)
+				if (ComponentInner->GetGenerationGridSize() == PCGHiGenGrid::UnboundedGridSize())
 				{
 					if (ComponentInner->GetOwner())
 					{
diff --git a/Engine/Plugins/PCG/Source/PCG/Private/PCGCommon.cpp b/Engine/Plugins/PCG/Source/PCG/Private/PCGCommon.cpp
--- a/Engine/Plugins/PCG/Source/PCG/Private/PCGCommon.cpp
+++ b/Engine/Plugins/PCG/Source/PCG/Private/PCGCommon.cpp
@@ -3,7 +3,6 @@
 #include "PCGCommon.h"
 #include "PCGComponent.h"
 #include "PCGSubsystem.h"
-#include "PCGWorldActor.h"
 
 #include "EngineDefines.h" // For UE_ENABLE_DEBUG_DRAWING
 
@@ -140,19 +139,7 @@ namespace PCGHiGenGrid
 {
 	bool IsValidGridSize(uint32 InGridSize)
 	{
-		// Must be a power of 2 (in m) and within the valid range
-		// TODO: support other units
-		if (FMath::IsPowerOfTwo(InGridSize / 100)
-			&& InGridSize >= GridToGridSize(EPCGHiGenGrid::GridMin)
-			&& InGridSize <= GridToGridSize(EPCGHiGenGrid::GridMax))
-		{
-			return true;
-		}
-
-		UPCGSubsystem* PCGSubsystem = UPCGSubsystem::GetSubsystemForCurrentWorld();
-		APCGWorldActor* PCGWorldActor = PCGSubsystem ? PCGSubsystem->GetPCGWorldActor() : nullptr;
-
-		return !PCGWorldActor || PCGWorldActor->PartitionGridSize == InGridSize;
+		return InGridSize > UninitializedGridSize() && InGridSize < UnboundedGridSize();
 	}
 
 	bool IsValidGrid(EPCGHiGenGrid InGrid)
@@ -166,6 +153,11 @@ namespace PCGHiGenGrid
 		return IsValidGrid(InGrid) || InGrid == EPCGHiGenGrid::Uninitialized;
 	}
 
+	bool IsValidGridOrUninitialized(uint32 InGridSize)
+	{
+		return IsValidGridSize(InGridSize) || InGridSize == PCGHiGenGrid::UninitializedGridSize();
+	}
+
 	uint32 GridToGridSize(EPCGHiGenGrid InGrid)
 	{
 		const uint32 GridAsUint = static_cast<uint32>(InGrid);
@@ -194,6 +186,27 @@ namespace PCGHiGenGrid
 		return static_cast<EPCGHiGenGrid>(RoundedGridSizeMeters);
 	}
 
+	EPCGHiGenGrid ScaledGridSizeToGrid(uint32 InScaledGridSize, double InGridScale)
+	{
+		if (InScaledGridSize == UnboundedGridSize())
+		{
+			return EPCGHiGenGrid::Unbounded;
+		}
+		else if (InScaledGridSize == UninitializedGridSize())
+		{
+			return EPCGHiGenGrid::Uninitialized;
+		}
+		else if (InGridScale <= UE_KINDA_SMALL_NUMBER)
+		{
+			ensure(false);
+			return EPCGHiGenGrid::Uninitialized;
+		}
+
+		const uint32 UnscaledGridSize = static_cast<uint32>(FMath::RoundToInt(static_cast<double>(InScaledGridSize) / InGridScale));
+
+		return GridSizeToGrid(UnscaledGridSize);
+	}
+
 	uint32 UninitializedGridSize()
 	{
 		return static_cast<uint32>(EPCGHiGenGrid::Uninitialized);
@@ -203,6 +216,15 @@ namespace PCGHiGenGrid
 	{
 		return static_cast<uint32>(EPCGHiGenGrid::Unbounded);
 	}
+
+	double SanitizeGridSizeMultiplier(double InGridSizeMultiplier)
+	{
+		// The smallest grid scaled by the multiplier have an integral size.
+		const double SmallestGridSize = static_cast<double>(GridToGridSize(EPCGHiGenGrid::GridMin));
+		// Arbitrary choice. We could reduce this if we hear the requirement to have smaller grid cells.
+		const double SmallestScaledGridSize = 100.0;
+		return FMath::Max(SmallestScaledGridSize, FMath::RoundToDouble(InGridSizeMultiplier * SmallestGridSize)) / SmallestGridSize;
+	}
 }
 
 void FPCGRuntimeGenerationRadii::PostSerialize(const FArchive& Ar)
diff --git a/Engine/Plugins/PCG/Source/PCG/Private/PCGComponent.cpp b/Engine/Plugins/PCG/Source/PCG/Private/PCGComponent.cpp
--- a/Engine/Plugins/PCG/Source/PCG/Private/PCGComponent.cpp
+++ b/Engine/Plugins/PCG/Source/PCG/Private/PCGComponent.cpp
@@ -90,6 +90,12 @@ namespace PCGComponent
 		constexpr double Factor = 10000;
 		return FMath::RoundToDouble(InTrackingPriority * Factor) / Factor;
 	}
+
+	EPCGHiGenGrid GetUnscaledGenerationGrid(uint32 ScaledGridSize, const UPCGGraph* Graph)
+	{
+		const double GridScale = Graph ? Graph->GetGridSizeMultiplier() : 1.0;
+		return PCGHiGenGrid::ScaledGridSizeToGrid(ScaledGridSize, GridScale);
+	}
 }
 
 UPCGComponent::UPCGComponent(const FObjectInitializer& InObjectInitializer)
@@ -475,27 +481,45 @@ void UPCGComponent::GenerateLocal(bool bForce)
 	GenerateLocalGetTaskId(bForce);
 }
 
-void UPCGComponent::GenerateLocal(EPCGComponentGenerationTrigger RequestedGenerationTrigger, bool bForce, EPCGHiGenGrid Grid, const TArray<FPCGTaskId>& Dependencies)
+void UPCGComponent::GenerateLocal(EPCGComponentGenerationTrigger RequestedGenerationTrigger, bool bForce, uint32 Grid, const TArray<FPCGTaskId>& Dependencies)
 {
 	GenerateInternal(bForce, Grid, RequestedGenerationTrigger, Dependencies);
 }
 
+// Deprecated 5.8
+void UPCGComponent::GenerateLocal(EPCGComponentGenerationTrigger RequestedGenerationTrigger, bool bForce, EPCGHiGenGrid Grid, const TArray<FPCGTaskId>& Dependencies)
+{
+	GenerateInternal(bForce, PCGHiGenGrid::GridToGridSize(Grid), RequestedGenerationTrigger, Dependencies);
+}
+
 FPCGTaskId UPCGComponent::GenerateLocalGetTaskId(bool bForce)
 {
-	return GenerateInternal(bForce, EPCGHiGenGrid::Uninitialized, EPCGComponentGenerationTrigger::GenerateOnDemand, {});
+	return GenerateInternal(bForce, PCGHiGenGrid::UninitializedGridSize(), EPCGComponentGenerationTrigger::GenerateOnDemand, {});
 }
 
-FPCGTaskId UPCGComponent::GenerateLocalGetTaskId(EPCGComponentGenerationTrigger RequestedGenerationTrigger, bool bForce, EPCGHiGenGrid Grid)
+FPCGTaskId UPCGComponent::GenerateLocalGetTaskId(EPCGComponentGenerationTrigger RequestedGenerationTrigger, bool bForce, uint32 Grid)
 {
 	return GenerateInternal(bForce, Grid, RequestedGenerationTrigger, {});
 }
 
-FPCGTaskId UPCGComponent::GenerateLocalGetTaskId(EPCGComponentGenerationTrigger RequestedGenerationTrigger, bool bForce, EPCGHiGenGrid Grid, const TArray<FPCGTaskId>& Dependencies)
+FPCGTaskId UPCGComponent::GenerateLocalGetTaskId(EPCGComponentGenerationTrigger RequestedGenerationTrigger, bool bForce, uint32 Grid, const TArray<FPCGTaskId>& Dependencies)
 {
 	return GenerateInternal(bForce, Grid, RequestedGenerationTrigger, Dependencies);
 }
 
-FPCGTaskId UPCGComponent::GenerateInternal(bool bForce, EPCGHiGenGrid Grid, EPCGComponentGenerationTrigger RequestedGenerationTrigger, const TArray<FPCGTaskId>& Dependencies)
+// Deprecated 5.8
+FPCGTaskId UPCGComponent::GenerateLocalGetTaskId(EPCGComponentGenerationTrigger RequestedGenerationTrigger, bool bForce, EPCGHiGenGrid Grid)
+{
+	return GenerateInternal(bForce, PCGHiGenGrid::GridToGridSize(Grid), RequestedGenerationTrigger, {});
+}
+
+// Deprecated 5.8
+FPCGTaskId UPCGComponent::GenerateLocalGetTaskId(EPCGComponentGenerationTrigger RequestedGenerationTrigger, bool bForce, EPCGHiGenGrid Grid, const TArray<FPCGTaskId>& Dependencies)
+{
+	return GenerateInternal(bForce, PCGHiGenGrid::GridToGridSize(Grid), RequestedGenerationTrigger, Dependencies);
+}
+
+FPCGTaskId UPCGComponent::GenerateInternal(bool bForce, uint32 Grid, EPCGComponentGenerationTrigger RequestedGenerationTrigger, const TArray<FPCGTaskId>& Dependencies)
 {
 	if (IsGenerating() || !GetSubsystem() || !ShouldGenerate(bForce, RequestedGenerationTrigger))
 	{
@@ -1226,21 +1250,9 @@ AActor* UPCGComponent::ClearPCGLink(const FPCGMoveResourceParams& InParams)
 	return NewActor;
 }
 
-EPCGHiGenGrid UPCGComponent::GetGenerationGrid() const
+EPCGHiGenGrid UPCGComponent::GetUnscaledGenerationGrid() const
 {
-	const uint32 GridSize = GetGenerationGridSize();
-	if (PCGHiGenGrid::IsValidGridSize(GridSize))
-	{
-		return PCGHiGenGrid::GridSizeToGrid(GridSize);
-	}
-	else if (GridSize == PCGHiGenGrid::UnboundedGridSize())
-	{
-		return EPCGHiGenGrid::Unbounded;
-	}
-	else
-	{
-		return PCGHiGenGrid::GridSizeToGrid(GetGraph()->GetDefaultGridSize());
-	}
+	return PCGComponent::GetUnscaledGenerationGrid(GetGenerationGridSize(), GetGraph());
 }
 
 void UPCGComponent::StoreOutputDataForPin(const FString& InResourceKey, const FPCGDataCollection& InData)
@@ -1299,17 +1311,19 @@ const FPCGRuntimeGenerationRadii& UPCGComponent::GetGenerationRadii() const
 	return (bOverrideGenerationRadii || !GetGraph()) ? GenerationRadii : GetGraph()->GenerationRadii;
 }
 
-double UPCGComponent::GetGenerationRadiusFromGrid(EPCGHiGenGrid Grid) const
+double UPCGComponent::GetGenerationRadiusFromGrid(uint32 Grid) const
 {
 	const double Multiplier = FMath::Max(PCGRuntimeGenSchedulerHelpers::CVarRuntimeGenerationRadiusMultiplier.GetValueOnAnyThread(), 0.0f);
 
+	const UPCGGraph* Graph = GetGraph();
+	ensure(Graph);
+
 	if (bOverrideGenerationRadii)
 	{
-		return Multiplier * GenerationRadii.GetGenerationRadiusFromGrid(Grid);
+		const EPCGHiGenGrid UnscaledGrid = PCGComponent::GetUnscaledGenerationGrid(Grid, GetGraph());
+		return Multiplier * GenerationRadii.GetGenerationRadiusFromGrid(UnscaledGrid);
 	}
-
-	const UPCGGraph* Graph = GetGraph();
-	if (ensure(Graph))
+	else if (Graph)
 	{
 		return Multiplier * Graph->GetGridGenerationRadiusFromGrid(Grid);
 	}
@@ -1317,17 +1331,19 @@ double UPCGComponent::GetGenerationRadiusFromGrid(EPCGHiGenGrid Grid) const
 	return 0;
 }
 
-double UPCGComponent::GetCleanupRadiusFromGrid(EPCGHiGenGrid Grid) const
+double UPCGComponent::GetCleanupRadiusFromGrid(uint32 Grid) const
 {
 	const double Multiplier = FMath::Max(PCGRuntimeGenSchedulerHelpers::CVarRuntimeGenerationRadiusMultiplier.GetValueOnAnyThread(), 0.0f);
 
+	const UPCGGraph* Graph = GetGraph();
+	ensure(Graph);
+
 	if (bOverrideGenerationRadii)
 	{
-		return Multiplier * GenerationRadii.GetCleanupRadiusFromGrid(Grid);
+		const EPCGHiGenGrid UnscaledGrid = PCGComponent::GetUnscaledGenerationGrid(Grid, GetGraph());
+		return Multiplier * GenerationRadii.GetCleanupRadiusFromGrid(UnscaledGrid);
 	}
-
-	const UPCGGraph* Graph = GetGraph();
-	if (ensure(Graph))
+	else if (Graph)
 	{
 		return Multiplier * Graph->GetGridCleanupRadiusFromGrid(Grid);
 	}
@@ -1680,7 +1696,7 @@ void UPCGComponent::BeginPlay()
 
 	if (OriginalComponent && bActivated && !bAlreadyGenerated && GenerationTrigger == EPCGComponentGenerationTrigger::GenerateOnLoad)
 	{
-		GenerateInternal(/*bForce=*/false, EPCGHiGenGrid::Uninitialized, EPCGComponentGenerationTrigger::GenerateOnLoad, {});
+		GenerateInternal(/*bForce=*/false, PCGHiGenGrid::UninitializedGridSize(), EPCGComponentGenerationTrigger::GenerateOnLoad, {});
 	}
 }
 
diff --git a/Engine/Plugins/PCG/Source/PCG/Private/PCGComponentExecutionState.cpp b/Engine/Plugins/PCG/Source/PCG/Private/PCGComponentExecutionState.cpp
--- a/Engine/Plugins/PCG/Source/PCG/Private/PCGComponentExecutionState.cpp
+++ b/Engine/Plugins/PCG/Source/PCG/Private/PCGComponentExecutionState.cpp
@@ -165,13 +165,13 @@ const FPCGRuntimeGenerationRadii& FPCGComponentExecutionState::GetGenerationRadi
 	return Component->GetGenerationRadii();
 }
 
-double FPCGComponentExecutionState::GetGenerationRadiusFromGrid(EPCGHiGenGrid InGrid) const
+double FPCGComponentExecutionState::GetGenerationRadiusFromGrid(uint32 InGrid) const
 {
 	check(Component);
 	return Component->GetGenerationRadiusFromGrid(InGrid);
 }
 
-double FPCGComponentExecutionState::GetCleanupRadiusFromGrid(EPCGHiGenGrid InGrid) const
+double FPCGComponentExecutionState::GetCleanupRadiusFromGrid(uint32 InGrid) const
 {
 	check(Component);
 	return Component->GetCleanupRadiusFromGrid(InGrid);
@@ -213,12 +213,6 @@ FPCGGridDescriptor FPCGComponentExecutionState::GetGridDescriptor(uint32 InGridS
 	return Component->GetGridDescriptor(InGridSize);
 }
 
-EPCGHiGenGrid FPCGComponentExecutionState::GetGenerationGrid() const
-{
-	check(Component);
-	return (IsPartitioned() || IsLocalSource()) ? Component->GetGenerationGrid() : EPCGHiGenGrid::Uninitialized;
-}
-
 uint32 FPCGComponentExecutionState::GetGenerationGridSize() const
 {
 	check(Component);
@@ -283,7 +277,7 @@ FPCGTaskId FPCGComponentExecutionState::Generate(const FGenerateParams& InGenera
 {
 	check(Component);
 
-	const EPCGHiGenGrid Grid = InGenerateParams.bGenerateLocalSources ? EPCGHiGenGrid::Uninitialized : GetGenerationGrid();
+	const uint32 Grid = InGenerateParams.bGenerateLocalSources ? PCGHiGenGrid::UninitializedGridSize() : GetGenerationGridSize();
 	const EPCGComponentGenerationTrigger GenTrigger = IsManagedByRuntimeGenSystem() ? EPCGComponentGenerationTrigger::GenerateAtRuntime : EPCGComponentGenerationTrigger::GenerateOnDemand;
 	return Component->GenerateLocalGetTaskId(GenTrigger, InGenerateParams.bEvenIfAlreadyGenerated, Grid, InGenerateParams.Dependencies);
 }
diff --git a/Engine/Plugins/PCG/Source/PCG/Private/PCGGraph.cpp b/Engine/Plugins/PCG/Source/PCG/Private/PCGGraph.cpp
--- a/Engine/Plugins/PCG/Source/PCG/Private/PCGGraph.cpp
+++ b/Engine/Plugins/PCG/Source/PCG/Private/PCGGraph.cpp
@@ -622,6 +622,14 @@ void UPCGGraph::PostLoad()
 	}
 
 	GenerationRadii.ApplyDeprecation();
+
+	PRAGMA_DISABLE_DEPRECATION_WARNINGS
+	if (HiGenExponential != 0)
+	{
+		HiGenGridSizeMultiplier = static_cast<double>(1 << HiGenExponential);
+		HiGenExponential = 0;
+	}
+	PRAGMA_ENABLE_DEPRECATION_WARNINGS
 #endif // WITH_EDITOR
 }
 
@@ -1004,7 +1012,7 @@ uint32 UPCGGraph::GetDefaultGridSize() const
 
 	if (IsHierarchicalGenerationEnabled() && PCGHiGenGrid::IsValidGrid(HiGenGridSize))
 	{
-		return PCGHiGenGrid::GridToGridSize(HiGenGridSize) * (1 << HiGenExponential);
+		return FMath::RoundToInt(PCGHiGenGrid::GridToGridSize(HiGenGridSize) * GetGridSizeMultiplier());
 	}
 	else
 	{
@@ -1730,44 +1738,16 @@ void UPCGGraph::GetParentGridSizes(const uint32 InChildGridSize, PCGHiGenGrid::F
 	ChildGridSizeToParentGridSizes.Add(InChildGridSize, OutParentGridSizes);
 }
 
-double UPCGGraph::GetGridGenerationRadiusFromGrid(EPCGHiGenGrid Grid) const
+double UPCGGraph::GetGridGenerationRadiusFromGrid(uint32 Grid) const
 {
-	if (Grid == EPCGHiGenGrid::Unbounded || Grid == EPCGHiGenGrid::Uninitialized)
-	{
-		return GenerationRadii.GetGenerationRadiusFromGrid(Grid);
-	}
-	// If the queried grid is smaller than the min grid including the exponent, we'll take the min grid and scale it down
-	else if (static_cast<uint32>(Grid) < (static_cast<uint32>(EPCGHiGenGrid::GridMin) << HiGenExponential))
-	{
-		check(static_cast<uint32>(Grid) >= static_cast<uint32>(EPCGHiGenGrid::GridMin));
-		uint32 Multiplier = static_cast<uint32>(Grid) / static_cast<uint32>(EPCGHiGenGrid::GridMin);
-		return GenerationRadii.GetGenerationRadiusFromGrid(EPCGHiGenGrid::GridMin) * Multiplier;
-	}
-	else
-	{
-		EPCGHiGenGrid AdjustedGrid = static_cast<EPCGHiGenGrid>(static_cast<uint32>(Grid) >> HiGenExponential);
-		return GenerationRadii.GetGenerationRadiusFromGrid(AdjustedGrid) * (1ULL << HiGenExponential);
-	}
+	const EPCGHiGenGrid UnscaledGrid = PCGHiGenGrid::ScaledGridSizeToGrid(Grid, GetGridSizeMultiplier());
+	return GenerationRadii.GetGenerationRadiusFromGrid(UnscaledGrid);
 }
 
-double UPCGGraph::GetGridCleanupRadiusFromGrid(EPCGHiGenGrid Grid) const
+double UPCGGraph::GetGridCleanupRadiusFromGrid(uint32 Grid) const
 {
-	if (Grid == EPCGHiGenGrid::Unbounded || Grid == EPCGHiGenGrid::Uninitialized)
-	{
-		return GenerationRadii.GetCleanupRadiusFromGrid(Grid);
-	}
-	// If the queried grid is smaller than the min grid including the exponent, we'll take the min grid and scale it down
-	else if (static_cast<uint32>(Grid) < (static_cast<uint32>(EPCGHiGenGrid::GridMin) << HiGenExponential))
-	{
-		check(static_cast<uint32>(Grid) >= static_cast<uint32>(EPCGHiGenGrid::GridMin));
-		uint32 Multiplier = static_cast<uint32>(Grid) / static_cast<uint32>(EPCGHiGenGrid::GridMin);
-		return GenerationRadii.GetCleanupRadiusFromGrid(EPCGHiGenGrid::GridMin) * Multiplier;
-	}
-	else
-	{
-		EPCGHiGenGrid AdjustedGrid = static_cast<EPCGHiGenGrid>(static_cast<uint32>(Grid) >> HiGenExponential);
-		return GenerationRadii.GetCleanupRadiusFromGrid(AdjustedGrid) * (1ULL << HiGenExponential);
-	}
+	const EPCGHiGenGrid UnscaledGrid = PCGHiGenGrid::ScaledGridSizeToGrid(Grid, GetGridSizeMultiplier());
+	return GenerationRadii.GetCleanupRadiusFromGrid(UnscaledGrid);
 }
 
 bool UPCGGraph::UseActorComponentlessGeneration() const
@@ -2157,6 +2137,7 @@ void UPCGGraph::PostEditChangeChainProperty(FPropertyChangedChainEvent& Property
 	}
 
 	const FName MemberPropertyName = MemberProperty->GetFName();
+	bool bHiGenSettingsChanged = false;
 
 	if (MemberPropertyName == GET_MEMBER_NAME_CHECKED(UPCGGraph, bLandscapeUsesMetadata) ||
 		MemberPropertyName == GET_MEMBER_NAME_CHECKED(UPCGGraph, bHasDefaultConstructedInputs))
@@ -2234,11 +2215,15 @@ void UPCGGraph::PostEditChangeChainProperty(FPropertyChangedChainEvent& Property
 	else if (MemberPropertyName == GET_MEMBER_NAME_CHECKED(UPCGGraph, HiGenGridSize)
 		|| MemberPropertyName == GET_MEMBER_NAME_CHECKED(UPCGGraph, bUseHierarchicalGeneration)
 		|| MemberPropertyName == GET_MEMBER_NAME_CHECKED(UPCGGraph, bUse2DGrid)
-		|| MemberPropertyName == GET_MEMBER_NAME_CHECKED(UPCGGraph, HiGenExponential)
 		|| MemberPropertyName == GET_MEMBER_NAME_CHECKED(UPCGGraph, bUseActorComponentlessGeneration))
 	{
-		// The higen settings change the structure of the graph (presence or absence of links between grid levels).
-		NotifyGraphChanged(EPCGChangeType::Structural | EPCGChangeType::GenerationGrid);
+		bHiGenSettingsChanged = true;
+	}
+	else if (MemberPropertyName == GET_MEMBER_NAME_CHECKED(UPCGGraph, HiGenGridSizeMultiplier))
+	{
+		HiGenGridSizeMultiplier = PCGHiGenGrid::SanitizeGridSizeMultiplier(HiGenGridSizeMultiplier);
+
+		bHiGenSettingsChanged = true;
 	}
 	else if (MemberPropertyName == GET_MEMBER_NAME_CHECKED(UPCGGraphInterface, Title)
 		|| MemberPropertyName == GET_MEMBER_NAME_CHECKED(UPCGGraphInterface, bOverrideTitle)
@@ -2263,6 +2248,12 @@ void UPCGGraph::PostEditChangeChainProperty(FPropertyChangedChainEvent& Property
 		// Handled in the function
 	}
 
+	if (bHiGenSettingsChanged)
+	{
+		// The higen settings change the structure of the graph (presence or absence of links between grid levels).
+		NotifyGraphChanged(EPCGChangeType::Structural | EPCGChangeType::GenerationGrid);
+	}
+
 	PreviousPropertyBag = nullptr;
 }
 
@@ -2483,7 +2474,7 @@ uint32 UPCGGraph::CalculateNodeGridSizeRecursive_Unsafe(const UPCGNode* InNode,
 	const UPCGHiGenGridSizeSettings* GridSizeSettings = Cast<UPCGHiGenGridSizeSettings>(InNode->GetSettings());
 	if (GridSizeSettings && GridSizeSettings->bEnabled)
 	{
-		GridSize = FMath::Min(GridSize, GridSizeSettings->GetGridSize());
+		GridSize = FMath::Min(GridSize, GridSizeSettings->GetScaledGridSize());
 	}
 	else
 	{
diff --git a/Engine/Plugins/PCG/Source/PCG/Private/RuntimeGen/PCGRuntimeGenExecutionSource.cpp b/Engine/Plugins/PCG/Source/PCG/Private/RuntimeGen/PCGRuntimeGenExecutionSource.cpp
--- a/Engine/Plugins/PCG/Source/PCG/Private/RuntimeGen/PCGRuntimeGenExecutionSource.cpp
+++ b/Engine/Plugins/PCG/Source/PCG/Private/RuntimeGen/PCGRuntimeGenExecutionSource.cpp
@@ -484,21 +484,6 @@ bool UPCGRuntimeGenExecutionSource::Use2DGrid() const
 	return OwningSource->GetExecutionState().Use2DGrid();
 }
 
-EPCGHiGenGrid UPCGRuntimeGenExecutionSource::GetGenerationGrid() const
-{
-	if (PCGHiGenGrid::IsValidGridSize(GridSize))
-	{
-		return PCGHiGenGrid::GridSizeToGrid(GridSize);
-	}
-	else if (GridSize == PCGHiGenGrid::UnboundedGridSize())
-	{
-		return EPCGHiGenGrid::Unbounded;
-	}
-	else
-	{
-		return PCGHiGenGrid::GridSizeToGrid(GetGraph()->GetDefaultGridSize());
-	}
-}
 
 uint32 UPCGRuntimeGenExecutionSource::GetGenerationGridSize() const
 {
@@ -529,13 +514,13 @@ const FPCGRuntimeGenerationRadii& UPCGRuntimeGenExecutionSource::GetGenerationRa
 	return OwningSource->GetExecutionState().GetGenerationRadii();
 }
 
-double UPCGRuntimeGenExecutionSource::GetGenerationRadiusFromGrid(EPCGHiGenGrid InGrid) const
+double UPCGRuntimeGenExecutionSource::GetGenerationRadiusFromGrid(uint32 InGrid) const
 {
 	check(OwningSource);
 	return OwningSource->GetExecutionState().GetGenerationRadiusFromGrid(InGrid);
 }
 
-double UPCGRuntimeGenExecutionSource::GetCleanupRadiusFromGrid(EPCGHiGenGrid InGrid) const
+double UPCGRuntimeGenExecutionSource::GetCleanupRadiusFromGrid(uint32 InGrid) const
 {
 	check(OwningSource);
 	return OwningSource->GetExecutionState().GetCleanupRadiusFromGrid(InGrid);
diff --git a/Engine/Plugins/PCG/Source/PCG/Private/RuntimeGen/PCGRuntimeGenExecutionState.cpp b/Engine/Plugins/PCG/Source/PCG/Private/RuntimeGen/PCGRuntimeGenExecutionState.cpp
--- a/Engine/Plugins/PCG/Source/PCG/Private/RuntimeGen/PCGRuntimeGenExecutionState.cpp
+++ b/Engine/Plugins/PCG/Source/PCG/Private/RuntimeGen/PCGRuntimeGenExecutionState.cpp
@@ -276,11 +276,6 @@ bool FPCGRuntimeGenExecutionState::Use2DGrid() const
 	return Source->Use2DGrid();
 }
 
-EPCGHiGenGrid FPCGRuntimeGenExecutionState::GetGenerationGrid() const
-{
-	check(Source);
-	return Source->GetGenerationGrid();
-}
 
 uint32 FPCGRuntimeGenExecutionState::GetGenerationGridSize() const
 {
@@ -312,13 +307,13 @@ const FPCGRuntimeGenerationRadii& FPCGRuntimeGenExecutionState::GetGenerationRad
 	return Source->GetGenerationRadii();
 }
 
-double FPCGRuntimeGenExecutionState::GetGenerationRadiusFromGrid(EPCGHiGenGrid InGrid) const
+double FPCGRuntimeGenExecutionState::GetGenerationRadiusFromGrid(uint32 InGrid) const
 {
 	check(Source);
 	return Source->GetGenerationRadiusFromGrid(InGrid);
 }
 
-double FPCGRuntimeGenExecutionState::GetCleanupRadiusFromGrid(EPCGHiGenGrid InGrid) const
+double FPCGRuntimeGenExecutionState::GetCleanupRadiusFromGrid(uint32 InGrid) const
 {
 	check(Source);
 	return Source->GetCleanupRadiusFromGrid(InGrid);
diff --git a/Engine/Plugins/PCG/Source/PCG/Private/RuntimeGen/PCGRuntimeGenScheduler.cpp b/Engine/Plugins/PCG/Source/PCG/Private/RuntimeGen/PCGRuntimeGenScheduler.cpp
--- a/Engine/Plugins/PCG/Source/PCG/Private/RuntimeGen/PCGRuntimeGenScheduler.cpp
+++ b/Engine/Plugins/PCG/Source/PCG/Private/RuntimeGen/PCGRuntimeGenScheduler.cpp
@@ -622,7 +622,7 @@ void FPCGRuntimeGenScheduler::TickQueueSourcesForGeneration(
 				WorldStreamingQueryGridSizes.SetNumZeroed(GridSizes.Num());
 			}
 
-			const EPCGHiGenGrid MaxGrid = bHasUnbounded ? EPCGHiGenGrid::Unbounded : PCGHiGenGrid::GridSizeToGrid(GridSizes[0]);
+			const uint32 MaxGrid = bHasUnbounded ? PCGHiGenGrid::UnboundedGridSize() : GridSizes[0];
 			const double MaxGenerationRadius = OriginalExecutionState.GetGenerationRadiusFromGrid(MaxGrid);
 
 			for (const IPCGGenSourceBase* GenSource : *Inputs.GenSources)
@@ -671,7 +671,7 @@ void FPCGRuntimeGenScheduler::TickQueueSourcesForGeneration(
 					{
 						ensure(PCGHiGenGrid::IsValidGridSize(GridSize));
 
-						const double GenerationRadius = OriginalExecutionState.GetGenerationRadiusFromGrid(PCGHiGenGrid::GridSizeToGrid(GridSize));
+						const double GenerationRadius = OriginalExecutionState.GetGenerationRadiusFromGrid(GridSize);
 
 						DrawDebugSphere(Context.World, GenSourcePosition, GenerationRadius, 64, FColor::Red, /*bPersistentLines=*/false, /*LifeTime=*/0.0f);
 					}
@@ -734,7 +734,7 @@ void FPCGRuntimeGenScheduler::TickQueueSourcesForGeneration(
 					const uint32 GridSize = GridSizes[GridIndex];
 
 					const FIntVector GenSourceGridPosition = UPCGActorHelpers::GetCellCoord(GenSourcePosition, GridSize, bIs2DGrid);
-					const double GenerationRadius = OriginalExecutionState.GetGenerationRadiusFromGrid(PCGHiGenGrid::GridSizeToGrid(GridSize));
+					const double GenerationRadius = OriginalExecutionState.GetGenerationRadiusFromGrid(GridSize);
 					const int32 GridRadius = FMath::CeilToInt32(GenerationRadius / GridSize); // Radius discretized to # of grid cells.
 					const int32 VerticalGridRadius = bIs2DGrid ? 0 : GridRadius; // Flatten the vertical grid radius in the 2D case.
 
@@ -894,7 +894,7 @@ void FPCGRuntimeGenScheduler::TickQueueSourcesForGeneration(
 			}
 
 			// Unbounded will grab the base GenerationRadius used for non-partitioned and unbounded.
-			const double MaxGenerationRadius = OriginalExecutionState.GetGenerationRadiusFromGrid(EPCGHiGenGrid::Unbounded);
+			const double MaxGenerationRadius = OriginalExecutionState.GetGenerationRadiusFromGrid(PCGHiGenGrid::UnboundedGridSize());
 
 			for (const IPCGGenSourceBase* GenSource : *Inputs.GenSources)
 			{
@@ -1006,7 +1006,6 @@ void FPCGRuntimeGenScheduler::TickCleanup(const TSet<IPCGGenSourceBase*>& InGenS
 			}
 
 			const FPCGGridDescriptor GridDescriptor = GenerationKey.GetGridDescriptor();
-			const EPCGHiGenGrid Grid = GridDescriptor.GetHiGenGrid();
 			const FIntVector& GridCoords = GenerationKey.GetGridCoords();
 			IPCGGraphExecutionSource* OriginalSource = GenerationKey.GetOriginalSource();
 			check(OriginalSource);
@@ -1014,11 +1013,11 @@ void FPCGRuntimeGenScheduler::TickCleanup(const TSet<IPCGGenSourceBase*>& InGenS
 			const UPCGSchedulingPolicyBase* Policy = OriginalSource->GetExecutionState().GetRuntimeGenSchedulingPolicy();
 			ensure(Policy);
 
-			const double CleanupRadius = OriginalSource->GetExecutionState().GetCleanupRadiusFromGrid(Grid);
+			const double CleanupRadius = OriginalSource->GetExecutionState().GetCleanupRadiusFromGrid(GridDescriptor.GetGridSize());
 			const double CleanupRadiusSquared = CleanupRadius * CleanupRadius;
 
 			// If the Grid is unbounded, we have a non-partitioned or unbounded execution source.
-			if (Grid == EPCGHiGenGrid::Unbounded)
+			if (GridDescriptor.IsUnboundedGrid())
 			{
 				if (!OriginalSource->GetExecutionState().IsActive())
 				{
@@ -1128,7 +1127,7 @@ void FPCGRuntimeGenScheduler::TickScheduleGeneration(TMap<FGridGenerationKey, do
 				continue;
 			}
 
-			const IPCGGraphExecutionSource* ExecutionSource = (Key.GetGridDescriptor().GetHiGenGrid() == EPCGHiGenGrid::Unbounded) ? Key.GetOriginalSource() : Key.GetCachedLocalSource();
+			const IPCGGraphExecutionSource* ExecutionSource = (Key.GetGridDescriptor().GetGridSize() == PCGHiGenGrid::UnboundedGridSize()) ? Key.GetOriginalSource() : Key.GetCachedLocalSource();
 
 			if (ExecutionSource && ExecutionSource->GetExecutionState().IsGenerating())
 			{
@@ -1169,7 +1168,6 @@ void FPCGRuntimeGenScheduler::TickScheduleGeneration(TMap<FGridGenerationKey, do
 		const double Priority = It.Value();
 
 		const FPCGGridDescriptor GridDescriptor = Key.GetGridDescriptor();
-		const EPCGHiGenGrid Grid = GridDescriptor.GetHiGenGrid();
 
 		const FIntVector GridCoords = Key.GetGridCoords();
 		IPCGGraphExecutionSource* OriginalSource = Key.GetOriginalSource();
@@ -1182,7 +1180,7 @@ void FPCGRuntimeGenScheduler::TickScheduleGeneration(TMap<FGridGenerationKey, do
 		IPCGGraphExecutionState& OriginalExecutionState = OriginalSource->GetExecutionState();
 
 		// If the Grid is unbounded, we have a non-partitioned or unbounded execution source.
-		if (Grid == EPCGHiGenGrid::Unbounded)
+		if (GridDescriptor.IsUnboundedGrid())
 		{
 			if (!OriginalExecutionState.IsGenerating())
 			{
@@ -1383,6 +1381,7 @@ void FPCGRuntimeGenScheduler::TickRequestVirtualTexturePriming(const TSet<IPCGGe
 
 		const UPCGComponent* OriginalComponent = OriginalComponentKey.ResolveObjectPtr();
 		check(OriginalComponent);
+		const UPCGGraph* Graph = OriginalComponent->GetGraph();
 
 		const FBox OriginalComponentBounds = OriginalComponent->GetGridBounds();
 
@@ -1404,7 +1403,9 @@ void FPCGRuntimeGenScheduler::TickRequestVirtualTexturePriming(const TSet<IPCGGe
 			{
 				check(VirtualTextureComponent);
 
-				const double PrimingRadius = OriginalComponent->GetGenerationRadiusFromGrid(PrimingInfo.Grid) + PCGHiGenGrid::GridToGridSize(PrimingInfo.Grid);
+				const uint32 GridSize = PCGHiGenGrid::GridToGridSize(PrimingInfo.Grid) * (Graph ? Graph->GetGridSizeMultiplier() : 1.0);
+
+				const double PrimingRadius = OriginalComponent->GetGenerationRadiusFromGrid(GridSize) + GridSize;
 
 				for (const IPCGGenSourceBase* GenSource : InGenSources)
 				{
@@ -1710,14 +1711,13 @@ void FPCGRuntimeGenScheduler::CleanupLocalSources(const APCGWorldActor* InPCGWor
 		}
 
 		const FPCGGridDescriptor GridDescriptor = GenerationKey.GetGridDescriptor();
-		const EPCGHiGenGrid Grid = GridDescriptor.GetHiGenGrid();
 		const FIntVector& GridCoords = GenerationKey.GetGridCoords();
 		IPCGGraphExecutionSource* OriginalSource = GenerationKey.GetOriginalSource();
 		check(OriginalSource);
 		OriginalSources.Add(OriginalSource);
 
 		// Only operate on LocalSources.
-		if (Grid != EPCGHiGenGrid::Unbounded)
+		if (!GridDescriptor.IsUnboundedGrid())
 		{
 			IPCGGraphExecutionSource* LocalSource = OriginalSource->GetExecutionState().GetLocalSource(GridDescriptor, GridCoords);
 			SourcesToClean.Add({ GenerationKey, LocalSource });
@@ -1830,7 +1830,6 @@ void FPCGRuntimeGenScheduler::CleanupSource(const FGridGenerationKey& Generation
 	check(TrackingManager);
 
 	const FPCGGridDescriptor GridDescriptor = GenerationKey.GetGridDescriptor();
-	const EPCGHiGenGrid Grid = GridDescriptor.GetHiGenGrid();
 	const FIntVector GridCoords = GenerationKey.GetGridCoords();
 
 	APCGPartitionActor* PartitionActor = nullptr;
@@ -1846,7 +1845,7 @@ void FPCGRuntimeGenScheduler::CleanupSource(const FGridGenerationKey& Generation
 	{
 		PartitionActor = InGeneratedSource->GetExecutionState().GetTypedTarget<APCGPartitionActor>();
 
-		if (Grid == EPCGHiGenGrid::Unbounded)
+		if (GridDescriptor.IsUnboundedGrid())
 		{
 			if (PCGRuntimeGenSchedulerHelpers::CVarRuntimeGenerationEnableDebugging.GetValueOnAnyThread())
 			{
@@ -1897,12 +1896,11 @@ void FPCGRuntimeGenScheduler::CleanupDelayedRefreshSources()
 		}
 
 		const FPCGGridDescriptor GridDescriptor = GenerationKey.GetGridDescriptor();
-		const EPCGHiGenGrid Grid = GridDescriptor.GetHiGenGrid();
 		IPCGGraphExecutionSource* OriginalSource = GenerationKey.GetOriginalSource();
 		const FIntVector& GridCoords = GenerationKey.GetGridCoords();
 
 		// The unbounded grid level will always lie inside the original source, so we can skip it.
-		if (Grid == EPCGHiGenGrid::Unbounded)
+		if (GridDescriptor.IsUnboundedGrid())
 		{
 			if (OriginalSource && !OriginalSource->GetExecutionState().IsActive())
 			{
@@ -2112,10 +2110,9 @@ void FPCGRuntimeGenScheduler::RefreshExecutionSource(IPCGGraphExecutionSource* I
 			for (FGridGenerationKey GenerationKey : GenerationKeys)
 			{
 				const FPCGGridDescriptor GridDescriptor = GenerationKey.GetGridDescriptor();
-				const EPCGHiGenGrid Grid = GridDescriptor.GetHiGenGrid();
 
-				// If the Grid is unbounded, we have a non-partitioned or unbounded execution source.
-				if (Grid == EPCGHiGenGrid::Unbounded)
+				// If the grid is unbounded, we have a non-partitioned or unbounded execution source.
+				if (GridDescriptor.IsUnboundedGrid())
 				{
 					if (bLoggingEnabled)
 					{
diff --git a/Engine/Plugins/PCG/Source/PCG/Private/Subsystems/PCGSubsystem.cpp b/Engine/Plugins/PCG/Source/PCG/Private/Subsystems/PCGSubsystem.cpp
--- a/Engine/Plugins/PCG/Source/PCG/Private/Subsystems/PCGSubsystem.cpp
+++ b/Engine/Plugins/PCG/Source/PCG/Private/Subsystems/PCGSubsystem.cpp
@@ -706,7 +706,7 @@ UPCGLandscapeCache* UPCGSubsystem::GetLandscapeCache()
 	return LandscapeCacheOwner ? LandscapeCacheOwner->LandscapeCacheObject.Get() : nullptr;
 }
 
-FPCGTaskId UPCGSubsystem::ScheduleComponent(UPCGComponent* PCGComponent, EPCGHiGenGrid Grid, bool bForce, const TArray<FPCGTaskId>& InDependencies)
+FPCGTaskId UPCGSubsystem::ScheduleComponent(UPCGComponent* PCGComponent, uint32 Grid, bool bForce, const TArray<FPCGTaskId>& InDependencies)
 {
 	check(GraphExecutor);
 
@@ -748,12 +748,12 @@ FPCGTaskId UPCGSubsystem::ScheduleComponent(UPCGComponent* PCGComponent, EPCGHiG
 	if (!PCGComponent->IsPartitioned())
 	{
 		// This component is either an unpartitioned original component or a local component. Generate if grid size matches preference (if provided).
-		bGeneratePCGComponent = (Grid == EPCGHiGenGrid::Uninitialized) || !!(Grid & PCGComponent->GetGenerationGrid());
+		bGeneratePCGComponent = (Grid == PCGHiGenGrid::UninitializedGridSize()) || (Grid == PCGComponent->GetGenerationGridSize());
 	}
 	else
 	{
 		// This component is a partitioned original component. Generate if the graph has unbounded nodes and if this grid matches preference (if provided).
-		bGeneratePCGComponent = bHasUnbounded && (!!(Grid & EPCGHiGenGrid::Unbounded) || Grid == EPCGHiGenGrid::Uninitialized);
+		bGeneratePCGComponent = bHasUnbounded && (Grid == PCGHiGenGrid::UnboundedGridSize() || Grid == PCGHiGenGrid::UninitializedGridSize());
 	}
 
 	if (bGeneratePCGComponent)
@@ -782,7 +782,7 @@ FPCGTaskId UPCGSubsystem::ScheduleComponent(UPCGComponent* PCGComponent, EPCGHiG
 				// Local component with invalid grid size. Grid sizes may have changed in graph.
 				return LocalComponent->CleanupLocal(/*bRemoveComponents=*/true, Dependencies);
 			}
-			else if (Grid != EPCGHiGenGrid::Uninitialized && !(Grid & LocalComponent->GetGenerationGrid()))
+			else if (Grid != PCGHiGenGrid::UninitializedGridSize() && Grid != LocalComponent->GetGenerationGridSize())
 			{
 				// Grid size does not match the given target grid, so skip.
 				return InvalidPCGTaskId;
@@ -826,7 +826,7 @@ FPCGTaskId UPCGSubsystem::ScheduleComponent(UPCGComponent* PCGComponent, EPCGHiG
 				AllDependencies = &AdditionalDependencies;
 			}
 
-			return LocalComponent->GenerateInternal(bForce, LocalComponent->GetGenerationGrid(), EPCGComponentGenerationTrigger::GenerateOnDemand, *AllDependencies);
+			return LocalComponent->GenerateInternal(bForce, LocalComponent->GetGenerationGridSize(), EPCGComponentGenerationTrigger::GenerateOnDemand, *AllDependencies);
 		};
 
 		ExecutionDependencyTasks.Append(TrackingManager->DispatchToRegisteredLocalComponents(PCGComponent, LocalGenerateTask));
@@ -864,6 +864,12 @@ FPCGTaskId UPCGSubsystem::ScheduleComponent(UPCGComponent* PCGComponent, EPCGHiG
 	}
 }
 
+// Deprecated 5.8
+FPCGTaskId UPCGSubsystem::ScheduleComponent(UPCGComponent* PCGComponent, EPCGHiGenGrid Grid, bool bForce, const TArray<FPCGTaskId>& InDependencies)
+{
+	return ScheduleComponent(PCGComponent, PCGHiGenGrid::GridToGridSize(Grid), bForce, InDependencies);
+}
+
 FPCGTaskId UPCGSubsystem::ScheduleCleanup(UPCGComponent* PCGComponent, bool bRemoveComponents, const TArray<FPCGTaskId>& Dependencies)
 {
 	if (!PCGComponent)
@@ -1157,7 +1163,7 @@ FPCGTaskId UPCGSubsystem::ForAllOverlappingCells(UPCGComponent* InComponent, con
 	PCGHiGenGrid::FSizeArray GridSizes = InGridSizes;
 
 	// We have no use for unbounded grids as this is a grid-centric function. Also discard invalid grid sizes.
-	GridSizes.Remove(static_cast<uint32>(EPCGHiGenGrid::Unbounded));
+	GridSizes.Remove(PCGHiGenGrid::UnboundedGridSize());
 	GridSizes.RemoveAll([](uint32 GridSize) { return !ensure(PCGHiGenGrid::IsValidGridSize(GridSize)); });
 
 	if (GridSizes.IsEmpty())
diff --git a/Engine/Plugins/PCG/Source/PCG/Private/Utils/PCGGraphExecutionLogging.cpp b/Engine/Plugins/PCG/Source/PCG/Private/Utils/PCGGraphExecutionLogging.cpp
--- a/Engine/Plugins/PCG/Source/PCG/Private/Utils/PCGGraphExecutionLogging.cpp
+++ b/Engine/Plugins/PCG/Source/PCG/Private/Utils/PCGGraphExecutionLogging.cpp
@@ -372,7 +372,7 @@ namespace PCGGraphExecutionLogging
 #endif // !(UE_BUILD_SHIPPING || UE_BUILD_TEST) || USE_LOGGING_IN_SHIPPING
 	}
 
-	void LogGridLinkageTaskExecuteStore(const FPCGContext* InContext, EPCGHiGenGrid InGenerationGrid, int32 InFromGridSize, int32 InToGridSize, const FString& InResourcePath, int32 InDataItemCount)
+	void LogGridLinkageTaskExecuteStore(const FPCGContext* InContext, uint32 InGenerationGrid, uint32 InFromGridSize, uint32 InToGridSize, const FString& InResourcePath, int32 InDataItemCount)
 	{
 #if !(UE_BUILD_SHIPPING || UE_BUILD_TEST) || USE_LOGGING_IN_SHIPPING
 		if (!LogEnabled())
@@ -383,7 +383,7 @@ namespace PCGGraphExecutionLogging
 
 		UE_LOG(LogPCG, Log, TEXT("[GRIDLINKING] [%s] STORE. GenerationGridSize=%u, FromGridSize=%u, ToGridSize=%u, Path=%s, DataItems=%d"),
 			*PCGLog::GetExecutionSourceName(InContext->ExecutionSource.Get(), /*bUseLabel=*/true),
-			PCGHiGenGrid::IsValidGrid(InGenerationGrid) ? PCGHiGenGrid::GridToGridSize(InGenerationGrid) : PCGHiGenGrid::UnboundedGridSize(),
+			InGenerationGrid,
 			InFromGridSize,
 			InToGridSize,
 			*InResourcePath,
@@ -391,7 +391,7 @@ namespace PCGGraphExecutionLogging
 #endif // !(UE_BUILD_SHIPPING || UE_BUILD_TEST) || USE_LOGGING_IN_SHIPPING
 	}
 
-	void LogGridLinkageTaskExecuteRetrieve(const FPCGContext* InContext, EPCGHiGenGrid InGenerationGrid, int32 InFromGridSize, int32 InToGridSize, const FString& InResourcePath)
+	void LogGridLinkageTaskExecuteRetrieve(const FPCGContext* InContext, uint32 InGenerationGrid, uint32 InFromGridSize, uint32 InToGridSize, const FString& InResourcePath)
 	{
 #if !(UE_BUILD_SHIPPING || UE_BUILD_TEST) || USE_LOGGING_IN_SHIPPING
 		if (!LogEnabled())
@@ -402,7 +402,7 @@ namespace PCGGraphExecutionLogging
 
 		UE_LOG(LogPCG, Log, TEXT("[GRIDLINKING] [%s] RETRIEVE. GenerationGridSize=%u, FromGridSize=%u, ToGridSize=%u, Path=%s"),
 			*PCGLog::GetExecutionSourceName(InContext->ExecutionSource.Get(), /*bUseLabel=*/true),
-			PCGHiGenGrid::IsValidGrid(InGenerationGrid) ? PCGHiGenGrid::GridToGridSize(InGenerationGrid) : PCGHiGenGrid::UnboundedGridSize(),
+			InGenerationGrid,
 			InFromGridSize,
 			InToGridSize,
 			*InResourcePath);
diff --git a/Engine/Plugins/PCG/Source/PCG/Public/Elements/PCGDataFromActor.h b/Engine/Plugins/PCG/Source/PCG/Public/Elements/PCGDataFromActor.h
--- a/Engine/Plugins/PCG/Source/PCG/Public/Elements/PCGDataFromActor.h
+++ b/Engine/Plugins/PCG/Source/PCG/Public/Elements/PCGDataFromActor.h
@@ -107,7 +107,7 @@ public:
 	UPROPERTY(BlueprintReadWrite, EditAnywhere, Category = "Data Retrieval Settings", meta = (EditCondition = "Mode == EPCGGetDataFromActorMode::GetDataFromPCGComponent || Mode == EPCGGetDataFromActorMode::GetDataFromPCGComponentOrParseComponents", EditConditionHides))
 	bool bGetDataOnAllGrids = true;
 
-	/** Select which grid sizes to consider when collecting data from partitioned PCG components. */
+	/** Select which grid sizes to consider when collecting data from partitioned PCG components (without grid scaling applied). */
 	UPROPERTY(BlueprintReadWrite, EditAnywhere, Category = "Data Retrieval Settings", meta = (Bitmask, BitmaskEnum = "/Script/PCG.EPCGHiGenGrid", EditCondition = "!bGetDataOnAllGrids && Mode == EPCGGetDataFromActorMode::GetDataFromPCGComponent || Mode == EPCGGetDataFromActorMode::GetDataFromPCGComponentOrParseComponents", EditConditionHides))
 	int32 AllowedGrids = int32(EPCGHiGenGrid::Uninitialized);
 
diff --git a/Engine/Plugins/PCG/Source/PCG/Public/Elements/PCGHiGenGridSize.h b/Engine/Plugins/PCG/Source/PCG/Public/Elements/PCGHiGenGridSize.h
--- a/Engine/Plugins/PCG/Source/PCG/Public/Elements/PCGHiGenGridSize.h
+++ b/Engine/Plugins/PCG/Source/PCG/Public/Elements/PCGHiGenGridSize.h
@@ -15,8 +15,7 @@ class UPCGHiGenGridSizeSettings : public UPCGSettings
 	GENERATED_BODY()
 
 public:
-	EPCGHiGenGrid GetGrid() const;
-	uint32 GetGridSize() const;
+	uint32 GetScaledGridSize() const;
 
 	bool ShouldOutputCellVolumeData() const;
 
diff --git a/Engine/Plugins/PCG/Source/PCG/Public/Grid/PCGGridDescriptor.h b/Engine/Plugins/PCG/Source/PCG/Public/Grid/PCGGridDescriptor.h
--- a/Engine/Plugins/PCG/Source/PCG/Public/Grid/PCGGridDescriptor.h
+++ b/Engine/Plugins/PCG/Source/PCG/Public/Grid/PCGGridDescriptor.h
@@ -48,8 +48,8 @@ private:
 public:
 	FPCGGridDescriptor& SetGridSize(uint32 InGridSize) { check(!Hash.IsSet()); GridSize = InGridSize; return *this; }
 	uint32 GetGridSize() const { return GridSize; }
-		
-	EPCGHiGenGrid GetHiGenGrid() const { return PCGHiGenGrid::GridSizeToGrid(GetGridSize()); }
+
+	bool IsUnboundedGrid() const { return GetGridSize() == PCGHiGenGrid::UnboundedGridSize(); }
 
 	FPCGGridDescriptor& SetIs2DGrid(bool bInIs2DGrid) { check(!Hash.IsSet()); bIs2DGrid = bInIs2DGrid; return *this; }
 	bool Is2DGrid() const { return bIs2DGrid; }
@@ -90,6 +90,11 @@ private:
 	uint32 ComputeHash() const;
 
 	mutable TOptional<uint32> Hash;
+
+	// Deprecated section
+public:
+	UE_DEPRECATED(5.8, "Use GetGridSize")
+	EPCGHiGenGrid GetHiGenGrid() const { return PCGHiGenGrid::GridSizeToGrid(GetGridSize()); }
 };
 
 /**
diff --git a/Engine/Plugins/PCG/Source/PCG/Public/PCGCommon.h b/Engine/Plugins/PCG/Source/PCG/Public/PCGCommon.h
--- a/Engine/Plugins/PCG/Source/PCG/Public/PCGCommon.h
+++ b/Engine/Plugins/PCG/Source/PCG/Public/PCGCommon.h
@@ -555,12 +555,15 @@ namespace PCGHiGenGrid
 	PCG_API bool IsValidGridSize(uint32 InGridSize);
 	PCG_API bool IsValidGrid(EPCGHiGenGrid InGrid);
 	PCG_API bool IsValidGridOrUninitialized(EPCGHiGenGrid InGrid);
+	PCG_API bool IsValidGridOrUninitialized(uint32 InGridSize);
 	PCG_API uint32 GridToGridSize(EPCGHiGenGrid InGrid);
 	/** Rounds up to the nearest power of 2 grid, clamping between GridMin and GridMax (excepting the unbounded grid, which may exceed GridMax). */
 	PCG_API EPCGHiGenGrid GridSizeToGrid(uint32 InGridSize);
+	PCG_API EPCGHiGenGrid ScaledGridSizeToGrid(uint32 InScaledGridSize, double InGridScale);
 
 	PCG_API uint32 UninitializedGridSize();
 	PCG_API uint32 UnboundedGridSize();
+	PCG_API double SanitizeGridSizeMultiplier(double GridSizeMultiplier);
 }
 
 UENUM()
diff --git a/Engine/Plugins/PCG/Source/PCG/Public/PCGComponent.h b/Engine/Plugins/PCG/Source/PCG/Public/PCGComponent.h
--- a/Engine/Plugins/PCG/Source/PCG/Public/PCGComponent.h
+++ b/Engine/Plugins/PCG/Source/PCG/Public/PCGComponent.h
@@ -219,11 +219,11 @@ public:
 	UE_API void GenerateLocal(bool bForce);
 
 	/** Requests the component to generate only on the specified grid level (all grid levels if EPCGHiGenGrid::Uninitialized). */
-	UE_API void GenerateLocal(EPCGComponentGenerationTrigger RequestedGenerationTrigger, bool bForce, EPCGHiGenGrid Grid = EPCGHiGenGrid::Uninitialized, const TArray<FPCGTaskId>& Dependencies = {});
+	UE_API void GenerateLocal(EPCGComponentGenerationTrigger RequestedGenerationTrigger, bool bForce, uint32 Grid = PCGHiGenGrid::UninitializedGridSize(), const TArray<FPCGTaskId>& Dependencies = {});
 
 	UE_API FPCGTaskId GenerateLocalGetTaskId(bool bForce);
-	UE_API FPCGTaskId GenerateLocalGetTaskId(EPCGComponentGenerationTrigger RequestedGenerationTrigger, bool bForce, EPCGHiGenGrid Grid = EPCGHiGenGrid::Uninitialized);
-	UE_API FPCGTaskId GenerateLocalGetTaskId(EPCGComponentGenerationTrigger RequestedGenerationTrigger, bool bForce, EPCGHiGenGrid Grid, const TArray<FPCGTaskId>& Dependencies);
+	UE_API FPCGTaskId GenerateLocalGetTaskId(EPCGComponentGenerationTrigger RequestedGenerationTrigger, bool bForce, uint32 Grid = PCGHiGenGrid::UninitializedGridSize());
+	UE_API FPCGTaskId GenerateLocalGetTaskId(EPCGComponentGenerationTrigger RequestedGenerationTrigger, bool bForce, uint32 Grid, const TArray<FPCGTaskId>& Dependencies);
 
 	/** Cleans up the generation from a local (vs. remote) standpoint. Will not be replicated. Will be delayed. */
 	UFUNCTION(BlueprintCallable, Category = PCG)
@@ -274,7 +274,9 @@ public:
 
 	uint32 GetGenerationGridSize() const { return GenerationGridSize; }
 	void SetGenerationGridSize(uint32 InGenerationGridSize) { GenerationGridSize = InGenerationGridSize; }
-	UE_API EPCGHiGenGrid GetGenerationGrid() const;
+
+	/** Get the grid that this component is executing on without grid scaling applied. */
+	EPCGHiGenGrid GetUnscaledGenerationGrid() const;
 
 	/** Store data with a resource key that identifies the pin. */
 	UE_API void StoreOutputDataForPin(const FString& InResourceKey, const FPCGDataCollection& InData);
@@ -291,11 +293,11 @@ public:
 	/** Get the generation radii that are currently active for this component. */
 	UE_API const FPCGRuntimeGenerationRadii& GetGenerationRadii() const;
 
-	/** Get the runtime generation radius for the given grid size. */
-	UE_API double GetGenerationRadiusFromGrid(EPCGHiGenGrid Grid) const;
+	/** Get the runtime generation radius for the given (scaled) grid size. */
+	UE_API double GetGenerationRadiusFromGrid(uint32 Grid) const;
 
-	/** Compute the runtime cleanup radius for the given grid size. */
-	UE_API double GetCleanupRadiusFromGrid(EPCGHiGenGrid Grid) const;
+	/** Compute the runtime cleanup radius for the given (scaled) grid size. */
+	UE_API double GetCleanupRadiusFromGrid(uint32 Grid) const;
 
 	/** Allows executing runtime PCG graphs without partition actors and local PCG components. */
 	UE_API bool UseActorComponentlessGeneration() const;
@@ -672,7 +674,7 @@ private:
 	UE_API bool ShouldGenerate(bool bForce, EPCGComponentGenerationTrigger RequestedGenerationTrigger) const;
 
 	/* Internal call that allows to delay a Generate/Cleanup call, chain with dependencies and keep track of the task id created. This task id is also returned. */
-	UE_API FPCGTaskId GenerateInternal(bool bForce, EPCGHiGenGrid Grid, EPCGComponentGenerationTrigger RequestedGenerationTrigger, const TArray<FPCGTaskId>& Dependencies);
+	UE_API FPCGTaskId GenerateInternal(bool bForce, uint32 Grid, EPCGComponentGenerationTrigger RequestedGenerationTrigger, const TArray<FPCGTaskId>& Dependencies);
 
 	/* Internal call to create tasks to generate the component. If there is nothing to do, an invalid task id will be returned. Should only be used by the subsystem. */
 	UE_API FPCGTaskId CreateGenerateTask(bool bForce, const TArray<FPCGTaskId>& Dependencies);
@@ -843,6 +845,21 @@ public:
 
 private:
 	FPCGComponentExecutionState ExecutionState;
+
+	// Deprecated section
+public:
+	UE_DEPRECATED(5.8, "Use the version that takes uint32 grid size instead")
+	UE_API void GenerateLocal(EPCGComponentGenerationTrigger RequestedGenerationTrigger, bool bForce, EPCGHiGenGrid Grid = EPCGHiGenGrid::Uninitialized, const TArray<FPCGTaskId>& Dependencies = {});
+	UE_DEPRECATED(5.8, "Use the version that takes uint32 grid size instead")
+	UE_API FPCGTaskId GenerateLocalGetTaskId(EPCGComponentGenerationTrigger RequestedGenerationTrigger, bool bForce, EPCGHiGenGrid Grid);
+	UE_DEPRECATED(5.8, "Use the version that takes uint32 grid size instead")
+	UE_API FPCGTaskId GenerateLocalGetTaskId(EPCGComponentGenerationTrigger RequestedGenerationTrigger, bool bForce, EPCGHiGenGrid Grid, const TArray<FPCGTaskId>& Dependencies);
+	UE_DEPRECATED(5.8, "Use GetGenerationGridSize")
+	EPCGHiGenGrid GetGenerationGrid() const { return PCGHiGenGrid::GridSizeToGrid(GetGenerationGridSize()); }
+	UE_DEPRECATED(5.8, "Use version that takes uint32 grid size instead")
+	double GetGenerationRadiusFromGrid(EPCGHiGenGrid Grid) const { return GetGenerationRadiusFromGrid(PCGHiGenGrid::GridToGridSize(Grid)); }
+	UE_DEPRECATED(5.8, "Use version that takes uint32 grid size instead")
+	double GetCleanupRadiusFromGrid(EPCGHiGenGrid Grid) const { return GetCleanupRadiusFromGrid(PCGHiGenGrid::GridToGridSize(Grid)); }
 };
 
 USTRUCT()
diff --git a/Engine/Plugins/PCG/Source/PCG/Public/PCGComponentExecutionState.h b/Engine/Plugins/PCG/Source/PCG/Public/PCGComponentExecutionState.h
--- a/Engine/Plugins/PCG/Source/PCG/Public/PCGComponentExecutionState.h
+++ b/Engine/Plugins/PCG/Source/PCG/Public/PCGComponentExecutionState.h
@@ -36,14 +36,13 @@ public:
 	virtual bool IsManagedByRuntimeGenSystem() const override;
 	virtual const UPCGSchedulingPolicyBase* GetRuntimeGenSchedulingPolicy() const override;
 	virtual const FPCGRuntimeGenerationRadii& GetGenerationRadii() const override;
-	virtual double GetGenerationRadiusFromGrid(EPCGHiGenGrid InGrid) const override;
-	virtual double GetCleanupRadiusFromGrid(EPCGHiGenGrid InGrid) const override;
+	virtual double GetGenerationRadiusFromGrid(uint32 InGrid) const override;
+	virtual double GetCleanupRadiusFromGrid(uint32 InGrid) const override;
 	virtual IPCGGraphExecutionSource* GetOriginalSource() const override;
 	virtual IPCGGraphExecutionSource* GetLocalSource(const FPCGGridDescriptor& GridDescriptor, const FIntVector& CellCoords) const override;
 	virtual bool IsPartitioned() const override;
 	virtual bool IsLocalSource() const override;
 	virtual FPCGGridDescriptor GetGridDescriptor(uint32 InGridSize) const override;
-	virtual EPCGHiGenGrid GetGenerationGrid() const override;
 	virtual uint32 GetGenerationGridSize() const override;
 	virtual FIntVector GetGenerationGridCoords() const override;
 	virtual bool IsInPreviewMode() const override;
diff --git a/Engine/Plugins/PCG/Source/PCG/Public/PCGGraph.h b/Engine/Plugins/PCG/Source/PCG/Public/PCGGraph.h
--- a/Engine/Plugins/PCG/Source/PCG/Public/PCGGraph.h
+++ b/Engine/Plugins/PCG/Source/PCG/Public/PCGGraph.h
@@ -364,9 +364,12 @@ public:
 	virtual const UPCGGraph* GetGraph() const override { return this; }
 	/** ~End UPCGGraphInterface interface */
 
-	// Default grid size for generation. For hierarchical generation, nodes outside of grid size graph ranges will generate on this grid.
-	UE_API EPCGHiGenGrid GetDefaultGrid() const;
+	/** Default grid for generation.For hierarchical generation, nodes outside of grid size graph ranges will generate on this grid. */
+	EPCGHiGenGrid GetDefaultGrid() const;
+
+	/** Default grid size for generation.For hierarchical generation, nodes outside of grid size graph ranges will generate on this grid. */
 	UE_API uint32 GetDefaultGridSize() const;
+
 	UE_API bool IsHierarchicalGenerationEnabled() const;
 	UE_API bool Use2DGrid() const;
 
@@ -492,14 +495,14 @@ public:
 	/** Returns all parent grid sizes for the given child grid size, calculated by inspecting nodes. */
 	UE_API void GetParentGridSizes(const uint32 InChildGridSize, PCGHiGenGrid::FSizeArray& OutParentGridSizes) const;
 
-	/** Returns exponential on grid size, which represents a shift in the grid */
-	uint32 GetGridExponential() const { return HiGenExponential; }
+	/** Returns multiplier on grid size. */
+	double GetGridSizeMultiplier() const { return PCGHiGenGrid::SanitizeGridSizeMultiplier(HiGenGridSizeMultiplier); }
 
-	/** Gets generation radius from grid, considering grid exponential. */
-	UE_API double GetGridGenerationRadiusFromGrid(EPCGHiGenGrid Grid) const;
+	/** Gets generation radius from grid, considering grid size multiplier. */
+	UE_API double GetGridGenerationRadiusFromGrid(uint32 Grid) const;
 
-	/** Gets cleanup radius from grid, considering grid exponential. */
-	UE_API double GetGridCleanupRadiusFromGrid(EPCGHiGenGrid Grid) const;
+	/** Gets cleanup radius from grid, considering grid size multiplier. */
+	UE_API double GetGridCleanupRadiusFromGrid(uint32 Grid) const;
 
 	/** Allows executing runtime PCG graphs without partition actors and local PCG components. */
 	bool UseActorComponentlessGeneration() const;
@@ -661,9 +664,9 @@ protected:
 	UPROPERTY(EditAnywhere, Category = Settings, meta = (DisplayName = "HiGen Default Grid Size", EditCondition = "CanEditHierarchicalGeneration() && bUseHierarchicalGeneration", EditConditionHides))
 	EPCGHiGenGrid HiGenGridSize = EPCGHiGenGrid::Grid256;
 
-	/** Shifts the grid sizes upwards based on the value, which allows to use larger grids. A value of 1 will effectively use the graph's Grid-400 values x 2 for the actual Grid-800 sizes and so on. */
-	UPROPERTY(EditAnywhere, Category = Settings, meta = (ClampMax = "10", DisplayName = "HiGen Grid Size Exponent", EditCondition = "SupportHierarchicalGeneration && bUseHierarchicalGeneration", EditConditionHides))
-	uint32 HiGenExponential = 0;
+	/** Scales the grid sizes allowing to use larger or smaller grids than those provided out of the box. Can be useful to align generation to other grids in the world. */
+	UPROPERTY(EditAnywhere, Category = Settings, meta = (DisplayName = "HiGen Grid Size Multiplier", EditCondition = "SupportHierarchicalGeneration && bUseHierarchicalGeneration", EditConditionHides))
+	double HiGenGridSizeMultiplier = 1.0;
 
 	UPROPERTY(EditAnywhere, Category = Settings, meta = (DisplayName = "2D Grid", EditCondition = "CanEditHierarchicalGeneration() && bUseHierarchicalGeneration", EditConditionHides))
 	bool bUse2DGrid = true;
@@ -766,6 +769,20 @@ private:
 	bool bIsInspecting = false;
 	FPCGStack InspectedStack;
 #endif // WITH_EDITOR
+
+	// Deprecated section
+public:
+	UE_DEPRECATED(5.8, "Use version that takes uint32 grid size instead")
+	double GetGridGenerationRadiusFromGrid(EPCGHiGenGrid Grid) const { return GetGridGenerationRadiusFromGrid(PCGHiGenGrid::GridToGridSize(Grid)); }
+	UE_DEPRECATED(5.8, "Use version that takes uint32 grid size instead")
+	double GetGridCleanupRadiusFromGrid(EPCGHiGenGrid Grid) const { return GetGridCleanupRadiusFromGrid(PCGHiGenGrid::GridToGridSize(Grid)); }
+	UE_DEPRECATED(5.8, "Use GetGridSizeMultiplier instead")
+	uint32 GetGridExponential() const { return FMath::RoundToInt(FMath::Log2(FMath::Max(HiGenGridSizeMultiplier, UE_KINDA_SMALL_NUMBER))); }
+protected:
+	UE_DEPRECATED(5.8, "Use HiGenGridSizeMultiplier instead")
+	UPROPERTY()
+	uint32 HiGenExponential = 0;
+
 };
 
 UCLASS(MinimalAPI, BlueprintType, ClassGroup = (Procedural), hidecategories = (Object))
diff --git a/Engine/Plugins/PCG/Source/PCG/Public/PCGGraphExecutionStateInterface.h b/Engine/Plugins/PCG/Source/PCG/Public/PCGGraphExecutionStateInterface.h
--- a/Engine/Plugins/PCG/Source/PCG/Public/PCGGraphExecutionStateInterface.h
+++ b/Engine/Plugins/PCG/Source/PCG/Public/PCGGraphExecutionStateInterface.h
@@ -226,9 +226,6 @@ public:
 	/** Returns a GridDescriptor based on this execution source for the specified grid size. */
 	PCG_API virtual FPCGGridDescriptor GetGridDescriptor(uint32 InGridSize) const;
 
-	/** Get the HiGen grid this execution source executes on. */
-	virtual EPCGHiGenGrid GetGenerationGrid() const { return EPCGHiGenGrid::Uninitialized; }
-
 	/** Get the size for the HiGen grid this execution source executes on. */
 	virtual uint32 GetGenerationGridSize() const { return PCGHiGenGrid::UninitializedGridSize(); }
 
@@ -246,10 +243,10 @@ public:
 	PCG_API virtual const FPCGRuntimeGenerationRadii& GetGenerationRadii() const;
 
 	/** Gets the radius inside of which this execution source will be generated when using runtime generation. */
-	virtual double GetGenerationRadiusFromGrid(EPCGHiGenGrid InGrid) const { return 0.0f; }
+	virtual double GetGenerationRadiusFromGrid(uint32 InGrid) const { return 0.0f; }
 
 	/** Gets the radius outside of which this execution source will be cleaned up when using runtime generation. */
-	virtual double GetCleanupRadiusFromGrid(EPCGHiGenGrid InGrid) const { return 0.0f; }
+	virtual double GetCleanupRadiusFromGrid(uint32 InGrid) const { return 0.0f; }
 
 	/** Returns true if the generated artifacts are transient */
 	virtual bool IsInPreviewMode() const { return false; }
diff --git a/Engine/Plugins/PCG/Source/PCG/Public/Subsystems/PCGSubsystem.h b/Engine/Plugins/PCG/Source/PCG/Public/Subsystems/PCGSubsystem.h
--- a/Engine/Plugins/PCG/Source/PCG/Public/Subsystems/PCGSubsystem.h
+++ b/Engine/Plugins/PCG/Source/PCG/Public/Subsystems/PCGSubsystem.h
@@ -155,6 +155,9 @@ public:
 	PCG_API UPCGLandscapeCache* GetLandscapeCache();
 
 	/** Schedule graph(owner->graph) */
+	PCG_API FPCGTaskId ScheduleComponent(UPCGComponent* PCGComponent, uint32 Grid, bool bForce, const TArray<FPCGTaskId>& InDependencies);
+
+	UE_DEPRECATED(5.8, "Use the version that takes a uint32 Grid")
 	PCG_API FPCGTaskId ScheduleComponent(UPCGComponent* PCGComponent, EPCGHiGenGrid Grid, bool bForce, const TArray<FPCGTaskId>& InDependencies);
 
 	/** Schedule cleanup(owner->graph). */
diff --git a/Engine/Plugins/PCG/Source/PCGEditor/Private/Nodes/PCGEditorGraphNodeBase.cpp b/Engine/Plugins/PCG/Source/PCGEditor/Private/Nodes/PCGEditorGraphNodeBase.cpp
--- a/Engine/Plugins/PCG/Source/PCGEditor/Private/Nodes/PCGEditorGraphNodeBase.cpp
+++ b/Engine/Plugins/PCG/Source/PCGEditor/Private/Nodes/PCGEditorGraphNodeBase.cpp
@@ -735,16 +735,14 @@ EPCGChangeType UPCGEditorGraphNodeBase::UpdateStructuralVisualization(const IPCG
 
 	// Set the inspected grid size - this is used for grid size visualization.
 	uint32 InspectingGridSize = PCGHiGenGrid::UninitializedGridSize();
-	EPCGHiGenGrid InspectingGrid = EPCGHiGenGrid::Uninitialized;
 	if (TopGraph && TopGraph->IsHierarchicalGenerationEnabled() && InSourceBeingDebugged && (InSourceBeingDebugged->GetExecutionState().IsPartitioned() || InSourceBeingDebugged->GetExecutionState().IsLocalSource()))
 	{
 		InspectingGridSize = InSourceBeingDebugged->GetExecutionState().GetGenerationGridSize();
-		InspectingGrid = InSourceBeingDebugged->GetExecutionState().GetGenerationGrid();
 	}
 
-	if (InspectedGenerationGrid != InspectingGrid)
+	if (InspectedGridSize != InspectingGridSize)
 	{
-		InspectedGenerationGrid = InspectingGrid;
+		InspectedGridSize = InspectingGridSize;
 		ChangeType |= EPCGChangeType::Cosmetic;
 	}
 
@@ -767,7 +765,7 @@ EPCGChangeType UPCGEditorGraphNodeBase::UpdateStructuralVisualization(const IPCG
 	// Don't do culling visualization on newly placed nodes. Let the execution complete notification update that.
 	bool bIsCulled = !bNewlyPlaced && !PCGEditorGraphNodeBase::ShouldDisplayAsActive(this, InSourceBeingDebugged, InStackBeingInspected);
 
-	EPCGHiGenGrid ThisGrid = EPCGHiGenGrid::Uninitialized;
+	uint32 ThisGridSize = PCGHiGenGrid::UninitializedGridSize();
 
 	// Show grid size visualization if higen is enabled and if we're inspecting a specific grid, and we're inspecting a subgraph since subgraphs
 	// execute at the invoked grid level.
@@ -792,20 +790,21 @@ EPCGChangeType UPCGEditorGraphNodeBase::UpdateStructuralVisualization(const IPCG
 				bIsCulled = false;
 			}
 
-			ThisGrid = PCGHiGenGrid::GridSizeToGrid(NodeGridSize);
+			ThisGridSize = NodeGridSize;
 		}
 		else
 		{
 			// If higen is enabled then we are inspecting an invoked subgraph. Display the inspected grid size so that the user still
 			// gets the execution grid information.
-			ThisGrid = PCGHiGenGrid::GridSizeToGrid(InspectingGridSize);
+			ThisGridSize = InspectingGridSize;
 		}
 	}
 
-	if (GenerationGrid != ThisGrid)
+	if (GenerationGridSize != ThisGridSize)
 	{
-		GenerationGrid = ThisGrid;
+		// Update node if the grid size has changed.
 		ChangeType |= EPCGChangeType::Cosmetic;
+		GenerationGridSize = ThisGridSize;
 	}
 
 	SetIsCulledFromExecution(bIsCulled);
diff --git a/Engine/Plugins/PCG/Source/PCGEditor/Private/Nodes/SPCGEditorGraphNode.cpp b/Engine/Plugins/PCG/Source/PCGEditor/Private/Nodes/SPCGEditorGraphNode.cpp
--- a/Engine/Plugins/PCG/Source/PCGEditor/Private/Nodes/SPCGEditorGraphNode.cpp
+++ b/Engine/Plugins/PCG/Source/PCGEditor/Private/Nodes/SPCGEditorGraphNode.cpp
@@ -2,6 +2,7 @@
 
 #include "SPCGEditorGraphNode.h"
 
+#include "PCGGraph.h"
 #include "PCGNode.h"
 #include "PCGPin.h"
 #include "PCGSettings.h"
@@ -339,7 +340,7 @@ void SPCGEditorGraphNode::OnNodeChanged()
 
 bool SPCGEditorGraphNode::UsesHiGenOverlay() const
 {
-	return PCGEditorGraphNode->GetInspectedGenerationGrid() != EPCGHiGenGrid::Uninitialized && PCGEditorGraphNode->IsNodeEnabled();
+	return PCGEditorGraphNode->GetInspectedGridSize() != PCGHiGenGrid::UninitializedGridSize() && PCGEditorGraphNode->IsNodeEnabled();
 }
 
 bool SPCGEditorGraphNode::UsesGPUOverlay() const
@@ -361,10 +362,24 @@ bool SPCGEditorGraphNode::UsesDebugBrush() const
 	return Schema ? Schema->ShouldAddDebugBrushOnNode() : true;
 }
 
-FLinearColor SPCGEditorGraphNode::GetGridLabelColor(EPCGHiGenGrid NodeGrid)
+FLinearColor SPCGEditorGraphNode::GetGridLabelColor(uint32 NodeGrid) const
 {
-	// All colours hand tweaked to give a kind of "temperature scale" for the hierarchy.
-	switch (NodeGrid)
+	check(PCGEditorGraphNode);
+
+	EPCGHiGenGrid UnscaledGrid = EPCGHiGenGrid::Uninitialized;
+
+	if (UPCGGraph* Graph = PCGEditorGraphNode->GetPCGNode() ? PCGEditorGraphNode->GetPCGNode()->GetGraph() : nullptr)
+	{
+		// Account for grid scaling.
+		UnscaledGrid = PCGHiGenGrid::ScaledGridSizeToGrid(NodeGrid, Graph->GetGridSizeMultiplier());
+	}
+	else
+	{
+		UnscaledGrid = PCGHiGenGrid::GridSizeToGrid(NodeGrid);
+	}
+
+	// All colors hand tweaked to give a kind of "temperature scale" for the hierarchy.
+	switch (UnscaledGrid)
 	{
 	case EPCGHiGenGrid::Unbounded:
 		return FColor(255, 255, 255, 255);
@@ -412,24 +427,23 @@ void SPCGEditorGraphNode::AddHiGenOverlayWidget(TArray<FOverlayWidgetInfo>& Over
 	check(UsesHiGenOverlay());
 
 	// Higen grid size overlay widget. All magic numbers below hand tweaked to match UI mockup.
-	const EPCGHiGenGrid InspectedGrid = PCGEditorGraphNode->GetInspectedGenerationGrid();
+	const uint32 InspectedGrid = PCGEditorGraphNode->GetInspectedGridSize();
+	const uint32 Grid = PCGEditorGraphNode->GetGenerationGridSize();
 
 	FText GenerationGridText;
-	const EPCGHiGenGrid Grid = PCGEditorGraphNode->GetGenerationGrid();
 
-	if (Grid == EPCGHiGenGrid::Unbounded)
+	if (Grid == PCGHiGenGrid::UnboundedGridSize())
 	{
 		GenerationGridText = FText::FromString(TEXT("UB"));
 	}
 	else
 	{
 		// Meters are easier on the eyes.
-		const uint32 GridSize = PCGHiGenGrid::GridToGridSize(Grid) / 100;
-		GenerationGridText = FText::AsNumber(GridSize, &FNumberFormattingOptions::DefaultNoGrouping());
+		GenerationGridText = FText::AsNumber(Grid / 100, &FNumberFormattingOptions::DefaultNoGrouping());
 	}
 
 	FLinearColor Tint = FLinearColor::White;
-	if (Grid != EPCGHiGenGrid::Uninitialized)
+	if (Grid != PCGHiGenGrid::UninitializedGridSize())
 	{
 		Tint = GetGridLabelColor(Grid);
 	}
@@ -511,7 +525,7 @@ void SPCGEditorGraphNode::AddGPUOverlayWidget(TArray<FOverlayWidgetInfo>& Overla
 	OverlayWidgets.Add(GPUUsageLabelInfo);
 }
 
-const FSlateBrush* SPCGEditorGraphNode::GetBorderBrush(EPCGHiGenGrid InspectedGrid, EPCGHiGenGrid NodeGrid) const
+const FSlateBrush* SPCGEditorGraphNode::GetBorderBrush(uint32 InspectedGrid, uint32 NodeGrid) const
 {
 	if (InspectedGrid == NodeGrid)
 	{
diff --git a/Engine/Plugins/PCG/Source/PCGEditor/Private/Nodes/SPCGEditorGraphNode.h b/Engine/Plugins/PCG/Source/PCGEditor/Private/Nodes/SPCGEditorGraphNode.h
--- a/Engine/Plugins/PCG/Source/PCGEditor/Private/Nodes/SPCGEditorGraphNode.h
+++ b/Engine/Plugins/PCG/Source/PCGEditor/Private/Nodes/SPCGEditorGraphNode.h
@@ -48,7 +48,7 @@ protected:
 	virtual bool UsesDebugBrush() const;
 
 private:
-	static FLinearColor GetGridLabelColor(EPCGHiGenGrid NodeGrid);
+	FLinearColor GetGridLabelColor(uint32 NodeGrid) const;
 
 	/** Adds the Hierarchical Generation overlay to the array, displaying the HiGen grid size on the node. */
 	void AddHiGenOverlayWidget(TArray<FOverlayWidgetInfo>& OverlayWidgets) const;
@@ -66,7 +66,7 @@ private:
 	* Get the border brush for the given combination of grid sizes and enabled state. All a big workaround for FSlateRoundedBoxBrush not respecting
 	* the tint colour.
 	*/
-	const FSlateBrush* GetBorderBrush(EPCGHiGenGrid InspectedGrid, EPCGHiGenGrid NodeGrid) const;
+	const FSlateBrush* GetBorderBrush(uint32 InspectedGrid, uint32 NodeGrid) const;
 
 	UPCGEditorGraphNodeBase* PCGEditorGraphNode = nullptr;
 };
diff --git a/Engine/Plugins/PCG/Source/PCGEditor/Private/WorldPartitionBuilder/PCGWorldPartitionBuilder.cpp b/Engine/Plugins/PCG/Source/PCGEditor/Private/WorldPartitionBuilder/PCGWorldPartitionBuilder.cpp
--- a/Engine/Plugins/PCG/Source/PCGEditor/Private/WorldPartitionBuilder/PCGWorldPartitionBuilder.cpp
+++ b/Engine/Plugins/PCG/Source/PCGEditor/Private/WorldPartitionBuilder/PCGWorldPartitionBuilder.cpp
@@ -1302,11 +1302,11 @@ FPCGTaskId FPCGWorldPartitionBuilder::GenerateComponent(UPCGComponent* InCompone
 	// The only Partitioned components we do generate in iterative mode are generated only for their Unbounded grid as the other grid sizes will be generated through the PAs directly
 	if (InArgs.bIterativeCellLoading && InComponent->IsPartitioned())
 	{
-		GenerateTask = InComponent->GenerateLocalGetTaskId(EPCGComponentGenerationTrigger::GenerateOnDemand, /*bForce=*/true, EPCGHiGenGrid::Unbounded);
+		GenerateTask = InComponent->GenerateLocalGetTaskId(EPCGComponentGenerationTrigger::GenerateOnDemand, /*bForce=*/true, PCGHiGenGrid::UnboundedGridSize());
 	}
 	else if(InTaskDependency != InvalidPCGTaskId)
 	{
-		GenerateTask = InComponent->GenerateLocalGetTaskId(EPCGComponentGenerationTrigger::GenerateOnDemand, /*bForce=*/true, EPCGHiGenGrid::Uninitialized, { InTaskDependency });
+		GenerateTask = InComponent->GenerateLocalGetTaskId(EPCGComponentGenerationTrigger::GenerateOnDemand, /*bForce=*/true, PCGHiGenGrid::UninitializedGridSize(), { InTaskDependency });
 	}
 	else
 	{
diff --git a/Engine/Plugins/PCG/Source/PCGEditor/Public/Nodes/PCGEditorGraphNodeBase.h b/Engine/Plugins/PCG/Source/PCGEditor/Public/Nodes/PCGEditorGraphNodeBase.h
--- a/Engine/Plugins/PCG/Source/PCGEditor/Public/Nodes/PCGEditorGraphNodeBase.h
+++ b/Engine/Plugins/PCG/Source/PCGEditor/Public/Nodes/PCGEditorGraphNodeBase.h
@@ -135,11 +135,11 @@ public:
 	/** Bitmask of inactive output pins. Bit N will be set if output pin index N is inactive. */
 	uint64 GetInactiveOutputPinMask() const { return InactiveOutputPinMask; }
 
-	/** The grid that this node executes on if higen is enabled, otherwise Unitialized. */
-	EPCGHiGenGrid GetGenerationGrid() const { return GenerationGrid; }
+	/** The grid that this node executes on if higen is enabled, otherwise UninitializedGridSize(). */
+	uint32 GetGenerationGridSize() const { return GenerationGridSize; }
 
-	/** The higen grid currently being inspected if any, otherwise Uninitialized. */
-	EPCGHiGenGrid GetInspectedGenerationGrid() const { return InspectedGenerationGrid; }
+	/** The higen grid currently being inspected if any, otherwise UninitializedGridSize(). */
+	uint32 GetInspectedGridSize() const { return InspectedGridSize; }
 
 	/** Whether the given output pin was active in the previous execution. */
 	bool IsOutputPinActive(const UEdGraphPin* InOutputPin) const;
@@ -199,10 +199,22 @@ protected:
 	/** Bitmask of inactive output pins. Bit N will be set if output pin index N is inactive. */
 	uint64 InactiveOutputPinMask = 0;
 
-	/** The grid that this node executes on if higen is enabled, otherwise Unitialized. */
-	EPCGHiGenGrid GenerationGrid = EPCGHiGenGrid::Uninitialized;
+	/** The grid that this node executes on if higen is enabled, otherwise UninitializedGridSize(). */
+	uint32 GenerationGridSize = PCGHiGenGrid::UninitializedGridSize();
+
+	/** The higen grid currently being inspected if any, otherwise UninitializedGridSize(). */
+	uint32 InspectedGridSize = PCGHiGenGrid::UninitializedGridSize();
 
-	/** The higen grid currently being inspected if any, otherwise Uninitialized. */
+	// Deprecated section
+public:
+	UE_DEPRECATED(5.8, "Use GenerationGridSize")
+	EPCGHiGenGrid GetGenerationGrid() const { return PCGHiGenGrid::GridSizeToGrid(GenerationGridSize); }
+	UE_DEPRECATED(5.8, "Use GetInspectedGridSize")
+	EPCGHiGenGrid GetInspectedGenerationGrid() const { return PCGHiGenGrid::GridSizeToGrid(InspectedGridSize); }
+protected:
+	UE_DEPRECATED(5.8, "Use GenerationGridSize")
+	EPCGHiGenGrid GenerationGrid = EPCGHiGenGrid::Uninitialized;
+	UE_DEPRECATED(5.8, "Use InspectedGridSize")
 	EPCGHiGenGrid InspectedGenerationGrid = EPCGHiGenGrid::Uninitialized;
 };
 
