[PCG] Broadcast OnPCGSourceUnregistered from the UPCGSubsystem::UnregisterExecutionSource()

#rnx
#Synced-CL 50600368

diff --git a/Engine/Plugins/PCG/Source/PCG/Private/PCGComponent.cpp b/Engine/Plugins/PCG/Source/PCG/Private/PCGComponent.cpp
--- a/Engine/Plugins/PCG/Source/PCG/Private/PCGComponent.cpp
+++ b/Engine/Plugins/PCG/Source/PCG/Private/PCGComponent.cpp
@@ -1715,8 +1715,6 @@ void UPCGComponent::OnUnregister()
 			Subsystem->CancelGeneration(this);
 		}
 
-		Subsystem->GetOnPCGSourceUnregistered().Broadcast(IsValidChecked(this) ? this : nullptr);
-
 		// Always unregister here if we are being unloaded (we are not being re-registered)
 		if (bUnregisteredThroughLoading)
 		{
diff --git a/Engine/Plugins/PCG/Source/PCG/Private/RuntimeGen/PCGRuntimeGenExecutionSource.cpp b/Engine/Plugins/PCG/Source/PCG/Private/RuntimeGen/PCGRuntimeGenExecutionSource.cpp
--- a/Engine/Plugins/PCG/Source/PCG/Private/RuntimeGen/PCGRuntimeGenExecutionSource.cpp
+++ b/Engine/Plugins/PCG/Source/PCG/Private/RuntimeGen/PCGRuntimeGenExecutionSource.cpp
@@ -48,9 +48,6 @@ void UPCGRuntimeGenExecutionSource::Unregister()
 {
 	if (UPCGSubsystem* Subsystem = Cast<UPCGSubsystem>(ExecutionState.GetSubsystem()))
 	{
-#if WITH_EDITOR
-		Subsystem->GetOnPCGSourceUnregistered().Broadcast(this);
-#endif
 		Subsystem->UnregisterExecutionSource(this);
 	}
 }
diff --git a/Engine/Plugins/PCG/Source/PCG/Private/Subsystems/PCGSubsystem.cpp b/Engine/Plugins/PCG/Source/PCG/Private/Subsystems/PCGSubsystem.cpp
--- a/Engine/Plugins/PCG/Source/PCG/Private/Subsystems/PCGSubsystem.cpp
+++ b/Engine/Plugins/PCG/Source/PCG/Private/Subsystems/PCGSubsystem.cpp
@@ -2263,6 +2263,10 @@ bool UPCGSubsystem::RemapExecutionSource(const IPCGGraphExecutionSource* InOldEx
 
 void UPCGSubsystem::UnregisterExecutionSource(IPCGGraphExecutionSource* InExecutionSource, bool bForce)
 {
+#if WITH_EDITOR
+	GetOnPCGSourceUnregistered().Broadcast(InExecutionSource);
+#endif
+
 	return TrackingManager->UnregisterExecutionSource(InExecutionSource, bForce);
 }
 
