[PCG] Runtime generation schedule - add sleep time when scheduler is idle.

[RN] [PCG] Optimization: Add 'pcg.RuntimeGeneration.TimeBetweenRuntimeGenSchedulerTicks' CVar to reduce Game Thread cost of runtime generation.

#Synced-CL 50651737

diff --git a/Engine/Plugins/PCG/Source/PCG/Private/RuntimeGen/PCGRuntimeGenScheduler.cpp b/Engine/Plugins/PCG/Source/PCG/Private/RuntimeGen/PCGRuntimeGenScheduler.cpp
--- a/Engine/Plugins/PCG/Source/PCG/Private/RuntimeGen/PCGRuntimeGenScheduler.cpp
+++ b/Engine/Plugins/PCG/Source/PCG/Private/RuntimeGen/PCGRuntimeGenScheduler.cpp
@@ -129,6 +129,12 @@ namespace PCGRuntimeGenSchedulerHelpers
 		false,
 		TEXT("Draws debug boxes to indicate regions where PCG is requesting virtual texture priming."));
 
+	static float GTimeBetweenTicks = 0.0f;
+	static FAutoConsoleVariableRef CVarTimeBetweenTicks(
+		TEXT("pcg.RuntimeGeneration.TimeBetweenRuntimeGenSchedulerTicks"),
+		GTimeBetweenTicks,
+		TEXT("Period between ticks of the runtime generation scheduler. Saves processing time. If scheduling work is in progress or execution sources updated or refreshed, tick will process every frame until scheduler is idle."));
+
 #if PCG_RGS_ONSCREENDEBUGMESSAGES
 	static TAutoConsoleVariable<bool> CVarEnableDebugOverlay(
 		TEXT("pcg.RuntimeGeneration.EnableDebugOverlay"),
@@ -305,13 +311,13 @@ FPCGRuntimeGenScheduler::~FPCGRuntimeGenScheduler()
 	FLevelStreamingDelegates::OnLevelStreamingStateChanged.RemoveAll(this);
 }
 
-void FPCGRuntimeGenScheduler::Tick(APCGWorldActor* InPCGWorldActor, double InEndTime)
+void FPCGRuntimeGenScheduler::Tick(APCGWorldActor* InPCGWorldActor, double InDeltaTime, double InEndTime)
 {
 	check(InPCGWorldActor && GenSourceManager);
 
 	// 0. Preamble - check if we should be active in this world and do lazy initialization.
 
-	if (!ShouldTick())
+	if (!ShouldTick(InDeltaTime))
 	{
 		return;
 	}
@@ -392,6 +398,12 @@ void FPCGRuntimeGenScheduler::Tick(APCGWorldActor* InPCGWorldActor, double InEnd
 		// Only apply time budget to cleanup currently. Currently too easy to introduce latency issues so don't hold back generation new execution sources
 		// (and we already have CVarNumGeneratingComponentsAtSameTime to throttle new generations).
 		TickScheduleGeneration(SourcesToGenerate);
+
+		if (!SourcesToGenerate.IsEmpty())
+		{
+			// Continue next frame.
+			TimeToNextTick = 0.0;
+		}
 	}
 
 #if PCG_RGS_ONSCREENDEBUGMESSAGES
@@ -399,14 +411,21 @@ void FPCGRuntimeGenScheduler::Tick(APCGWorldActor* InPCGWorldActor, double InEnd
 #endif
 }
 
-bool FPCGRuntimeGenScheduler::ShouldTick()
+bool FPCGRuntimeGenScheduler::ShouldTick(double InDeltaTime)
 {
 	check(Context.World);
 
-	if (!PCGRuntimeGenSchedulerHelpers::CVarRuntimeGenerationEnable.GetValueOnAnyThread())
+	if (!PCGRuntimeGenSchedulerHelpers::CVarRuntimeGenerationEnable.GetValueOnGameThread())
+	{
+		return false;
+	}
+
+	TimeToNextTick -= InDeltaTime;
+	if (TimeToNextTick > 0.0)
 	{
 		return false;
 	}
+	TimeToNextTick = PCGRuntimeGenSchedulerHelpers::GTimeBetweenTicks;
 
 	// Disable tick of editor scheduling if in runtime or PIE.
 	if (PCGHelpers::IsRuntimeOrPIE() && !Context.World->IsGameWorld())
@@ -1107,6 +1126,11 @@ void FPCGRuntimeGenScheduler::TickCleanup(const TSet<IPCGGenSourceBase*>& InGenS
 		// Record that we completed a cleanup scan.
 		bCleanupScanRequired = false;
 	}
+	else
+	{
+		// Continue next frame.
+		TimeToNextTick = 0.0;
+	}
 }
 
 void FPCGRuntimeGenScheduler::TickScheduleGeneration(TMap<FGridGenerationKey, double>& InOutSourcesToGenerate)
@@ -1522,6 +1546,8 @@ void FPCGRuntimeGenScheduler::OnOriginalComponentRegistered(UPCGComponent* InOri
 	// When an original/non-partitioned component is registered, we need to dirty the state.
 	bContextDirty = true;
 
+	TimeToNextTick = 0.0f;
+
 #if WITH_EDITOR
 	RegisterGraphParameterChangedEvent(InOriginalComponent);
 #endif
@@ -1542,6 +1568,8 @@ void FPCGRuntimeGenScheduler::OnOriginalComponentUnregistered(UPCGComponent* InO
 	// When an original/non-partitioned component is unregistered, we need to dirty the state.
 	bContextDirty = true;
 
+	TimeToNextTick = 0.0f;
+
 	// Make sure we perform a scan for cells that need cleanup on next tick.
 	bCleanupScanRequired = true;
 
@@ -1985,6 +2013,8 @@ void FPCGRuntimeGenScheduler::RefreshExecutionSource(IPCGGraphExecutionSource* I
 	// Trigger a rescan of generation cells for this execution source.
 	ChangeDetector.FlushCachedState(OriginalSource);
 
+	TimeToNextTick = 0.0;
+
 	const bool bLoggingEnabled = PCGRuntimeGenSchedulerHelpers::CVarRuntimeGenerationEnableDebugging.GetValueOnAnyThread();
 
 	// Useful because we can run into generation order issues if execution sources are left to continue generating.
diff --git a/Engine/Plugins/PCG/Source/PCG/Private/Subsystems/PCGSubsystem.cpp b/Engine/Plugins/PCG/Source/PCG/Private/Subsystems/PCGSubsystem.cpp
--- a/Engine/Plugins/PCG/Source/PCG/Private/Subsystems/PCGSubsystem.cpp
+++ b/Engine/Plugins/PCG/Source/PCG/Private/Subsystems/PCGSubsystem.cpp
@@ -412,7 +412,7 @@ void UPCGSubsystem::Tick(float DeltaSeconds)
 
 	if (PCGWorldActor)
 	{
-		RuntimeGenScheduler->Tick(PCGWorldActor, EndTime);
+		RuntimeGenScheduler->Tick(PCGWorldActor, DeltaSeconds, EndTime);
 	}
 }
 
diff --git a/Engine/Plugins/PCG/Source/PCG/Public/RuntimeGen/PCGRuntimeGenScheduler.h b/Engine/Plugins/PCG/Source/PCG/Public/RuntimeGen/PCGRuntimeGenScheduler.h
--- a/Engine/Plugins/PCG/Source/PCG/Public/RuntimeGen/PCGRuntimeGenScheduler.h
+++ b/Engine/Plugins/PCG/Source/PCG/Public/RuntimeGen/PCGRuntimeGenScheduler.h
@@ -79,7 +79,7 @@ public:
 	FPCGRuntimeGenScheduler& operator=(const FPCGRuntimeGenScheduler& other) = delete;
 	FPCGRuntimeGenScheduler& operator=(FPCGRuntimeGenScheduler&& other) = delete; 
 
-	void Tick(APCGWorldActor* InPCGWorldActor, double InEndTime);
+	void Tick(APCGWorldActor* InPCGWorldActor, double InDeltaTime, double InEndTime);
 
 	void OnOriginalComponentRegistered(UPCGComponent* InOriginalComponent);
 	void OnOriginalComponentUnregistered(UPCGComponent* InOriginalComponent);
@@ -167,7 +167,7 @@ protected:
 	};
 
 	/** Returns true if the scheduler should tick this frame. */
-	bool ShouldTick();
+	bool ShouldTick(double InDeltaTime);
 
 	// Add constructor, get moves for the sets
 	struct FTickQueueSourcesForGenerationInputs
@@ -324,6 +324,8 @@ private:
 	/** Maps an original PCG component to its runtime gen execution source. This relationship only exists for components whose graphs opt-in to actor-component-less generation. */
 	TMap<TObjectKey<UPCGComponent>, TObjectPtr<UPCGRuntimeGenExecutionSource>> OriginalExecutionSources;
 
+	double TimeToNextTick = 0.0;
+
 	friend struct PCGRuntimeGenSchedulerHelpers::FStatsOverlay;
 
 	// Deprecated section
