[PCG] Prototype for a details viewer for generic attributes

* Add a node to create any type of constant as a generic attribute (temporary)
* Add a details view like widget to visualize a custom struct
* Add Button with Ellipsis to open that view in the Attribute List View

#rnx

#Synced-CL 50489647

diff --git a/Engine/Plugins/PCG/Source/PCG/Public/PCGDataVisualization.h b/Engine/Plugins/PCG/Source/PCG/Public/PCGDataVisualization.h
--- a/Engine/Plugins/PCG/Source/PCG/Public/PCGDataVisualization.h
+++ b/Engine/Plugins/PCG/Source/PCG/Public/PCGDataVisualization.h
@@ -11,6 +11,7 @@
 
 #if WITH_EDITOR
 class AActor;
+class FPCGMetadataAttributeGeneric;
 class UPCGData;
 class UPCGSettingsInterface;
 struct FPCGContext;
@@ -63,6 +64,7 @@ struct FPCGTableVisualizerColumnInfo
 	EPCGTableVisualizerCellAlignment CellAlignment = EPCGTableVisualizerCellAlignment::Right;
 	TSharedPtr<const IPCGAttributeAccessor> Accessor;
 	TSharedPtr<const IPCGAttributeAccessorKeys> AccessorKeys;
+	const FPCGMetadataAttributeGeneric* GenericAttribute = nullptr;
 };
 
 struct FPCGTableVisualizerInfo
diff --git a/Engine/Plugins/PCG/Source/PCGEditor/PCGEditor.Build.cs b/Engine/Plugins/PCG/Source/PCGEditor/PCGEditor.Build.cs
--- a/Engine/Plugins/PCG/Source/PCGEditor/PCGEditor.Build.cs
+++ b/Engine/Plugins/PCG/Source/PCGEditor/PCGEditor.Build.cs
@@ -61,6 +61,7 @@ namespace UnrealBuildTool.Rules
 					"RHI",
 					"RenderCore",
 					"SceneOutliner",
+					"ScriptableEditorWidgets",
 					"Slate",
 					"SlateCore",
 					"SourceControl",
diff --git a/Engine/Plugins/PCG/Source/PCGEditor/Private/DataVisualizations/PCGDataVisualizationHelpers.cpp b/Engine/Plugins/PCG/Source/PCGEditor/Private/DataVisualizations/PCGDataVisualizationHelpers.cpp
--- a/Engine/Plugins/PCG/Source/PCGEditor/Private/DataVisualizations/PCGDataVisualizationHelpers.cpp
+++ b/Engine/Plugins/PCG/Source/PCGEditor/Private/DataVisualizations/PCGDataVisualizationHelpers.cpp
@@ -166,7 +166,20 @@ namespace PCGDataVisualizationHelpers
 
 		for (int32 I = 0; I < AttributeNames.Num(); I++)
 		{
+			const int32 ColumnInfosNum = OutInfo.ColumnInfos.Num();
 			PCGMetadataAttribute::CallbackWithRightType(static_cast<uint16>(AttributeTypes[I]), AddInfo, AttributeNames[I]);
+			if (ColumnInfosNum == OutInfo.ColumnInfos.Num())
+			{
+				if (const FPCGMetadataAttributeGeneric* GenericAttr = MetadataDomain->GetConstGenericAttribute(AttributeNames[I]))
+				{
+					FPCGTableVisualizerColumnInfo& ColumnInfo = OutInfo.ColumnInfos.Emplace_GetRef();
+					ColumnInfo.Label = FText::FromName(AttributeNames[I]);
+					ColumnInfo.Id = AttributeNames[I] == NAME_None ? TEXT("AttrNone") : AttributeNames[I];
+					ColumnInfo.GenericAttribute = GenericAttr;
+					Selector.SetAttributeName(AttributeNames[I]);
+					ColumnInfo.AccessorKeys = TSharedPtr<const IPCGAttributeAccessorKeys>(PCGAttributeAccessorHelpers::CreateConstKeys(Data, Selector).Release());
+				}
+			}
 		}
 	}
 
diff --git a/Engine/Plugins/PCG/Source/PCGEditor/Private/Widgets/SPCGEditorGraphAttributeListView.cpp b/Engine/Plugins/PCG/Source/PCGEditor/Private/Widgets/SPCGEditorGraphAttributeListView.cpp
--- a/Engine/Plugins/PCG/Source/PCGEditor/Private/Widgets/SPCGEditorGraphAttributeListView.cpp
+++ b/Engine/Plugins/PCG/Source/PCGEditor/Private/Widgets/SPCGEditorGraphAttributeListView.cpp
@@ -23,6 +23,7 @@
 #include "PCGEditorSettings.h"
 #include "PCGEditorStyle.h"
 #include "PCGEditorUtils.h"
+#include "SPCGMetadataAttributeGenericViewer.h"
 #include "DataVisualizations/PCGDataVisualizationHelpers.h"
 #include "Managers/PCGEditorInspectionDataManager.h"
 #include "Nodes/PCGEditorGraphNodeBase.h"
@@ -141,6 +142,10 @@ TSharedRef<SWidget> SPCGListViewItemRow::GenerateWidgetForColumn(const FName& Co
 
 	// Make the text clickable for soft object paths
 	bool bIsClickable = false;
+	
+	// @todo_pcg: To be removed when we do not make the distinction between typed and generic.
+	bool bDataRetrievedWithAccessors = false;
+	TTuple<const FPCGMetadataAttributeGeneric*, PCGMetadataEntryKey> AttrAndKey{nullptr, PCGInvalidEntryKey};
 
 	ETextOverflowPolicy OverflowPolicy = GetDefault<UPCGEditorSettings>()->AttributeListViewSettings.GeneralOverflowPolicy;
 
@@ -149,6 +154,7 @@ TSharedRef<SWidget> SPCGListViewItemRow::GenerateWidgetForColumn(const FName& Co
 		// We have to make sure the data is still valid otherwise this will crash.
 		if (ensure(SharedAttributeListView->DataPtr) && PCGColumnData->DataAccessor.IsValid() && PCGColumnData->DataKeys.IsValid())
 		{
+			bDataRetrievedWithAccessors = true;
 			const int32 Index = InternalItem->Index;
 			auto Callback = [&PCGColumnData, Index, &RowText, &bShouldComputeCellWidth, &bIsClickable, &OverflowPolicy, ColumnId] (auto Dummy)
 			{
@@ -201,6 +207,15 @@ TSharedRef<SWidget> SPCGListViewItemRow::GenerateWidgetForColumn(const FName& Co
 
 			PCGMetadataAttribute::CallbackWithRightType(PCGColumnData->DataAccessor->GetUnderlyingType(), Callback);
 		}
+		else if (PCGColumnData->DataKeys.IsValid() && PCGColumnData->GenericAttribute)
+		{
+			const PCGMetadataEntryKey* Key = nullptr;
+			if (PCGColumnData->DataKeys->GetKey<PCGMetadataEntryKey>(Key))
+			{
+				AttrAndKey.Get<0>() = PCGColumnData->GenericAttribute;
+				AttrAndKey.Get<1>() = *Key;
+			}
+		}
 	}
 
 	if (bShouldComputeCellWidth)
@@ -237,7 +252,7 @@ TSharedRef<SWidget> SPCGListViewItemRow::GenerateWidgetForColumn(const FName& Co
 
 		return Hyperlink.ToSharedRef();
 	}
-	else
+	else if (bDataRetrievedWithAccessors)
 	{
 		return SNew(STextBlock)
 			.Text(RowText)
@@ -245,6 +260,33 @@ TSharedRef<SWidget> SPCGListViewItemRow::GenerateWidgetForColumn(const FName& Co
 			.Margin(Margin)
 			.TextStyle(TextStyle);
 	}
+	else if (AttrAndKey.Get<0>() && SharedAttributeListView->DataPtr)
+	{
+		return SNew(SButton)
+			.Text(INVTEXT("..."))
+			.OnClicked_Lambda([AttrAndKey, WeakData = TWeakObjectPtr<const UPCGData>(SharedAttributeListView->DataPtr)]()
+				{
+					TSharedRef<SWindow> ModalWindow = SNew(SWindow)
+						.Title(INVTEXT("Viewer"))
+						.ClientSize(FVector2D(700, 700));
+
+					TSharedRef<SPCGMetadataAttributeGenericViewer> Viewer = SNew(SPCGMetadataAttributeGenericViewer);
+					Viewer->Setup(WeakData, AttrAndKey.Get<0>(), AttrAndKey.Get<1>());
+
+					ModalWindow->SetContent(Viewer);
+
+					GEditor->EditorAddModalWindow(ModalWindow);
+					return FReply::Handled();
+				});
+	}
+	else
+	{
+		return SNew(STextBlock)
+			.Text(INVTEXT(""))
+			.OverflowPolicy(OverflowPolicy)
+			.Margin(Margin)
+			.TextStyle(TextStyle);
+	}
 }
 
 void SPCGListViewItemRow::OnSoftObjectPathHyperlinkClicked(const FText& InText)
@@ -1792,6 +1834,7 @@ void SPCGEditorGraphAttributeListView::AddColumn(const FPCGTableVisualizerColumn
 	FPCGColumnData& ColumnData = PCGColumnData.Add(InColumnInfo.Id);
 	ColumnData.DataAccessor = InColumnInfo.Accessor;
 	ColumnData.DataKeys = InColumnInfo.AccessorKeys;
+	ColumnData.GenericAttribute = InColumnInfo.GenericAttribute;
 
 	const float ColumnWidth = InColumnInfo.Width < 0.0f ? PCGEditorGraphAttributeListView::CalculateColumnWidth(InColumnInfo.Label) : InColumnInfo.Width;
 	ColumnsMaxWidthMapping.FindOrAdd(InColumnInfo.Id) = ColumnWidth;
diff --git a/Engine/Plugins/PCG/Source/PCGEditor/Private/Widgets/SPCGEditorGraphAttributeListView.h b/Engine/Plugins/PCG/Source/PCGEditor/Private/Widgets/SPCGEditorGraphAttributeListView.h
--- a/Engine/Plugins/PCG/Source/PCGEditor/Private/Widgets/SPCGEditorGraphAttributeListView.h
+++ b/Engine/Plugins/PCG/Source/PCGEditor/Private/Widgets/SPCGEditorGraphAttributeListView.h
@@ -57,6 +57,8 @@ struct FPCGColumnData
 {
 	TSharedPtr<const IPCGAttributeAccessor> DataAccessor;
 	TSharedPtr<const IPCGAttributeAccessorKeys> DataKeys;
+	// @todo_pcg: To be removed when we do not make the distinction between typed and generic.
+	const FPCGMetadataAttributeGeneric* GenericAttribute = nullptr;
 };
 
 typedef TSharedPtr<FPCGListViewItem> PCGListViewItemPtr;
diff --git a/Engine/Plugins/PCG/Source/PCGEditor/Private/Widgets/SPCGMetadataAttributeGenericViewer.cpp b/Engine/Plugins/PCG/Source/PCGEditor/Private/Widgets/SPCGMetadataAttributeGenericViewer.cpp
new file mode 100644
--- /dev/null
+++ b/Engine/Plugins/PCG/Source/PCGEditor/Private/Widgets/SPCGMetadataAttributeGenericViewer.cpp
@@ -0,0 +1,106 @@
+﻿// Copyright Epic Games, Inc. All Rights Reserved.
+
+#include "Widgets/SPCGMetadataAttributeGenericViewer.h"
+
+#include "DetailCategoryBuilder.h"
+#include "DetailLayoutBuilder.h"
+#include "PCGData.h"
+
+#include "DetailsViewArgs.h"
+#include "InstancedPropertyBagStructureDataProvider.h"
+#include "PropertyEditorModule.h"
+#include "Metadata/PCGMetadataAttributeGeneric.h"
+#include "Modules/ModuleManager.h"
+
+TSharedRef<IDetailCustomization> FPCGMetadataAttributeGenericViewerDetailCustomization::MakeInstance(TSharedPtr<SPCGMetadataAttributeGenericViewer> Viewer)
+{
+	return MakeShareable(new FPCGMetadataAttributeGenericViewerDetailCustomization(Viewer));
+}
+
+void FPCGMetadataAttributeGenericViewerDetailCustomization::CustomizeDetails(IDetailLayoutBuilder& DetailBuilder)
+{
+	const TArray<TWeakObjectPtr<UObject>>& Objects = DetailBuilder.GetSelectedObjects();
+	if(Objects.Num() != 1 || !Objects[0].IsValid())
+	{
+		return;
+	}
+	
+	if (!WeakViewer.IsValid())
+	{
+		return;
+	}
+	
+	FInstancedPropertyBag& Struct = WeakViewer.Pin()->TempStruct;
+	const UPropertyBag* Bag = Struct.GetPropertyBagStruct();
+	if (!Bag || Bag->GetPropertyDescs().IsEmpty())
+	{
+		return;
+	}
+	
+	const FName PropertyName = Bag->GetPropertyDescs()[0].Name;
+
+	IDetailCategoryBuilder& CategoryBuilder = DetailBuilder.EditCategory(TEXT("NoCategory"));
+	CategoryBuilder.AddExternalStructureProperty(MakeShared<FInstancePropertyBagStructureDataProvider>(Struct), PropertyName);
+}
+
+FPCGMetadataAttributeGenericViewerDetailCustomization::FPCGMetadataAttributeGenericViewerDetailCustomization(TSharedPtr<SPCGMetadataAttributeGenericViewer> Viewer)
+	: WeakViewer(Viewer)
+{
+}
+	
+void SPCGMetadataAttributeGenericViewer::Construct(const FArguments& InArgs)
+{
+	FDetailsViewArgs DetailsViewArgs;
+	DetailsViewArgs.NameAreaSettings = FDetailsViewArgs::HideNameArea;
+	DetailsViewArgs.bShowHeaders = false;
+
+	FPropertyEditorModule& PropertyEditorModule = FModuleManager::GetModuleChecked<FPropertyEditorModule>("PropertyEditor");
+	DetailsView = PropertyEditorModule.CreateDetailView(DetailsViewArgs);
+
+	
+	DetailsView->RegisterInstancedCustomPropertyLayout(UPCGData::StaticClass(), FOnGetDetailCustomizationInstance::CreateStatic(&FPCGMetadataAttributeGenericViewerDetailCustomization::MakeInstance, StaticCastSharedPtr<SPCGMetadataAttributeGenericViewer>(AsShared().ToSharedPtr())));
+	
+	DetailsView->SetIsPropertyReadOnlyDelegate(FIsPropertyReadOnly::CreateLambda([](const FPropertyAndParent& In)
+	{
+		return true;
+	}));
+	
+	DetailsView->SetIsPropertyVisibleDelegate(FIsPropertyVisible::CreateLambda([](const FPropertyAndParent& In)
+	{
+		// Make sure to remove all the properties that are coming from the Data (aka, has an object attached to it).
+		return In.Objects.IsEmpty();
+	}));
+
+	ChildSlot
+	[
+		DetailsView.ToSharedRef()
+	];
+}
+
+void SPCGMetadataAttributeGenericViewer::Refresh()
+{
+	if (DetailsView)
+	{
+		if (Data.IsValid() && Attr)
+		{
+			TempStruct = Attr->BuildStructForDebug(EntryKey);
+		}
+		else
+		{
+			TempStruct.Reset();
+		}
+		
+		DetailsView->SetObject(const_cast<UPCGData*>(Data.Get()), /*bForceRefresh=*/true);
+		
+		// Make sure the root is expanded by default, but not necessarily all its children.
+		DetailsView->SetRootExpansionStates(/*bExpand=*/true, /*bRecurse=*/false);
+	}
+}
+
+void SPCGMetadataAttributeGenericViewer::Setup(TWeakObjectPtr<const UPCGData> InData, const FPCGMetadataAttributeGeneric* InAttr, PCGMetadataEntryKey InEntryKey)
+{
+	Data = InData;
+	Attr = InAttr;
+	EntryKey = InEntryKey;
+	Refresh();
+}
diff --git a/Engine/Plugins/PCG/Source/PCGEditor/Private/Widgets/SPCGMetadataAttributeGenericViewer.h b/Engine/Plugins/PCG/Source/PCGEditor/Private/Widgets/SPCGMetadataAttributeGenericViewer.h
new file mode 100644
--- /dev/null
+++ b/Engine/Plugins/PCG/Source/PCGEditor/Private/Widgets/SPCGMetadataAttributeGenericViewer.h
@@ -0,0 +1,69 @@
+﻿// Copyright Epic Games, Inc. All Rights Reserved.
+
+#pragma once
+
+#include "IDetailCustomization.h"
+
+
+#include "Components/DetailsView.h"
+#include "Metadata/PCGMetadataCommon.h"
+#include "StructUtils/PropertyBag.h"
+
+#include "Widgets/SCompoundWidget.h"
+#include "Widgets/Layout/SExpandableArea.h"
+#include "Widgets/Views/ITableRow.h"
+#include "Widgets/Views/SListView.h"
+#include "Widgets/Views/STableRow.h"
+
+class UPCGData;
+class FPCGMetadataAttributeGeneric;
+class FPCGEditor;
+class IDetailsView;
+class SPCGMetadataAttributeGenericViewer;
+
+/** Header customization for when multiple objects are displayed in the details panel. */
+class FPCGMetadataAttributeGenericViewerDetailCustomization final : public IDetailCustomization
+{
+public:
+	static TSharedRef<IDetailCustomization> MakeInstance(TSharedPtr<SPCGMetadataAttributeGenericViewer> Viewer);
+
+	virtual void CustomizeDetails(IDetailLayoutBuilder& DetailBuilder) override;
+
+private:
+	FPCGMetadataAttributeGenericViewerDetailCustomization(TSharedPtr<SPCGMetadataAttributeGenericViewer> Viewer);
+	TWeakPtr<SPCGMetadataAttributeGenericViewer> WeakViewer; 
+};
+
+/** Thin wrapper on top of a SDetailsView to add some PCG-specific logic around locking and improve layout */
+class SPCGMetadataAttributeGenericViewer : public SCompoundWidget
+{
+public:
+	SLATE_BEGIN_ARGS(SPCGMetadataAttributeGenericViewer)
+		{}	
+	SLATE_END_ARGS()
+
+	/** Constructs this widget with InArgs */
+	void Construct(const FArguments& InArgs);
+
+	/** Sets the editor associated to this details view, used for interaction from the details view to the editor */
+	void SetEditor(TWeakPtr<FPCGEditor> InEditorPtr) { EditorPtr = InEditorPtr; }
+
+	/** Returns the details view under this widget */
+	TSharedPtr<IDetailsView> GetDetailsView() const { return DetailsView; }
+	
+	void Refresh();
+
+	/** Controls lock from the editor */
+	void Setup(TWeakObjectPtr<const UPCGData> InData, const FPCGMetadataAttributeGeneric* InAttr, PCGMetadataEntryKey InEntryKey);
+
+protected:
+	friend FPCGMetadataAttributeGenericViewerDetailCustomization;
+	TSharedPtr<IDetailsView> DetailsView;
+
+	FInstancedPropertyBag TempStruct;
+	
+	TWeakObjectPtr<const UPCGData> Data;
+	const FPCGMetadataAttributeGeneric* Attr;
+	PCGMetadataEntryKey EntryKey;
+	TWeakPtr<FPCGEditor> EditorPtr = nullptr;
+};
\ No newline at end of file
