[PCG - Sandbox] Temporary nodes to create/add constant attributes

* In the PCGSample sandbox at the moment to avoid people starting to use it.

#rnx
#Synced-CL 50489647

diff --git a/Engine/Plugins/PCG/Source/PCGEditor/Private/Schema/PCGEditorGraphSchema.cpp b/Engine/Plugins/PCG/Source/PCGEditor/Private/Schema/PCGEditorGraphSchema.cpp
--- a/Engine/Plugins/PCG/Source/PCGEditor/Private/Schema/PCGEditorGraphSchema.cpp
+++ b/Engine/Plugins/PCG/Source/PCGEditor/Private/Schema/PCGEditorGraphSchema.cpp
@@ -967,6 +967,21 @@ FLinearColor UPCGEditorGraphSchema::GetMetadataTypeColor(const EPCGMetadataTypes
 			return GraphSettings->SoftObjectPinTypeColor;
 		case EPCGMetadataTypes::SoftClassPath:
 			return GraphSettings->SoftClassPinTypeColor;
+		case EPCGMetadataTypes::Text:
+			return GraphSettings->TextPinTypeColor;
+		case EPCGMetadataTypes::Struct:
+			return GraphSettings->StructPinTypeColor;
+		case EPCGMetadataTypes::Object:
+			return GraphSettings->ObjectPinTypeColor;
+		case EPCGMetadataTypes::SoftObject:
+			return GraphSettings->SoftObjectPinTypeColor;
+		case EPCGMetadataTypes::Class:
+			return GraphSettings->ClassPinTypeColor;
+		case EPCGMetadataTypes::SoftClass:
+			return GraphSettings->SoftClassPinTypeColor;
+		case EPCGMetadataTypes::Byte: // fall-through
+		case EPCGMetadataTypes::Enum:
+			return GraphSettings->BytePinTypeColor;
 		default: // fall-through
 		case EPCGMetadataTypes::Unknown:
 			return FLinearColor::White;
