[PCG] Generate Landscape Textures - Gather landscape components from the LandscapeInfo rather than looping over every landscape proxy to find overlapping components

Fixes issue where GLT node no-ops on LandscapeData that were found via a soft object path

#rnx
#Synced-CL 50643236

#Synced-CL 50638379

diff --git a/Engine/Plugins/PCG/Source/PCG/Private/Elements/PCGGenerateLandscapeTextures.cpp b/Engine/Plugins/PCG/Source/PCG/Private/Elements/PCGGenerateLandscapeTextures.cpp
--- a/Engine/Plugins/PCG/Source/PCG/Private/Elements/PCGGenerateLandscapeTextures.cpp
+++ b/Engine/Plugins/PCG/Source/PCG/Private/Elements/PCGGenerateLandscapeTextures.cpp
@@ -666,10 +666,7 @@ bool FPCGGenerateLandscapeTexturesElement::ExecuteInternal(FPCGContext* InContex
 			FMemory::Memset(Context->OutputIndexToExportIndex.GetData(), -1, Context->OutputIndexToExportIndex.Num() * Context->OutputIndexToExportIndex.GetTypeSize());
 		}
 
-		Context->GenerationBounds = ExecutionSource->GetExecutionState().GetBounds();
-		TArray<ULandscapeComponent*> OverlappingComponents;
-		TSet<FName> ProcessedGrassNames;
-
+		// Find the first valid landscape - only one supported right now.
 		for (const TSoftObjectPtr<ALandscapeProxy>& LandscapeProxyPtr : LandscapeData->Landscapes)
 		{
 			ALandscapeProxy* LandscapeProxy = LandscapeProxyPtr.Get();
@@ -684,11 +681,6 @@ bool FPCGGenerateLandscapeTexturesElement::ExecuteInternal(FPCGContext* InContex
 				continue;
 			}
 
-			// Grow storage to accommodate new components.
-			const int32 NumToReserve = FMath::Min(OverlappingComponents.Num() + LandscapeProxy->LandscapeComponents.Num(), (int32)FPCGGrassMapUnpackerCS::MaxNumLandscapeComponents);
-			OverlappingComponents.Reserve(NumToReserve);
-			Context->LandscapeComponents.Reserve(NumToReserve);
-
 			if (!Context->Landscape.Get())
 			{
 				Context->Landscape = Landscape;
@@ -696,20 +688,40 @@ bool FPCGGenerateLandscapeTexturesElement::ExecuteInternal(FPCGContext* InContex
 			// TODO [jonathan.bard] : The grass map exporter doesn't support landscape components from different landscapes, simply skip those components for now : 
 			else if (Context->Landscape.Get() != Landscape)
 			{
-				PCGLog::LogWarningOnGraph(FText::Format(LOCTEXT("LandscapeMismatch", "Encountered multiple landscapes which is not supported at the moment. Only the first landscape will be considered : {0} (current : {1})"), 
+				PCGLog::LogWarningOnGraph(FText::Format(LOCTEXT("LandscapeMismatch", "Encountered multiple landscapes which is not supported at the moment. Only the first landscape will be considered : {0} (current : {1})"),
 					FText::FromString(Context->Landscape->GetActorNameOrLabel()), FText::FromString(Landscape->GetActorNameOrLabel())), Context);
-				continue;
 			}
+		}
+
+		if (!Context->Landscape.Get() || !Context->Landscape->GetLandscapeInfo())
+		{
+			return true;
+		}
 
-			for (int32 ComponentIndex = 0; ComponentIndex < LandscapeProxy->LandscapeComponents.Num(); ++ComponentIndex)
+		// Gather all landscape components overlapping our generation bounds.
+		{
+			Context->GenerationBounds = ExecutionSource->GetExecutionState().GetBounds();
+
+			const FBox2D Extent2D = FBox2D(FVector2D(Context->GenerationBounds.Min), FVector2D(Context->GenerationBounds.Max));
+			TMap<FIntPoint, ULandscapeComponent*> OverlappingComponents;
+			FIntRect ComponentIndicesBoundingRect;
+			Context->Landscape->GetLandscapeInfo()->GetOverlappedComponents(FTransform::Identity, Extent2D, OverlappingComponents, ComponentIndicesBoundingRect);
+
+			// Grow storage to accommodate new components.
+			const int32 NumToReserve = FMath::Min(OverlappingComponents.Num(), (int32)FPCGGrassMapUnpackerCS::MaxNumLandscapeComponents);
+			Context->LandscapeComponents.Reserve(NumToReserve);
+
+			TSet<FName> ProcessedGrassNames;
+
+			for (auto& [Coordinate, LandscapeComponent] : OverlappingComponents)
 			{
-				if (ULandscapeComponent* LandscapeComponent = LandscapeProxy->LandscapeComponents[ComponentIndex])
+				if (ensure(LandscapeComponent))
 				{
 					// Only generate grass map if there is meaningful overlap with our domain of interest.
 					const FBox LandscapeComponentBounds = LandscapeComponent->Bounds.GetBox();
 					if (LandscapeComponentBounds.Overlap(Context->GenerationBounds).GetVolume() > UE_KINDA_SMALL_NUMBER)
 					{
-						if (OverlappingComponents.Num() >= FPCGGrassMapUnpackerCS::MaxNumLandscapeComponents)
+						if (Context->LandscapeComponents.Num() >= FPCGGrassMapUnpackerCS::MaxNumLandscapeComponents)
 						{
 							PCGLog::LogWarningOnGraph(LOCTEXT("MaxLandscapeComponentsExceeded", "Too many landscape components overlap the generation domain. Consider partitioning the component onto a smaller grid size."), Context);
 							break;
@@ -787,14 +799,12 @@ bool FPCGGenerateLandscapeTexturesElement::ExecuteInternal(FPCGContext* InContex
 						}
 
 						Context->LandscapeComponents.Add(LandscapeComponent);
-						OverlappingComponents.Add(LandscapeComponent);
-
 						Context->TotalLandscapeComponentBounds += LandscapeComponentBounds;
 
 						if (Context->LandscapeComponents.Num() == 1)
 						{
 							Context->LandscapeComponentExtent = LandscapeComponent->Bounds.BoxExtent.X * 2.0;
-							Context->TexelSizeWorld = Context->LandscapeComponentExtent / static_cast<double>(LandscapeProxy->ComponentSizeQuads);
+							Context->TexelSizeWorld = Context->LandscapeComponentExtent / static_cast<double>(Context->Landscape->ComponentSizeQuads);
 
 							if (!ensure(FMath::IsNearlyEqual(LandscapeComponent->Bounds.BoxExtent.X, LandscapeComponent->Bounds.BoxExtent.Y)))
 							{
@@ -812,7 +822,7 @@ bool FPCGGenerateLandscapeTexturesElement::ExecuteInternal(FPCGContext* InContex
 							ensure(FMath::IsNearlyEqual(Context->LandscapeComponentExtent, LandscapeComponent->Bounds.BoxExtent.X * 2.0));
 						}
 
-	#if WITH_EDITOR
+#if WITH_EDITOR
 						if (PCGGenerateLandscapeTextures::CVarDebugDrawGeneratedComponents.GetValueOnGameThread())
 						{
 							DrawDebugBox(ExecutionSource->GetExecutionState().GetWorld(),
@@ -822,13 +832,13 @@ bool FPCGGenerateLandscapeTexturesElement::ExecuteInternal(FPCGContext* InContex
 								/*bPersistentLines=*/false,
 								/*LifeTime=*/3.0f);
 						}
-	#endif
+#endif
 					}
 				}
 			}
 		}
 
-		if (!Context->Landscape.Get() || Context->LandscapeComponents.IsEmpty() || (Context->GrassNamesToOutput.IsEmpty() && !Settings->bGenerateHeightMap) || !Context->TotalLandscapeComponentBounds.IsValid)
+		if (Context->LandscapeComponents.IsEmpty() || (Context->GrassNamesToOutput.IsEmpty() && !Settings->bGenerateHeightMap) || !Context->TotalLandscapeComponentBounds.IsValid)
 		{
 			return true;
 		}
