[PCG] Revert changes to PISM refresh path introduced alongside fixes for instance flushing

#rnx

#Synced-CL 50807095

diff --git a/Engine/Plugins/PCG/Source/PCG/Private/Components/PCGProceduralISMComponent.cpp b/Engine/Plugins/PCG/Source/PCG/Private/Components/PCGProceduralISMComponent.cpp
--- a/Engine/Plugins/PCG/Source/PCG/Private/Components/PCGProceduralISMComponent.cpp
+++ b/Engine/Plugins/PCG/Source/PCG/Private/Components/PCGProceduralISMComponent.cpp
@@ -416,10 +416,62 @@ void UPCGProceduralISMComponent::OnRenderStateDirty(UActorComponent& InComponent
 
 	UPCGSubsystem* Subsystem = GetOwner() ? UPCGSubsystem::GetInstance(GetOwner()->GetWorld()) : nullptr;
 
-	if (Subsystem)
+	if (!Subsystem)
 	{
-		Subsystem->DirtyRuntimeGenExecutionSources();
+		return;
 	}
+
+	// Helper that returns true if the given PCG component is managing this PISMC.
+	auto PCGComponentManagesThisPISMC = [this](UPCGComponent* InComponent) -> bool
+	{
+		bool bManagesThis = false;
+
+		if (InComponent->AreProceduralInstancesInUse() && InComponent->AreManagedResourcesAccessible())
+		{
+			InComponent->ForEachManagedResource([this, InComponent, &bManagesThis](UPCGManagedResource* InResource)
+			{
+				if (!bManagesThis)
+				{
+					const UPCGManagedProceduralISMComponent* PISMC = Cast<UPCGManagedProceduralISMComponent>(InResource);
+					if (PISMC && PISMC->GetComponent() == this)
+					{
+						bManagesThis = true;
+					}
+				}
+			});
+		}
+
+		return bManagesThis;
+	};
+
+	Subsystem->RefreshAllComponentsFiltered(
+		[this, Subsystem, &PCGComponentManagesThisPISMC](UPCGComponent* InComponent) -> bool
+		{
+			// If the original component manages this PISMC, request a refresh of it and we're done. If it has local components they will also be refreshed.
+			if (PCGComponentManagesThisPISMC(InComponent))
+			{
+				return true;
+			}
+
+			// A local component of the original component might manage this, so check those.
+			bool bManagesThis = false;
+
+			if (InComponent->bIsComponentPartitioned)
+			{
+				Subsystem->ForAllRegisteredLocalComponents(
+					InComponent,
+					[this, &bManagesThis, &PCGComponentManagesThisPISMC](UPCGComponent* InComponent)
+					{
+						if (!bManagesThis)
+						{
+							bManagesThis = PCGComponentManagesThisPISMC(InComponent);
+						}
+					});
+			}
+
+			return bManagesThis;
+		},
+		EPCGChangeType::Structural);
 }
 #endif // WITH_EDITOR
 
