[PCG] TrackingManager fix execution source generated tracking issues
- Add a IPCGGraphExecutionState::GetTotalBounds which returns the source's bounds including its generated resource bounds

#Synced-CL 50905358

diff --git a/Engine/Plugins/PCG/Source/PCG/Private/Graph/PCGGraphExecutor.h b/Engine/Plugins/PCG/Source/PCG/Private/Graph/PCGGraphExecutor.h
--- a/Engine/Plugins/PCG/Source/PCG/Private/Graph/PCGGraphExecutor.h
+++ b/Engine/Plugins/PCG/Source/PCG/Private/Graph/PCGGraphExecutor.h
@@ -174,7 +174,7 @@ private:
 	TOptional<FBox> GetOriginalBounds(FPCGTaskId InGraphExecutionTaskId) { return PerExecutionCache.GetExecutionCacheBounds(InGraphExecutionTaskId, EPCGPerExecutionCacheDataType::OriginalBounds); }
 	TOptional<FBox> GetLocalSpaceBounds(FPCGTaskId InGraphExecutionTaskId) { return PerExecutionCache.GetExecutionCacheBounds(InGraphExecutionTaskId, EPCGPerExecutionCacheDataType::LocalSpaceBounds); }
 	TOptional<FBox> GetOriginalLocalSpaceBounds(FPCGTaskId InGraphExecutionTaskId) { return PerExecutionCache.GetExecutionCacheBounds(InGraphExecutionTaskId, EPCGPerExecutionCacheDataType::OriginalLocalSpaceBounds); }
-
+	TOptional<FBox> GetTotalBounds(FPCGTaskId InGraphExecutionTaskId) { return PerExecutionCache.GetExecutionCacheBounds(InGraphExecutionTaskId, EPCGPerExecutionCacheDataType::TotalBounds); }
 		
 	void SetPCGData(FPCGTaskId InGraphExecutionTaskId, UPCGData* InData) { PerExecutionCache.SetExecutionCachePCGData(InGraphExecutionTaskId, EPCGPerExecutionCacheDataType::PCGData, InData); }
 	void SetInputPCGData(FPCGTaskId InGraphExecutionTaskId, UPCGData* InData) { PerExecutionCache.SetExecutionCachePCGData(InGraphExecutionTaskId, EPCGPerExecutionCacheDataType::InputData, InData); }
@@ -186,6 +186,7 @@ private:
 	void SetOriginalBounds(FPCGTaskId InGraphExecutionTaskId, const FBox& InBounds) { PerExecutionCache.SetExecutionCacheBounds(InGraphExecutionTaskId, EPCGPerExecutionCacheDataType::OriginalBounds, InBounds); }
 	void SetLocalSpaceBounds(FPCGTaskId InGraphExecutionTaskId, const FBox& InBounds) { PerExecutionCache.SetExecutionCacheBounds(InGraphExecutionTaskId, EPCGPerExecutionCacheDataType::LocalSpaceBounds, InBounds); }
 	void SetOriginalLocalSpaceBounds(FPCGTaskId InGraphExecutionTaskId, const FBox& InBounds) { PerExecutionCache.SetExecutionCacheBounds(InGraphExecutionTaskId, EPCGPerExecutionCacheDataType::OriginalLocalSpaceBounds, InBounds); }
+	void SetTotalBounds(FPCGTaskId InGraphExecutionTaskId, const FBox& InBounds) { PerExecutionCache.SetExecutionCacheBounds(InGraphExecutionTaskId, EPCGPerExecutionCacheDataType::TotalBounds, InBounds); }
 
 	double GetTickBudgetInSeconds() const;
 	int32 GetMaxNumProcessingTasks() const;
diff --git a/Engine/Plugins/PCG/Source/PCG/Private/Graph/PCGGraphPerExecutionCache.h b/Engine/Plugins/PCG/Source/PCG/Private/Graph/PCGGraphPerExecutionCache.h
--- a/Engine/Plugins/PCG/Source/PCG/Private/Graph/PCGGraphPerExecutionCache.h
+++ b/Engine/Plugins/PCG/Source/PCG/Private/Graph/PCGGraphPerExecutionCache.h
@@ -64,6 +64,7 @@ enum class EPCGPerExecutionCacheDataType : uint8
 	OriginalBounds,
 	LocalSpaceBounds,
 	OriginalLocalSpaceBounds,
+	TotalBounds,
 	Count
 };
 
diff --git a/Engine/Plugins/PCG/Source/PCG/Private/PCGComponent.cpp b/Engine/Plugins/PCG/Source/PCG/Private/PCGComponent.cpp
--- a/Engine/Plugins/PCG/Source/PCG/Private/PCGComponent.cpp
+++ b/Engine/Plugins/PCG/Source/PCG/Private/PCGComponent.cpp
@@ -3192,6 +3192,7 @@ void UPCGComponent::ExecutePreGraph(FPCGContext* InContext)
 		GetOriginalGridBounds();
 		GetLocalSpaceBounds();
 		GetOriginalLocalSpaceBounds();
+		GetTotalBounds();
 	}
 }
 
@@ -3477,6 +3478,37 @@ FBox UPCGComponent::GetOriginalLocalSpaceBounds() const
 	}
 }
 
+FBox UPCGComponent::GetTotalBounds() const
+{
+	/* Total bounds are cached at the beginning of execution through UPCGComponent::ExecutePreGraph for tracking purposes of the previously generated resources. */
+	if (CurrentGenerationTask != InvalidPCGTaskId)
+	{
+		UPCGSubsystem* PCGSubsystem = GetSubsystem();
+		check(PCGSubsystem);
+
+		TOptional<FBox> TotalBounds = PCGSubsystem->GetTotalBounds(CurrentGenerationTask);
+
+		if (!TotalBounds.IsSet())
+		{
+			PCG_EXECUTION_CACHE_VALIDATION_CHECK(this);
+			TotalBounds = PCGHelpers::GetActorBounds(GetOwner(), /*bIgnorePCGCreatedComponents=*/false);
+
+			PCGSubsystem->SetTotalBounds(CurrentGenerationTask, TotalBounds.GetValue());
+			if (PCGSubsystem->IsGraphCacheDebuggingEnabled())
+			{
+				UE_LOG(LogPCG, Log, TEXT("         [%s] CACHE REFRESH Total Bounds"), *GetOwner()->GetName());
+			}
+		}
+
+		check(TotalBounds.IsSet());
+		return TotalBounds.GetValue();
+	}
+	else
+	{
+		return PCGHelpers::GetActorBounds(GetOwner(), /*bIgnorePCGCreatedComponents=*/false);
+	}
+}
+
 FBox UPCGComponent::GetGridBounds(const AActor* Actor) const
 {
 	return PCGHelpers::GetGridBounds(Actor, this);
diff --git a/Engine/Plugins/PCG/Source/PCG/Private/PCGComponentExecutionState.cpp b/Engine/Plugins/PCG/Source/PCG/Private/PCGComponentExecutionState.cpp
--- a/Engine/Plugins/PCG/Source/PCG/Private/PCGComponentExecutionState.cpp
+++ b/Engine/Plugins/PCG/Source/PCG/Private/PCGComponentExecutionState.cpp
@@ -93,6 +93,12 @@ FBox FPCGComponentExecutionState::GetOriginalLocalSpaceBounds() const
 	return Component->GetOriginalLocalSpaceBounds();
 }
 
+FBox FPCGComponentExecutionState::GetTotalBounds() const
+{
+	check(Component);
+	return Component->GetTotalBounds();
+}
+
 bool FPCGComponentExecutionState::Use2DGrid() const
 {
 	check(Component);
diff --git a/Engine/Plugins/PCG/Source/PCG/Private/PCGTrackingManager.cpp b/Engine/Plugins/PCG/Source/PCG/Private/PCGTrackingManager.cpp
--- a/Engine/Plugins/PCG/Source/PCG/Private/PCGTrackingManager.cpp
+++ b/Engine/Plugins/PCG/Source/PCG/Private/PCGTrackingManager.cpp
@@ -1724,11 +1724,18 @@ void FPCGTrackingManager::OnPCGGraphGeneratedOrCleaned(IPCGGraphExecutionSource*
 		return;
 	}
 
-	// @todo_pcg: support execution source
+	const UObject* ChangedObject = Cast<UObject>(InExecutionSource);
+	
+	// @todo_pcg: add concept of Owner to the execution source? to abstract out the actor code here.
+	const UObject* ChangedRoot = ChangedObject;
 	if (UPCGComponent* PCGComponent = Cast<UPCGComponent>(InExecutionSource))
 	{
-		OnObjectChanged(PCGComponent->GetOwner(), PCGTrackingManager::ExecutionSourceGenerated, PCGComponent);
+		ChangedRoot = PCGComponent->GetOwner();
 	}
+
+	FPCGSelectionKey SelectionKey = FPCGSelectionKey::CreateFromObjectPtr(ChangedRoot);
+
+	OnSelectionKeyChanged(SelectionKey, ChangedObject, /*InRemoveTags=*/{}, /*bNoRefreshOwner=*/true, { InExecutionSource->GetExecutionState().GetTotalBounds() });
 }
 
 bool FPCGTrackingManager::ClearCacheForKeys(TArrayView<FPCGSelectionKey> InKeys, const IPCGGraphExecutionSource* InExecutionSource, const bool bIntersect, const UObject* InOriginatingChange) const
diff --git a/Engine/Plugins/PCG/Source/PCG/Private/Subsystems/PCGSubsystem.cpp b/Engine/Plugins/PCG/Source/PCG/Private/Subsystems/PCGSubsystem.cpp
--- a/Engine/Plugins/PCG/Source/PCG/Private/Subsystems/PCGSubsystem.cpp
+++ b/Engine/Plugins/PCG/Source/PCG/Private/Subsystems/PCGSubsystem.cpp
@@ -2151,6 +2151,11 @@ TOptional<FBox> UPCGSubsystem::GetOriginalLocalSpaceBounds(FPCGTaskId InGraphExe
 	return GraphExecutor ? GraphExecutor->GetOriginalLocalSpaceBounds(InGraphExecutionTaskId) : TOptional<FBox>();
 }
 
+TOptional<FBox> UPCGSubsystem::GetTotalBounds(FPCGTaskId InGraphExecutionTaskId)
+{
+	return GraphExecutor ? GraphExecutor->GetTotalBounds(InGraphExecutionTaskId) : TOptional<FBox>();
+}
+
 void UPCGSubsystem::SetPCGData(FPCGTaskId InGraphExecutionTaskId, UPCGData* InData)
 {
 	if (GraphExecutor)
@@ -2231,6 +2236,14 @@ void UPCGSubsystem::SetOriginalLocalSpaceBounds(FPCGTaskId InGraphExecutionTaskI
 	}
 }
 
+void UPCGSubsystem::SetTotalBounds(FPCGTaskId InGraphExecutionTaskId, const FBox& InBounds)
+{
+	if (GraphExecutor)
+	{
+		GraphExecutor->SetTotalBounds(InGraphExecutionTaskId, InBounds);
+	}
+}
+
 void UPCGSubsystem::ExecuteBeginTickActions()
 {
 	TArray<FTickAction> Actions = MoveTemp(BeginTickActions);
diff --git a/Engine/Plugins/PCG/Source/PCG/Public/PCGComponent.h b/Engine/Plugins/PCG/Source/PCG/Public/PCGComponent.h
--- a/Engine/Plugins/PCG/Source/PCG/Public/PCGComponent.h
+++ b/Engine/Plugins/PCG/Source/PCG/Public/PCGComponent.h
@@ -554,6 +554,7 @@ public:
 	UE_API FBox GetOriginalGridBounds() const;
 	UE_API FBox GetLocalSpaceBounds() const;
 	UE_API FBox GetOriginalLocalSpaceBounds() const;
+	UE_API FBox GetTotalBounds() const;
 
 	FBox GetLastGeneratedBounds() const { return LastGeneratedBounds; }
 
diff --git a/Engine/Plugins/PCG/Source/PCG/Public/PCGComponentExecutionState.h b/Engine/Plugins/PCG/Source/PCG/Public/PCGComponentExecutionState.h
--- a/Engine/Plugins/PCG/Source/PCG/Public/PCGComponentExecutionState.h
+++ b/Engine/Plugins/PCG/Source/PCG/Public/PCGComponentExecutionState.h
@@ -24,6 +24,7 @@ public:
 	virtual FBox GetOriginalBounds() const override;
 	virtual FBox GetLocalSpaceBounds() const override;
 	virtual FBox GetOriginalLocalSpaceBounds() const override;
+	virtual FBox GetTotalBounds() const override;
 	virtual bool Use2DGrid() const override;
 	virtual UPCGGraph* GetGraph() const override;
 	virtual UPCGGraphInstance* GetGraphInstance() const override;
diff --git a/Engine/Plugins/PCG/Source/PCG/Public/PCGGraphExecutionStateInterface.h b/Engine/Plugins/PCG/Source/PCG/Public/PCGGraphExecutionStateInterface.h
--- a/Engine/Plugins/PCG/Source/PCG/Public/PCGGraphExecutionStateInterface.h
+++ b/Engine/Plugins/PCG/Source/PCG/Public/PCGGraphExecutionStateInterface.h
@@ -110,6 +110,9 @@ public:
 	/** Returns the ExecutionState's original source local space bounds if this ExecutionState is a local source. */
 	virtual FBox GetOriginalLocalSpaceBounds() const { return GetBounds(); }
 
+	/** Returns the ExecutionState's generated bounds which includes the artifact bounds. */
+	virtual FBox GetTotalBounds() const { return GetBounds(); }
+
 	/** Returns the UPCGGraph this ExecutionState is executing. */
 	virtual UPCGGraph* GetGraph() const = 0;
 
diff --git a/Engine/Plugins/PCG/Source/PCG/Public/Subsystems/PCGSubsystem.h b/Engine/Plugins/PCG/Source/PCG/Public/Subsystems/PCGSubsystem.h
--- a/Engine/Plugins/PCG/Source/PCG/Public/Subsystems/PCGSubsystem.h
+++ b/Engine/Plugins/PCG/Source/PCG/Public/Subsystems/PCGSubsystem.h
@@ -337,6 +337,7 @@ public:
 	PCG_API TOptional<FBox> GetOriginalBounds(FPCGTaskId InGraphExecutionTaskId);
 	PCG_API TOptional<FBox> GetLocalSpaceBounds(FPCGTaskId InGraphExecutionTaskId);
 	PCG_API TOptional<FBox> GetOriginalLocalSpaceBounds(FPCGTaskId InGraphExecutionTaskId);
+	PCG_API TOptional<FBox> GetTotalBounds(FPCGTaskId InGraphExecutionTaskId);
 
 	/** Setters for per execution data. Usually called from execution state ExecutePreGraph implementations. */
 	PCG_API void SetPCGData(FPCGTaskId InGraphExecutionTaskId, UPCGData* InData);
@@ -349,6 +350,7 @@ public:
 	PCG_API void SetOriginalBounds(FPCGTaskId InGraphExecutionTaskId, const FBox& InBounds);
 	PCG_API void SetLocalSpaceBounds(FPCGTaskId InGraphExecutionTaskId, const FBox& InBounds);
 	PCG_API void SetOriginalLocalSpaceBounds(FPCGTaskId InGraphExecutionTaskId, const FBox& InBounds);
+	PCG_API void SetTotalBounds(FPCGTaskId InGraphExecutionTaskId, const FBox& InBounds);
 
 #if WITH_EDITOR
 public:
