[PCG] Move api for FPCGManagedResourceContainer inside FPCGManagedResourceContainerHelper to remove need for Owner pointer
- Helper created on stack with owner avoids having to maintain that transient owner pointer

#rnx

#Synced-CL 50905358

diff --git a/Engine/Plugins/Experimental/PCGInterops/PCGFastGeoInterop/Source/PCGFastGeoInterop/Private/Compute/PrimitiveFactories/PCGPrimitiveFactoryFastGeoPISMC.cpp b/Engine/Plugins/Experimental/PCGInterops/PCGFastGeoInterop/Source/PCGFastGeoInterop/Private/Compute/PrimitiveFactories/PCGPrimitiveFactoryFastGeoPISMC.cpp
--- a/Engine/Plugins/Experimental/PCGInterops/PCGFastGeoInterop/Source/PCGFastGeoInterop/Private/Compute/PrimitiveFactories/PCGPrimitiveFactoryFastGeoPISMC.cpp
+++ b/Engine/Plugins/Experimental/PCGInterops/PCGFastGeoInterop/Source/PCGFastGeoInterop/Private/Compute/PrimitiveFactories/PCGPrimitiveFactoryFastGeoPISMC.cpp
@@ -186,9 +186,10 @@ bool FPCGPrimitiveFactoryFastGeoPISMC::Create(FPCGContext* InContext)
 	ManagedPrimitives->SetFastGeoContainer(FastGeo);
 	ManagedPrimitives->SetObjectReferences(CollectObjectReferences());
 
-	if (FPCGManagedResourceContainer* Container = ExecutionSource->GetExecutionState().GetManagedResourceContainer(); ensure(Container))
+	FPCGManagedResourceContainerHelper ContainerHelper(ExecutionSource);
+	if (ensure(ContainerHelper.IsValid()))
 	{
-		Container->AddManagedResource(ManagedPrimitives);
+		ContainerHelper.AddManagedResource(ManagedPrimitives);
 	}
 
 	return true;
diff --git a/Engine/Plugins/PCG/Source/PCG/Private/Helpers/PCGActorHelpers.cpp b/Engine/Plugins/PCG/Source/PCG/Private/Helpers/PCGActorHelpers.cpp
--- a/Engine/Plugins/PCG/Source/PCG/Private/Helpers/PCGActorHelpers.cpp
+++ b/Engine/Plugins/PCG/Source/PCG/Private/Helpers/PCGActorHelpers.cpp
@@ -95,8 +95,8 @@ UPCGManagedISMComponent* UPCGActorHelpers::GetOrCreateManagedISMC(AActor* InTarg
 
 	UObject* ExecutionSourceObject = CastChecked<UObject>(InExecutionSource);
 	
-	FPCGManagedResourceContainer* ResourceContainer = InExecutionSource->GetExecutionState().GetManagedResourceContainer();
-	if (!ResourceContainer)
+	FPCGManagedResourceContainerHelper ContainerHelper(InExecutionSource);
+	if (!ContainerHelper.IsValid())
 	{
 		return nullptr;
 	}
@@ -161,7 +161,7 @@ UPCGManagedISMComponent* UPCGActorHelpers::GetOrCreateManagedISMC(AActor* InTarg
 			// If we explicitely are trying to find a RF_Transient through the params or if the Execution source itself is in preview mode
 			const bool bTransient = InParams.bTransient || (InExecutionSource->GetExecutionState().IsInPreviewMode());
 
-			ResourceContainer->ForEachManagedResource([&MatchingResource, &InParams, &InTargetActor, &Descriptor, bTransient](UPCGManagedResource* InResource)
+			ContainerHelper.ForEachManagedResource([&MatchingResource, &InParams, &InTargetActor, &Descriptor, bTransient](UPCGManagedResource* InResource)
 			{
 				// Early out if already found a match
 				if (MatchingResource)
@@ -267,7 +267,7 @@ UPCGManagedISMComponent* UPCGActorHelpers::GetOrCreateManagedISMC(AActor* InTarg
 	}
 	
 	Resource->SetSettingsCrc(InParams.SettingsCrc);
-	ResourceContainer->AddManagedResource(Resource);
+	ContainerHelper.AddManagedResource(Resource);
 
 	return Resource;
 }
@@ -280,8 +280,8 @@ UPCGManagedISKMComponent* UPCGActorHelpers::GetOrCreateManagedABMC(AActor* InTar
 
 	UObject* ExecutionSourceObject = CastChecked<UObject>(InExecutionSource);
 
-	FPCGManagedResourceContainer* ResourceContainer = InExecutionSource->GetExecutionState().GetManagedResourceContainer();
-	if (!ResourceContainer)
+	FPCGManagedResourceContainerHelper ContainerHelper(InExecutionSource);
+	if (!ContainerHelper.IsValid())
 	{
 		return nullptr;
 	}
@@ -315,7 +315,7 @@ UPCGManagedISKMComponent* UPCGActorHelpers::GetOrCreateManagedABMC(AActor* InTar
 		UPCGManagedISKMComponent* MatchingResource = nullptr;
 		if(InParams.SettingsCrc.IsValid())
 		{
-			ResourceContainer->ForEachManagedResource([&MatchingResource, &InParams, &InTargetActor, &Descriptor](UPCGManagedResource* InResource)
+			ContainerHelper.ForEachManagedResource([&MatchingResource, &InParams, &InTargetActor, &Descriptor](UPCGManagedResource* InResource)
 			{
 				// Early out if already found a match
 				if (MatchingResource)
@@ -400,7 +400,7 @@ UPCGManagedISKMComponent* UPCGActorHelpers::GetOrCreateManagedABMC(AActor* InTar
 	
 	Resource->SetSettingsCrc(InParams.SettingsCrc);
 
-	ResourceContainer->AddManagedResource(Resource);
+	ContainerHelper.AddManagedResource(Resource);
 
 	return Resource;
 }
@@ -435,8 +435,8 @@ UPCGManagedSplineMeshComponent* UPCGActorHelpers::GetOrCreateManagedSplineMeshCo
 {
 	check(InTargetActor && InExecutionSource);
 
-	FPCGManagedResourceContainer* ResourceContainer = InExecutionSource->GetExecutionState().GetManagedResourceContainer();
-	if (!ResourceContainer)
+	FPCGManagedResourceContainerHelper ContainerHelper(InExecutionSource);
+	if (!ContainerHelper.IsValid())
 	{
 		return nullptr;
 	}
@@ -478,7 +478,7 @@ UPCGManagedSplineMeshComponent* UPCGActorHelpers::GetOrCreateManagedSplineMeshCo
 		UPCGManagedSplineMeshComponent* MatchingResource = nullptr;
 		if(InParams.SettingsCrc.IsValid())
 		{
-			ResourceContainer->ForEachManagedResource([&MatchingResource, &InParams, &InTargetActor](UPCGManagedResource* InResource)
+			ContainerHelper.ForEachManagedResource([&MatchingResource, &InParams, &InTargetActor](UPCGManagedResource* InResource)
 			{
 				// Early out if already found a match
 				if (MatchingResource)
@@ -574,7 +574,7 @@ UPCGManagedSplineMeshComponent* UPCGActorHelpers::GetOrCreateManagedSplineMeshCo
 	Resource->SetSplineMeshParams(InParams.SplineMeshParams);
 	Resource->SetSettingsCrc(InParams.SettingsCrc);
 	
-	ResourceContainer->AddManagedResource(Resource);
+	ContainerHelper.AddManagedResource(Resource);
 
 	return Resource;
 }
diff --git a/Engine/Plugins/PCG/Source/PCG/Private/PCGComponent.cpp b/Engine/Plugins/PCG/Source/PCG/Private/PCGComponent.cpp
--- a/Engine/Plugins/PCG/Source/PCG/Private/PCGComponent.cpp
+++ b/Engine/Plugins/PCG/Source/PCG/Private/PCGComponent.cpp
@@ -101,7 +101,6 @@ namespace PCGComponent
 UPCGComponent::UPCGComponent(const FObjectInitializer& InObjectInitializer)
 	: Super(InObjectInitializer)
 {
-	ManagedResourceContainer.SetOwner(this);
 	ExecutionState.Component = this;
 	GraphInstance = InObjectInitializer.CreateDefaultSubobject<UPCGGraphInstance>(this, TEXT("PCGGraphInstance"));
 	SchedulingPolicyClass = UPCGSchedulingPolicyDistanceAndDirection::StaticClass();
@@ -239,7 +238,8 @@ void UPCGComponent::SetGraphInterfaceLocal(UPCGGraphInterface* InGraphInterface)
 
 void UPCGComponent::AddToManagedResources(UPCGManagedResource* InResource)
 {
-	ManagedResourceContainer.AddManagedResource(InResource);
+	FPCGManagedResourceContainerHelper ContainerHelper(this);
+	ContainerHelper.AddManagedResource(InResource);
 }
 
 void UPCGComponent::OnManagedResourceAdded(UPCGManagedResource* InResource)
@@ -269,11 +269,12 @@ void UPCGComponent::AddComponentsToManagedResources(const TArray<UActorComponent
 		}
 	}
 
+	FPCGManagedResourceContainerHelper ContainerHelper(this);
 	UE::TScopeLock ResourcesLock(GeneratedResourcesLock);
-	ensure(ManagedResourceContainer.AreResourcesAccessible());
+	ensure(ContainerHelper.AreResourcesAccessible());
 
 	UPCGManagedComponentDefaultList* DefaultList = nullptr;
-	for (const TObjectPtr<UPCGManagedResource>& ManagedResource : ManagedResourceContainer.GetConstManagedResourcesNoLock())
+	for (const TObjectPtr<UPCGManagedResource>& ManagedResource : ContainerHelper.GetConstManagedResourcesNoLock())
 	{
 		if (UPCGManagedComponentDefaultList* ExistingList = Cast<UPCGManagedComponentDefaultList>(ManagedResource))
 		{
@@ -285,8 +286,7 @@ void UPCGComponent::AddComponentsToManagedResources(const TArray<UActorComponent
 	if (!DefaultList)
 	{
 		DefaultList = NewObject<UPCGManagedComponentDefaultList>(this);
-		FPCGManagerResourceContainerHelper Helper(ManagedResourceContainer);
-		Helper.AddManagedResourceNoLock(DefaultList);
+		ContainerHelper.AddManagedResourceNoLock(DefaultList);
 	}
 
 	check(DefaultList);
@@ -327,19 +327,28 @@ void UPCGComponent::AddActorsToManagedResources(const TArray<AActor*>& InActors)
 	AddToManagedResources(ManagedResource);
 }
 
+bool UPCGComponent::AreManagedResourcesAccessible() const
+{
+	const FPCGManagedResourceConstContainerHelper ContainerHelper(this);
+	return ContainerHelper.AreResourcesAccessible();
+}
+
 void UPCGComponent::ForEachManagedResource(TFunctionRef<void(UPCGManagedResource*)> InFunction)
 {
-	ManagedResourceContainer.ForEachManagedResource(InFunction);
+	FPCGManagedResourceContainerHelper ContainerHelper(this);
+	ContainerHelper.ForEachManagedResource(InFunction);
 }
 
 void UPCGComponent::ForEachConstManagedResource(TFunctionRef<void(const UPCGManagedResource*)> InFunction) const
 {
-	ManagedResourceContainer.ForEachConstManagedResource(InFunction);
+	const FPCGManagedResourceConstContainerHelper ContainerHelper(this);
+	ContainerHelper.ForEachConstManagedResource(InFunction);
 }
 
 bool UPCGComponent::IsAnyObjectManagedByResource(const TArrayView<const UObject*> InObjects) const
 {
-	return ManagedResourceContainer.IsAnyObjectManagedByResource(InObjects);
+	const FPCGManagedResourceConstContainerHelper ContainerHelper(this);
+	return ContainerHelper.IsAnyObjectManagedByResource(InObjects);
 }
 
 bool UPCGComponent::ShouldGenerate(bool bForce, EPCGComponentGenerationTrigger RequestedGenerationTrigger) const
@@ -991,9 +1000,10 @@ void UPCGComponent::PurgeUnlinkedResources(const AActor* InActor)
 		}
 
 		UE::TScopeLock ResourcesLock(Component->GeneratedResourcesLock);
-		if(ensure(Component->ManagedResourceContainer.AreResourcesAccessible()))
+		FPCGManagedResourceContainerHelper ContainerHelper(Component);
+		if(ensure(ContainerHelper.AreResourcesAccessible()))
 		{
-			for (UPCGManagedResource* ManagedResource : Component->ManagedResourceContainer.GetConstManagedResourcesNoLock())
+			for (UPCGManagedResource* ManagedResource : ContainerHelper.GetConstManagedResourcesNoLock())
 			{
 				if (UPCGManagedComponent* ManagedComponent = Cast<UPCGManagedComponent>(ManagedResource))
 				{
@@ -1090,7 +1100,8 @@ FPCGTaskId UPCGComponent::CleanupLocal(bool bRemoveComponents, const TArray<FPCG
 		return InvalidPCGTaskId;
 	}
 
-	const bool bNeedsLocalCleanup = bGenerated || IsGenerating() || !ManagedResourceContainer.IsEmpty();
+	const FPCGManagedResourceConstContainerHelper ContainerHelper(this);
+	const bool bNeedsLocalCleanup = bGenerated || IsGenerating() || !ContainerHelper.IsEmpty();
 	if (!bNeedsLocalCleanup && Subsystem->GetPCGComponentPartitionActorMappings(this).IsEmpty())
 	{
 		return InvalidPCGTaskId;
@@ -1437,13 +1448,13 @@ bool UPCGComponent::MoveResourcesToNewActor(const FPCGMoveResourceParams& Params
 #endif
 	{
 		UE::TScopeLock ResourcesLock(GeneratedResourcesLock);
-		ensure(ManagedResourceContainer.AreResourcesAccessible());
-		FPCGManagerResourceContainerHelper Helper(ManagedResourceContainer);
+		FPCGManagedResourceContainerHelper ContainerHelper(this);
+		ensure(ContainerHelper.AreResourcesAccessible());
 #if WITH_EDITOR
-		Scope.AddResources(this, ManagedResourceContainer.GetConstManagedResourcesNoLock());
+		Scope.AddResources(this, ContainerHelper.GetConstManagedResourcesNoLock());
 #endif
 		TArray<TSoftObjectPtr<AActor>> TentativeMovedGeneratedActors;
-		const TArray<TObjectPtr<UPCGManagedResource>>& ManagedResources = ManagedResourceContainer.GetConstManagedResourcesNoLock();
+		const TArray<TObjectPtr<UPCGManagedResource>>& ManagedResources = ContainerHelper.GetConstManagedResourcesNoLock();
 
 		for(int ResourceIndex = ManagedResources.Num() - 1; ResourceIndex >=0; --ResourceIndex)
 		{
@@ -1471,7 +1482,7 @@ bool UPCGComponent::MoveResourcesToNewActor(const FPCGMoveResourceParams& Params
 					}
 #endif
 
-					Helper.RemoveManagedResourceAtNoLock(ResourceIndex);
+					ContainerHelper.RemoveManagedResourceAtNoLock(ResourceIndex);
 				}
 			}
 			else
@@ -1556,7 +1567,8 @@ bool UPCGComponent::MoveResourcesToNewActor(const FPCGMoveResourceParams& Params
 
 void UPCGComponent::CleanupLocalImmediate(bool bRemoveComponents, bool bCleanupLocalComponents)
 {
-	PCGGeneratedResourcesLogging::LogCleanupLocalImmediate(this, bRemoveComponents, ManagedResourceContainer.GetManagedResourcesCopy());
+	FPCGManagedResourceContainerHelper ContainerHelper(this);
+	PCGGeneratedResourcesLogging::LogCleanupLocalImmediate(this, bRemoveComponents, ContainerHelper.GetManagedResourcesCopy());
 
 	UPCGSubsystem* Subsystem = GetSubsystem();
 
@@ -1573,7 +1585,7 @@ void UPCGComponent::CleanupLocalImmediate(bool bRemoveComponents, bool bCleanupL
 	// Cancels generation of this component if there is an ongoing generation in progress.
 	CancelGeneration();
 
-	ManagedResourceContainer.Cleanup(bRemoveComponents);
+	ContainerHelper.Cleanup(bRemoveComponents);
 
 	PostCleanupGraph(bRemoveComponents);
 
@@ -1583,7 +1595,7 @@ void UPCGComponent::CleanupLocalImmediate(bool bRemoveComponents, bool bCleanupL
 		Subsystem->CleanupLocalComponentsImmediate(this, bRemoveComponents);
 	}
 
-	PCGGeneratedResourcesLogging::LogCleanupLocalImmediateFinished(this, ManagedResourceContainer.GetManagedResourcesCopy());
+	PCGGeneratedResourcesLogging::LogCleanupLocalImmediateFinished(this, ContainerHelper.GetManagedResourcesCopy());
 }
 
 FPCGTaskId UPCGComponent::CreateCleanupTask(bool bRemoveComponents, const TArray<FPCGTaskId>& Dependencies)
@@ -1593,7 +1605,8 @@ FPCGTaskId UPCGComponent::CreateCleanupTask(bool bRemoveComponents, const TArray
 		UE_LOG(LogPCG, Log, TEXT("[%s] --- CLEANUP COMPONENT ---"), GetOwner() ? *GetOwner()->GetName() : TEXT("MissingComponent"));
 	}
 
-	if ((!bGenerated && ManagedResourceContainer.IsEmpty() && !IsGenerating()) || IsCleaningUp())
+	FPCGManagedResourceContainerHelper ContainerHelper(this);
+	if ((!bGenerated && ContainerHelper.IsEmpty() && !IsGenerating()) || IsCleaningUp())
 	{
 		return InvalidPCGTaskId;
 	}
@@ -1612,7 +1625,7 @@ FPCGTaskId UPCGComponent::CreateCleanupTask(bool bRemoveComponents, const TArray
 		AllDependencies = &AdditionalDependencies;
 	}
 
-	TSharedPtr<FPCGManagedResourceContainer::FCleanupTask> ContainerCleanup = ManagedResourceContainer.CreateCleanupTask();
+	TSharedPtr<FPCGManagedResourceContainerHelper::FCleanupTask> ContainerCleanup = ContainerHelper.CreateCleanupTask();
 
 	auto AbortCleanup = [ContainerCleanup]()
 	{
@@ -1630,7 +1643,8 @@ FPCGTaskId UPCGComponent::CreateCleanupTask(bool bRemoveComponents, const TArray
 
 void UPCGComponent::CleanupUnusedManagedResources()
 {
-	ManagedResourceContainer.CleanupUnusedManagedResources();
+	FPCGManagedResourceContainerHelper ContainerHelper(this);
+	ContainerHelper.CleanupUnusedManagedResources();
 }
 
 void UPCGComponent::ClearGraphGeneratedOutput(bool bClearLoadedPreviewData)
@@ -1713,7 +1727,8 @@ void UPCGComponent::EndPlay(const EEndPlayReason::Type EndPlayReason)
 		Subsystem->UnregisterExecutionSource(this);
 	}
 
-	ManagedResourceContainer.EndPlay();
+	FPCGManagedResourceContainerHelper ContainerHelper(this);
+	ContainerHelper.EndPlay();
 
 	Super::EndPlay(EndPlayReason);
 }
@@ -1743,19 +1758,20 @@ void UPCGComponent::OnUnregister()
 	// In this case we do want to release those resources and the actors they might hold on to. If we don't those actors become orphaned and we might end up with duplicates when reloading the component.
 	TSet<TSoftObjectPtr<AActor>> ActorsToDelete;
 	{
+		FPCGManagedResourceContainerHelper ContainerHelper(this);
 		UE::TScopeLock ResourcesLock(GeneratedResourcesLock);
-		if (ensure(ManagedResourceContainer.AreResourcesAccessible()))
+		if (ensure(ContainerHelper.AreResourcesAccessible()))
 		{
 			const bool bCleanupPersistentResources = bUnregisteredThroughLoading && PCGComponent::CVarDeletePreviewResourcesWhenUnloading.GetValueOnAnyThread() && CurrentEditingMode == EPCGEditorDirtyMode::Preview;
 
-			for (UPCGManagedResource* ManagedResource : ManagedResourceContainer.GetConstManagedResourcesNoLock())
+			for (UPCGManagedResource* ManagedResource : ContainerHelper.GetConstManagedResourcesNoLock())
 			{
 				if (ManagedResource)
 				{
 					// Make sure resource is indeed a transient/preview resource
 					const bool bReleasePersistent = bCleanupPersistentResources &&
 						!ManagedResource->IsMarkedTransientOnLoad() &&
-						!ManagedResourceContainer.GetConstLoadedPreviewResourcesNoLock().Contains(ManagedResource);
+						!ContainerHelper.GetConstLoadedPreviewResourcesNoLock().Contains(ManagedResource);
 
 					if (bReleasePersistent || ManagedResource->ReleaseOnTeardown())
 					{
@@ -1843,9 +1859,10 @@ void UPCGComponent::Serialize(FArchive& Ar)
 	TArray<TObjectPtr<UPCGManagedResource>> GeneratedResourcesCopy;
 	FPCGDataCollection GeneratedGraphOutputCopy;
 
+	FPCGManagedResourceContainerHelper ContainerHelper(this);
 	if (!Ar.IsLoading() && CurrentEditingMode == EPCGEditorDirtyMode::Preview)
 	{
-		ManagedResourceContainer.BeginPreviewSave(GeneratedResourcesCopy);
+		ContainerHelper.BeginPreviewSave(GeneratedResourcesCopy);
 
 		GeneratedGraphOutputCopy = GeneratedGraphOutput;
 		GeneratedGraphOutput = LoadedPreviewGeneratedGraphOutput;
@@ -1887,13 +1904,10 @@ void UPCGComponent::Serialize(FArchive& Ar)
 #if WITH_EDITOR
 	if (!Ar.IsLoading() && CurrentEditingMode == EPCGEditorDirtyMode::Preview)
 	{
-		ManagedResourceContainer.EndPreviewSave(GeneratedResourcesCopy);
+		ContainerHelper.EndPreviewSave(GeneratedResourcesCopy);
 		GeneratedGraphOutput = GeneratedGraphOutputCopy;
 	}
 #endif // WITH_EDITOR
-
-	// Fixup Owner after serialization
-	ManagedResourceContainer.SetOwner(this);
 }
 
 void UPCGComponent::PostLoad()
@@ -1903,16 +1917,16 @@ void UPCGComponent::PostLoad()
 
 #if WITH_EDITOR
 
+	FPCGManagedResourceContainerHelper ContainerHelper(this);
 PRAGMA_DISABLE_DEPRECATION_WARNINGS
 	if (!GeneratedResources_DEPRECATED.IsEmpty())
 	{
-		ManagedResourceContainer.SetManagedResources(MoveTemp(GeneratedResources_DEPRECATED));
+		ContainerHelper.SetManagedResources(MoveTemp(GeneratedResources_DEPRECATED));
 	}
 PRAGMA_ENABLE_DEPRECATION_WARNINGS
 
 	// Debug resources are transient and will generate null entries on save, clean them up here
-	FPCGManagerResourceContainerHelper Helper(ManagedResourceContainer);
-	Helper.RemoveManagedResourceNoLock(nullptr);
+	ContainerHelper.RemoveManagedResourceNoLock(nullptr);
 
 	// Force dirty to be false on load. We should never refresh on load.
 	bDirtyGenerated = false;
@@ -1944,7 +1958,7 @@ PRAGMA_ENABLE_DEPRECATION_WARNINGS
 	}
 
 	/** Deprecation code, should be removed once generated data has been updated */
-	if (GetOwner() && bGenerated && ManagedResourceContainer.IsEmpty())
+	if (GetOwner() && bGenerated && ContainerHelper.IsEmpty())
 	{
 		TArray<UInstancedStaticMeshComponent*> ISMCs;
 		GetOwner()->GetComponents(ISMCs);
@@ -1955,7 +1969,7 @@ PRAGMA_ENABLE_DEPRECATION_WARNINGS
 			{
 				UPCGManagedISMComponent* ManagedComponent = NewObject<UPCGManagedISMComponent>(this);
 				ManagedComponent->GeneratedComponent = ISMC;
-				Helper.AddManagedResourceNoLock(ManagedComponent);
+				ContainerHelper.AddManagedResourceNoLock(ManagedComponent);
 			}
 		}
 
@@ -1963,7 +1977,7 @@ PRAGMA_ENABLE_DEPRECATION_WARNINGS
 		{
 			UPCGManagedActors* ManagedActors = NewObject<UPCGManagedActors>(this);
 			ManagedActors->GetMutableGeneratedActors() = GeneratedActors_DEPRECATED.Array();
-			Helper.AddManagedResourceNoLock(ManagedActors);
+			ContainerHelper.AddManagedResourceNoLock(ManagedActors);
 			GeneratedActors_DEPRECATED.Reset();
 		}
 	}
@@ -2405,9 +2419,6 @@ void UPCGComponent::PostEditImport()
 	LLM_SCOPE_BYTAG(PCG);
 	Super::PostEditImport();
 
-	// Update owner post import which doesn't use object serialization
-	ManagedResourceContainer.SetOwner(this);
-
 	SetupCallbacksOnCreation();
 }
 
@@ -2870,8 +2881,10 @@ void UPCGComponent::OnRefresh(bool bForceRefresh)
 	const bool bWasGenerated = bGenerated;
 	const bool bWasGeneratedOrGenerating = bWasGenerated || bForceRefresh || IsGenerating();
 
+	const FPCGManagedResourceConstContainerHelper ContainerHelper(this);
+
 	// If we are partitioned but we have resources, we need to force a cleanup
-	if (IsPartitioned() && !ManagedResourceContainer.IsEmpty())
+	if (IsPartitioned() && !ContainerHelper.IsEmpty())
 	{
 		CleanupLocalImmediate(/*bRemoveComponents=*/true);
 	}
@@ -3694,12 +3707,14 @@ void UPCGComponent::ResetIgnoredChangeOrigins(bool bLogIfAnyPresent)
 
 void UPCGComponent::SetManagedResources(const TArray<TObjectPtr<UPCGManagedResource>>& Resources)
 {
-	ManagedResourceContainer.SetManagedResources(Resources);
+	FPCGManagedResourceContainerHelper ContainerHelper(this);
+	ContainerHelper.SetManagedResources(Resources);
 }
 
 void UPCGComponent::GetManagedResources(TArray<TObjectPtr<UPCGManagedResource>>& Resources) const
 {
-	Resources =	ManagedResourceContainer.GetManagedResourcesCopy();
+	const FPCGManagedResourceConstContainerHelper ContainerHelper(this);
+	Resources =	ContainerHelper.GetManagedResourcesCopy();
 }
 
 void UPCGComponent::SetEditingMode(EPCGEditorDirtyMode InEditingMode, EPCGEditorDirtyMode InSerializedEditingMode)
@@ -3711,12 +3726,14 @@ void UPCGComponent::SetEditingMode(EPCGEditorDirtyMode InEditingMode, EPCGEditor
 #if WITH_EDITOR
 bool UPCGComponent::DeletePreviewResources()
 {
-	return ManagedResourceContainer.DeletePreviewResources();
+	FPCGManagedResourceContainerHelper ContainerHelper(this);
+	return ContainerHelper.DeletePreviewResources();
 }
 
 void UPCGComponent::MarkResourcesAsTransientOnLoad()
 {
-	ManagedResourceContainer.MarkResourcesAsTransientOnLoad();
+	FPCGManagedResourceContainerHelper ContainerHelper(this);
+	ContainerHelper.MarkResourcesAsTransientOnLoad();
 
 	// Force duplicate owner for PIE so that transient changes are reflected in game
 	if (AActor* Owner = GetOwner())
@@ -3731,7 +3748,8 @@ void UPCGComponent::ChangeTransientState(EPCGEditorDirtyMode NewEditingMode)
 	// Instead, UPCGComponent will react to CurrentEditingMode changes between PreEditUndo/PostEditUndo and call ChangeTransientState accordingly to reflect the change to CurrentEditingMode.
 	TGuardValue<ITransaction*> SuppressTransaction(GUndo, nullptr);
 
-	bool bShouldMarkDirty = ManagedResourceContainer.ChangeTransientState(NewEditingMode);
+	FPCGManagedResourceContainerHelper ContainerHelper(this);
+	bool bShouldMarkDirty = ContainerHelper.ChangeTransientState(NewEditingMode);
 
 	if (NewEditingMode != EPCGEditorDirtyMode::Preview)
 	{
@@ -3808,10 +3826,11 @@ TArray<TSoftObjectPtr<AActor>> UPCGComponent::GetManagedActorPaths(AActor* InAct
 	TSet<TSoftObjectPtr<AActor>> ManagedActorPaths;
 	InActor->ForEachComponent<UPCGComponent>(/*bIncludeFromChildActors=*/true, [&ManagedActorPaths](UPCGComponent* Component)
 	{
+		FPCGManagedResourceContainerHelper ContainerHelper(Component);
 		PCG::TScopeLock ResourceLock(Component->GeneratedResourcesLock);
-		if (ensure(Component->ManagedResourceContainer.AreResourcesAccessible()))
+		if (ensure(ContainerHelper.AreResourcesAccessible()))
 		{
-			for (const UPCGManagedResource* ManagedResource : Component->ManagedResourceContainer.GetConstManagedResourcesNoLock())
+			for (const UPCGManagedResource* ManagedResource : ContainerHelper.GetConstManagedResourcesNoLock())
 			{
 				if (const UPCGManagedActors* ManagedActors = Cast<UPCGManagedActors>(ManagedResource))
 				{
@@ -3894,7 +3913,6 @@ void FPCGComponentInstanceData::ApplyToComponent(UActorComponent* Component, con
 		PCGComponent->PreviousEditingMode = PreviousEditingMode;	
 		PCGComponent->CurrentEditingMode = CurrentEditingMode;
 		PCGComponent->bGeneratedOffline = bGeneratedOffline;
-		PCGComponent->ManagedResourceContainer.SetOwner(PCGComponent);
 
 		auto CopyToTrackingKeys = [](const TArray<FPCGTrackedKeyInstanceData>& InTrackingKeys, FPCGSelectionKeyToSettingsMap& OutKeysToSettingsMap)
 		{
@@ -3912,6 +3930,8 @@ void FPCGComponentInstanceData::ApplyToComponent(UActorComponent* Component, con
 
 		if (ConstructionSourceComponent)
 		{
+			const FPCGManagedResourceConstContainerHelper ConstructionContainerHelper(ConstructionSourceComponent);
+
 			// Critical: LastGeneratedBounds
 			PCGComponent->LastGeneratedBounds = ConstructionSourceComponent->LastGeneratedBounds;
 
@@ -3920,7 +3940,7 @@ void FPCGComponentInstanceData::ApplyToComponent(UActorComponent* Component, con
 			TMap<TObjectPtr<UPCGManagedResource>, TObjectPtr<UPCGManagedResource>> GeneratedResourceMapping;
 #endif
 			TArray<TObjectPtr<UPCGManagedResource>> DuplicatedResources;
-			for (const TObjectPtr<UPCGManagedResource>& Resource : ConstructionSourceComponent->ManagedResourceContainer.GetConstManagedResourcesNoLock())
+			for (const TObjectPtr<UPCGManagedResource>& Resource : ConstructionContainerHelper.GetConstManagedResourcesNoLock())
 			{
 				if (Resource)
 				{
@@ -3936,7 +3956,7 @@ void FPCGComponentInstanceData::ApplyToComponent(UActorComponent* Component, con
 
 			if (DuplicatedResources.Num() > 0)
 			{
-				PCGComponent->SetManagedResources(DuplicatedResources);
+				PCGComponent->SetManagedResources(MoveTemp(DuplicatedResources));
 			}
 
 #if WITH_EDITOR
@@ -3944,7 +3964,7 @@ void FPCGComponentInstanceData::ApplyToComponent(UActorComponent* Component, con
 			PCGComponent->bDirtyGenerated = ConstructionSourceComponent->bDirtyGenerated;
 
 			TArray<TObjectPtr<UPCGManagedResource>> DuplicateLoadedPreviewResources;
-			for (const TObjectPtr<UPCGManagedResource>& Resource : ConstructionSourceComponent->ManagedResourceContainer.GetConstLoadedPreviewResourcesNoLock())
+			for (const TObjectPtr<UPCGManagedResource>& Resource : ConstructionContainerHelper.GetConstLoadedPreviewResourcesNoLock())
 			{
 				if (Resource)
 				{
@@ -3964,7 +3984,8 @@ void FPCGComponentInstanceData::ApplyToComponent(UActorComponent* Component, con
 
 			if (DuplicateLoadedPreviewResources.Num() > 0)
 			{
-				PCGComponent->ManagedResourceContainer.SetLoadedPreviewResourcesNoLock(DuplicateLoadedPreviewResources);
+				FPCGManagedResourceContainerHelper ContainerHelper(PCGComponent);
+				ContainerHelper.SetLoadedPreviewResourcesNoLock(MoveTemp(DuplicateLoadedPreviewResources));
 			}
 
 			PCGComponent->LoadedPreviewGeneratedGraphOutput = ConstructionSourceComponent->LoadedPreviewGeneratedGraphOutput;
diff --git a/Engine/Plugins/PCG/Source/PCG/Private/PCGGraphExecutionStateInterface.cpp b/Engine/Plugins/PCG/Source/PCG/Private/PCGGraphExecutionStateInterface.cpp
--- a/Engine/Plugins/PCG/Source/PCG/Private/PCGGraphExecutionStateInterface.cpp
+++ b/Engine/Plugins/PCG/Source/PCG/Private/PCGGraphExecutionStateInterface.cpp
@@ -12,14 +12,8 @@
 
 void IPCGGraphExecutionState::AddToManagedResources(UPCGManagedResource* InResource)
 {
-	if (FPCGManagedResourceContainer* Container = GetManagedResourceContainer())
-	{
-		Container->AddManagedResource(InResource);
-	}
-	else
-	{
-		UE_LOG(LogPCG, Error, TEXT("This execution source (%s) does not support managed resources."), *GetDebugName());
-	}
+	ensure(false);
+	UE_LOG(LogPCG, Error, TEXT("Use FPCGManagedResourceContainerHelper with GetManagedResourceContainer instead."));
 }
 
 IPCGBaseSubsystem* IPCGGraphExecutionState::GetSubsystem() const
diff --git a/Engine/Plugins/PCG/Source/PCG/Private/PCGManagedResourceContainer.cpp b/Engine/Plugins/PCG/Source/PCG/Private/PCGManagedResourceContainer.cpp
--- a/Engine/Plugins/PCG/Source/PCG/Private/PCGManagedResourceContainer.cpp
+++ b/Engine/Plugins/PCG/Source/PCG/Private/PCGManagedResourceContainer.cpp
@@ -15,87 +15,91 @@
 #include "EditorActorFolders.h"
 #endif
 
-void FPCGManagedResourceContainer::SetOwner(IPCGGraphExecutionSource* InOwner)
+void FPCGManagedResourceContainerHelper::AddManagedResource(UPCGManagedResource* InManagedResource)
 {
-	Owner = InOwner;
-}
+	check(IsValid());
 
-void FPCGManagedResourceContainer::AddManagedResource(UPCGManagedResource* InManagedResource)
-{
 	AddManagedResourceInternal(InManagedResource, Owner->GetExecutionState().GetManagedResourceContainerLock());
 }
 
-void FPCGManagedResourceContainer::AddManagedResourceNoLock(UPCGManagedResource* InManagedResource)
+void FPCGManagedResourceContainerHelper::AddManagedResourceNoLock(UPCGManagedResource* InManagedResource)
 {
 	AddManagedResourceInternal(InManagedResource, nullptr);
 }
 
-void FPCGManagedResourceContainer::AddManagedResourceInternal(UPCGManagedResource* InManagedResource, FTransactionallySafeCriticalSection* InLock)
+void FPCGManagedResourceContainerHelper::AddManagedResourceInternal(UPCGManagedResource* InManagedResource, FTransactionallySafeCriticalSection* InLock)
 {
+	check(IsValid());
+
 	PCGGeneratedResourcesLogging::LogAddToManagedResources(Owner, InManagedResource);
 	if (InManagedResource)
 	{
 		PCG::TScopeLock ResourcesLock(InLock);
-		ensure(!bAreResourcesInaccessible);
-		GeneratedResources.Add(InManagedResource);
+		ensure(!Container->bAreResourcesInaccessible);
+		Container->GeneratedResources.Add(InManagedResource);
 		Owner->GetExecutionState().OnManagedResourceAdded(InManagedResource);
 	}
 }
 
-void FPCGManagedResourceContainer::RemoveManagedResourceAtNoLock(int32 ResourceIndex)
+void FPCGManagedResourceContainerHelper::RemoveManagedResourceAtNoLock(int32 ResourceIndex)
 {
-	GeneratedResources.RemoveAtSwap(ResourceIndex);
+	check(IsValid());
+	Container->GeneratedResources.RemoveAtSwap(ResourceIndex);
 }
-void FPCGManagedResourceContainer::RemoveManagedResourceNoLock(UPCGManagedResource* InManagedResource)
+void FPCGManagedResourceContainerHelper::RemoveManagedResourceNoLock(UPCGManagedResource* InManagedResource)
 {
-	GeneratedResources.Remove(InManagedResource);
+	check(IsValid());
+	Container->GeneratedResources.Remove(InManagedResource);
 }
 
-void FPCGManagedResourceContainer::SetManagedResources(const TArray<TObjectPtr<UPCGManagedResource>>& InManagedResources)
+void FPCGManagedResourceContainerHelper::SetManagedResources(TArray<TObjectPtr<UPCGManagedResource>> InManagedResources)
 {
+	check(IsValid());
 	PCG::TScopeLock ResourcesLock(Owner->GetExecutionState().GetManagedResourceContainerLock());
 
 	// We expect the GeneratedResources to be empty here, as otherwise they might not be taken care of properly - they
 	// will be lost down below, but this should not happen. However, if the GeneratedResources are marked as Visible,
 	// then they will be copied over during BP duplication, hence why this will happen, hence the ensure here.
-	ensure(GeneratedResources.IsEmpty());
+	ensure(Container->GeneratedResources.IsEmpty());
 
-	GeneratedResources = InManagedResources;
+	Container->GeneratedResources = MoveTemp(InManagedResources);
 
 	// Remove any null entries
-	for (int32 ResourceIndex = GeneratedResources.Num() - 1; ResourceIndex >= 0; --ResourceIndex)
+	for (int32 ResourceIndex = Container->GeneratedResources.Num() - 1; ResourceIndex >= 0; --ResourceIndex)
 	{
-		if (!GeneratedResources[ResourceIndex])
+		if (!Container->GeneratedResources[ResourceIndex])
 		{
-			GeneratedResources.RemoveAtSwap(ResourceIndex);
+			Container->GeneratedResources.RemoveAtSwap(ResourceIndex);
 		}
 	}
 }
 
-TArray<TObjectPtr<UPCGManagedResource>> FPCGManagedResourceContainer::GetManagedResourcesCopy() const
+TArray<TObjectPtr<UPCGManagedResource>> FPCGManagedResourceConstContainerHelper::GetManagedResourcesCopy() const
 {
+	check(IsValid());
 	TArray<TObjectPtr<UPCGManagedResource>> ManagedResourcesCopy;
 	
 	{
 		PCG::TScopeLock ResourcesLock(Owner->GetExecutionState().GetManagedResourceContainerLock());
-		ManagedResourcesCopy = GeneratedResources;
+		ManagedResourcesCopy = Container->GeneratedResources;
 	}
 
 	return ManagedResourcesCopy;
 }
 
-void FPCGManagedResourceContainer::CleanupUnusedManagedResources()
+void FPCGManagedResourceContainerHelper::CleanupUnusedManagedResources()
 {
-	PCGGeneratedResourcesLogging::LogCleanupUnusedManagedResources(Owner, GeneratedResources);
+	check(IsValid());
+	PCGGeneratedResourcesLogging::LogCleanupUnusedManagedResources(Owner, Container->GeneratedResources);
 
 	TSet<TSoftObjectPtr<AActor>> ActorsToDelete;
 
 	{
 		PCG::TScopeLock ResourcesLock(Owner->GetExecutionState().GetManagedResourceContainerLock());
-		ensure(!bAreResourcesInaccessible);
-		for (int32 ResourceIndex = GeneratedResources.Num() - 1; ResourceIndex >= 0; --ResourceIndex)
+		ensure(!Container->bAreResourcesInaccessible);
+		for (int32 ResourceIndex = Container->GeneratedResources.Num() - 1; ResourceIndex >= 0; --ResourceIndex)
 		{
-			UPCGManagedResource* Resource = GetValid(GeneratedResources[ResourceIndex]);
+			UPCGManagedResource* Resource = GetValid(Container->GeneratedResources[ResourceIndex]);
 
 			PCGGeneratedResourcesLogging::LogCleanupUnusedManagedResourcesResource(Owner, Resource);
 
@@ -112,7 +116,7 @@ void FPCGManagedResourceContainer::CleanupUnusedManagedResources()
 					Resource->Rename(nullptr, GetTransientPackage(), REN_DoNotDirty | REN_DontCreateRedirectors | REN_NonTransactional);
 				}
 #endif
-				GeneratedResources.RemoveAtSwap(ResourceIndex);
+				Container->GeneratedResources.RemoveAtSwap(ResourceIndex);
 			}
 		}
 	}
@@ -122,11 +126,13 @@ void FPCGManagedResourceContainer::CleanupUnusedManagedResources()
 		UPCGActorHelpers::DeleteActors(Owner->GetExecutionState().GetWorld(), ActorsToDelete.Array());
 	}
 
-	PCGGeneratedResourcesLogging::LogCleanupUnusedManagedResourcesFinished(Owner, GeneratedResources);
+	PCGGeneratedResourcesLogging::LogCleanupUnusedManagedResourcesFinished(Owner, Container->GeneratedResources);
 }
 
-void FPCGManagedResourceContainer::Cleanup(bool bReleaseManagedResources)
+void FPCGManagedResourceContainerHelper::Cleanup(bool bReleaseManagedResources)
 {
+	check(IsValid());
+
 	TSet<TSoftObjectPtr<AActor>> ActorsToDelete;
 	FPCGManagedActorLoadingScope Scope;
 
@@ -136,12 +142,12 @@ void FPCGManagedResourceContainer::Cleanup(bool bReleaseManagedResources)
 	}
 
 	PCG::TScopeLock ResourcesLock(Owner->GetExecutionState().GetManagedResourceContainerLock());
-	ensure(!bAreResourcesInaccessible);
-	Scope.AddResources(Owner, GeneratedResources);
-	for (int32 ResourceIndex = GeneratedResources.Num() - 1; ResourceIndex >= 0; --ResourceIndex)
+	ensure(!Container->bAreResourcesInaccessible);
+	Scope.AddResources(Owner, Container->GeneratedResources);
+	for (int32 ResourceIndex = Container->GeneratedResources.Num() - 1; ResourceIndex >= 0; --ResourceIndex)
 	{
 		// Note: resources can be null here in some loading + bp object cases
-		UPCGManagedResource* Resource = GeneratedResources[ResourceIndex];
+		UPCGManagedResource* Resource = Container->GeneratedResources[ResourceIndex];
 
 		PCGGeneratedResourcesLogging::LogCleanupLocalImmediateResource(Owner, Resource);
 
@@ -154,18 +160,20 @@ void FPCGManagedResourceContainer::Cleanup(bool bReleaseManagedResources)
 			}
 #endif
 
-			GeneratedResources.RemoveAt(ResourceIndex);
+			Container->GeneratedResources.RemoveAt(ResourceIndex);
 		}
 	}
 
 	UPCGActorHelpers::DeleteActors(Owner->GetExecutionState().GetWorld(), ActorsToDelete.Array());
 }
 
-void FPCGManagedResourceContainer::ForEachManagedResource(TFunctionRef<void(UPCGManagedResource*)> InFunction)
+void FPCGManagedResourceContainerHelper::ForEachManagedResource(TFunctionRef<void(UPCGManagedResource*)> InFunction)
 {
+	check(IsValid());
+
 	PCG::TScopeLock ResourcesLock(Owner->GetExecutionState().GetManagedResourceContainerLock());
-	ensure(!bAreResourcesInaccessible);
-	for (TObjectPtr<UPCGManagedResource>& ManagedResource : GeneratedResources)
+	ensure(!Container->bAreResourcesInaccessible);
+	for (TObjectPtr<UPCGManagedResource>& ManagedResource : Container->GeneratedResources)
 	{
 		if (ManagedResource)
 		{
@@ -174,11 +182,13 @@ void FPCGManagedResourceContainer::ForEachManagedResource(TFunctionRef<void(UPCG
 	}
 }
 
-void FPCGManagedResourceContainer::ForEachConstManagedResource(TFunctionRef<void(const UPCGManagedResource*)> InFunction) const
+void FPCGManagedResourceConstContainerHelper::ForEachConstManagedResource(TFunctionRef<void(const UPCGManagedResource*)> InFunction) const
 {
+	check(IsValid());
+
 	PCG::TScopeLock ResourcesLock(Owner->GetExecutionState().GetManagedResourceContainerLock());
-	ensure(!bAreResourcesInaccessible);
-	for (const TObjectPtr<UPCGManagedResource>& ManagedResource : GeneratedResources)
+	ensure(!Container->bAreResourcesInaccessible);
+	for (const TObjectPtr<UPCGManagedResource>& ManagedResource : Container->GeneratedResources)
 	{
 		if (ManagedResource)
 		{
@@ -187,11 +197,13 @@ void FPCGManagedResourceContainer::ForEachConstManagedResource(TFunctionRef<void
 	}
 }
 
-bool FPCGManagedResourceContainer::IsAnyObjectManagedByResource(const TArrayView<const UObject*> InObjects) const
+bool FPCGManagedResourceConstContainerHelper::IsAnyObjectManagedByResource(const TArrayView<const UObject*> InObjects) const
 {
-	TRACE_CPUPROFILER_EVENT_SCOPE(FPCGManagedResourceContainer::IsAnyObjectManagedByResource);
+	TRACE_CPUPROFILER_EVENT_SCOPE(FPCGManagedResourceContainerHelper::IsAnyObjectManagedByResource);
+
+	check(IsValid());
 
-	if (!ensure(!bAreResourcesInaccessible))
+	if (!ensure(!Container->bAreResourcesInaccessible))
 	{
 		return false;
 	}
@@ -219,9 +231,11 @@ bool FPCGManagedResourceContainer::IsAnyObjectManagedByResource(const TArrayView
 
 #if WITH_EDITOR
 
-void FPCGManagedResourceContainer::MarkResourcesAsTransientOnLoad()
+void FPCGManagedResourceContainerHelper::MarkResourcesAsTransientOnLoad()
 {
-	for (TObjectPtr<UPCGManagedResource>& GeneratedResource : GeneratedResources)
+	check(IsValid());
+
+	for (TObjectPtr<UPCGManagedResource>& GeneratedResource : Container->GeneratedResources)
 	{
 		if (GeneratedResource)
 		{
@@ -229,17 +243,19 @@ void FPCGManagedResourceContainer::MarkResourcesAsTransientOnLoad()
 		}
 	}
 
-	LoadedPreviewResources = GeneratedResources;
+	Container->LoadedPreviewResources = Container->GeneratedResources;
 }
 
-bool FPCGManagedResourceContainer::ChangeTransientState(EPCGEditorDirtyMode NewEditingMode)
+bool FPCGManagedResourceContainerHelper::ChangeTransientState(EPCGEditorDirtyMode NewEditingMode)
 {
+	check(IsValid());
+
 	bool bShouldMarkDirty = false;
 
 	PCG::TScopeLock ResourcesLock(Owner->GetExecutionState().GetManagedResourceContainerLock());
-	ensure(!bAreResourcesInaccessible);
+	ensure(!Container->bAreResourcesInaccessible);
 
-	for (TObjectPtr<UPCGManagedResource>& GeneratedResource : GeneratedResources)
+	for (TObjectPtr<UPCGManagedResource>& GeneratedResource : Container->GeneratedResources)
 	{
 		if (GeneratedResource)
 		{
@@ -265,15 +281,17 @@ bool FPCGManagedResourceContainer::ChangeTransientState(EPCGEditorDirtyMode NewE
 	return bShouldMarkDirty;
 }
 
-bool FPCGManagedResourceContainer::DeletePreviewResources()
+bool FPCGManagedResourceContainerHelper::DeletePreviewResources()
 {
+	check(IsValid());
+
 	bool bResourceWasReleased = false;
 
 	TSet<TSoftObjectPtr<AActor>> ActorsToDelete;
 	// Make sure to release fully the resources that were loaded
-	for (TObjectPtr<UPCGManagedResource> ResourceToRelease : LoadedPreviewResources)
+	for (TObjectPtr<UPCGManagedResource> ResourceToRelease : Container->LoadedPreviewResources)
 	{
-		if (!GeneratedResources.Contains(ResourceToRelease))
+		if (!Container->GeneratedResources.Contains(ResourceToRelease))
 		{
 			// Changing the transient state will clear the "marked transient on load" flag
 			ResourceToRelease->ChangeTransientState(EPCGEditorDirtyMode::Normal);
@@ -285,7 +303,7 @@ bool FPCGManagedResourceContainer::DeletePreviewResources()
 		ensure(!ResourceToRelease->IsMarkedTransientOnLoad());
 	}
 
-	LoadedPreviewResources.Empty();
+	Container->LoadedPreviewResources.Empty();
 
 	if (!ActorsToDelete.IsEmpty())
 	{
@@ -295,58 +313,66 @@ bool FPCGManagedResourceContainer::DeletePreviewResources()
 	return bResourceWasReleased;
 }
 
-void FPCGManagedResourceContainer::BeginPreviewSave(TArray<TObjectPtr<UPCGManagedResource>>& OutGeneratedCopy)
+void FPCGManagedResourceContainerHelper::BeginPreviewSave(TArray<TObjectPtr<UPCGManagedResource>>& OutGeneratedCopy)
 {
-	OutGeneratedCopy = GeneratedResources;
-	GeneratedResources = LoadedPreviewResources;
+	check(IsValid());
+
+	OutGeneratedCopy = Container->GeneratedResources;
+	Container->GeneratedResources = Container->LoadedPreviewResources;
 }
 
-void FPCGManagedResourceContainer::EndPreviewSave(const TArray<TObjectPtr<UPCGManagedResource>>& InGeneratedCopy)
+void FPCGManagedResourceContainerHelper::EndPreviewSave(const TArray<TObjectPtr<UPCGManagedResource>>& InGeneratedCopy)
 {
-	GeneratedResources = InGeneratedCopy;
+	check(IsValid());
+
+	Container->GeneratedResources = InGeneratedCopy;
 }
 
 #endif // WITH_EDITOR
 
-void FPCGManagedResourceContainer::EndPlay()
+void FPCGManagedResourceContainerHelper::EndPlay()
 {
-	if (ensure(!bAreResourcesInaccessible))
+	check(IsValid());
+
+	if (ensure(!Container->bAreResourcesInaccessible))
 	{
 		PCG::TScopeLock ResourcesLock(Owner->GetExecutionState().GetManagedResourceContainerLock());
 
-		for (int32 ResourceIndex = GeneratedResources.Num() - 1; ResourceIndex >= 0; --ResourceIndex)
+		for (int32 ResourceIndex = Container->GeneratedResources.Num() - 1; ResourceIndex >= 0; --ResourceIndex)
 		{
-			if (!GeneratedResources[ResourceIndex])
+			if (!Container->GeneratedResources[ResourceIndex])
 			{
 				// Remove null entries.
-				GeneratedResources.RemoveAtSwap(ResourceIndex);
+				Container->GeneratedResources.RemoveAtSwap(ResourceIndex);
 			}
-			else if (GeneratedResources[ResourceIndex]->ReleaseOnTeardown())
+			else if (Container->GeneratedResources[ResourceIndex]->ReleaseOnTeardown())
 			{
-				if (GeneratedResources[ResourceIndex])
+				if (Container->GeneratedResources[ResourceIndex])
 				{
 					TSet<TSoftObjectPtr<AActor>> ActorsToDelete;
-					GeneratedResources[ResourceIndex]->Release(/*bHardRelease=*/true, ActorsToDelete);
+					Container->GeneratedResources[ResourceIndex]->Release(/*bHardRelease=*/true, ActorsToDelete);
 
 					// Don't support deleting actors during teardown.
 					ensure(ActorsToDelete.IsEmpty());
 				}
 
-				GeneratedResources.RemoveAtSwap(ResourceIndex);
+				Container->GeneratedResources.RemoveAtSwap(ResourceIndex);
 			}
 		}
 	}
 }
 
-TSharedPtr<FPCGManagedResourceContainer::FCleanupTask> FPCGManagedResourceContainer::CreateCleanupTask()
+TSharedPtr<FPCGManagedResourceContainerHelper::FCleanupTask> FPCGManagedResourceContainerHelper::CreateCleanupTask()
 {
+	check(IsValid());
+
 	TSharedPtr<FCleanupTask> CleanupTask = MakeShared<FCleanupTask>();
 	CleanupTask->WeakExecutionSource = Owner;
 
 	return CleanupTask;
 }
 
-void FPCGManagedResourceContainer::FCleanupTask::Abort()
+void FPCGManagedResourceContainerHelper::FCleanupTask::Abort()
 {
 	// If the component is not valid anymore, just early out
 	if (!WeakExecutionSource.IsValid())
@@ -354,15 +380,15 @@ void FPCGManagedResourceContainer::FCleanupTask::Abort()
 		return;
 	}
 
-	FPCGManagedResourceContainer* Container = WeakExecutionSource->GetExecutionState().GetManagedResourceContainer();
-	if (Container)
+	FPCGManagedResourceContainer* ManagedResourceContainer = WeakExecutionSource->GetExecutionState().GetManagedResourceContainer();
+	if (ManagedResourceContainer)
 	{
 		PCG::TScopeLock ContainerLock(WeakExecutionSource->GetExecutionState().GetManagedResourceContainerLock());
-		Container->bAreResourcesInaccessible = false;
+		ManagedResourceContainer->bAreResourcesInaccessible = false;
 	}
 }
 
-bool FPCGManagedResourceContainer::FCleanupTask::Execute(bool bReleaseManagedResources)
+bool FPCGManagedResourceContainerHelper::FCleanupTask::Execute(bool bReleaseManagedResources)
 {
 	IPCGGraphExecutionSource* ExecutionSource = WeakExecutionSource.Get();
 	if (!ExecutionSource)
@@ -370,8 +396,8 @@ bool FPCGManagedResourceContainer::FCleanupTask::Execute(bool bReleaseManagedRes
 		return true;
 	}
 
-	FPCGManagedResourceContainer* Container = ExecutionSource->GetExecutionState().GetManagedResourceContainer();
-	if (!Container)
+	FPCGManagedResourceContainer* ManagedResourceContainer = ExecutionSource->GetExecutionState().GetManagedResourceContainer();
+	if (!ManagedResourceContainer)
 	{
 		return true;
 	}
@@ -381,21 +407,21 @@ bool FPCGManagedResourceContainer::FCleanupTask::Execute(bool bReleaseManagedRes
 	// Safeguard to track illegal modifications of the generated resources array while doing cleanup
 	if (bIsFirstIteration)
 	{
-		ensure(!Container->bAreResourcesInaccessible);
-		Container->bAreResourcesInaccessible = true;
-		ResourceIndex = Container->GeneratedResources.Num() - 1;
+		ensure(!ManagedResourceContainer->bAreResourcesInaccessible);
+		ManagedResourceContainer->bAreResourcesInaccessible = true;
+		ResourceIndex = ManagedResourceContainer->GeneratedResources.Num() - 1;
 		bIsFirstIteration = false;
-		AddResources(WeakExecutionSource.Get(), Container->GeneratedResources);
+		AddResources(WeakExecutionSource.Get(), ManagedResourceContainer->GeneratedResources);
 	}
 
 	// Going backward
 	if (ResourceIndex >= 0)
 	{
-		UPCGManagedResource* Resource = Container->GeneratedResources[ResourceIndex];
+		UPCGManagedResource* Resource = ManagedResourceContainer->GeneratedResources[ResourceIndex];
 
 		if (!Resource)
 		{
-			UE_LOG(LogPCG, Error, TEXT("[FPCGManagedResourceContainer::FCleanupTask::Execute] Null generated resource encountered on execution source \"%s\"."), *ExecutionSource->GetExecutionState().GetDebugName());
+			UE_LOG(LogPCG, Error, TEXT("[FPCGManagedResourceContainerHelper::FCleanupTask::Execute] Null generated resource encountered on execution source \"%s\"."), *ExecutionSource->GetExecutionState().GetDebugName());
 		}
 
 		PCGGeneratedResourcesLogging::LogCreateCleanupTaskResource(WeakExecutionSource.Get(), Resource);
@@ -409,7 +435,7 @@ bool FPCGManagedResourceContainer::FCleanupTask::Execute(bool bReleaseManagedRes
 			}
 #endif
 
-			Container->GeneratedResources.RemoveAtSwap(ResourceIndex);
+			ManagedResourceContainer->GeneratedResources.RemoveAtSwap(ResourceIndex);
 		}
 
 		ResourceIndex--;
@@ -419,7 +445,7 @@ bool FPCGManagedResourceContainer::FCleanupTask::Execute(bool bReleaseManagedRes
 	}
 	else
 	{
-		Container->bAreResourcesInaccessible = false;
+		ManagedResourceContainer->bAreResourcesInaccessible = false;
 	}
 
 	TSet<FName> DeletedActorFolders;
@@ -497,7 +523,7 @@ bool FPCGManagedResourceContainer::FCleanupTask::Execute(bool bReleaseManagedRes
 	}
 #endif
 
-	PCGGeneratedResourcesLogging::LogCreateCleanupTaskFinished(ExecutionSource, &Container->GeneratedResources);
+	PCGGeneratedResourcesLogging::LogCreateCleanupTaskFinished(ExecutionSource, &ManagedResourceContainer->GeneratedResources);
 
 	return true;
 }
\ No newline at end of file
diff --git a/Engine/Plugins/PCG/Source/PCG/Private/RuntimeGen/PCGRuntimeGenExecutionSource.cpp b/Engine/Plugins/PCG/Source/PCG/Private/RuntimeGen/PCGRuntimeGenExecutionSource.cpp
--- a/Engine/Plugins/PCG/Source/PCG/Private/RuntimeGen/PCGRuntimeGenExecutionSource.cpp
+++ b/Engine/Plugins/PCG/Source/PCG/Private/RuntimeGen/PCGRuntimeGenExecutionSource.cpp
@@ -28,7 +28,6 @@ UPCGRuntimeGenExecutionSource::UPCGRuntimeGenExecutionSource(const FObjectInitia
 	: Super(InObjectInitializer)
 {
 	ExecutionState.SetSource(this);
-	ManagedResourceContainer.SetOwner(this);
 }
 
 void UPCGRuntimeGenExecutionSource::Register(IPCGGraphExecutionSource* InOwningSource)
@@ -219,7 +218,8 @@ void UPCGRuntimeGenExecutionSource::OnGraphExecutionAborted(bool bQuiet, bool bC
 
 	if (bCleanupUnusedResources)
 	{
-		ManagedResourceContainer.CleanupUnusedManagedResources();
+		FPCGManagedResourceContainerHelper ContainerHelper(this);
+		ContainerHelper.CleanupUnusedManagedResources();
 	}
 
 #if WITH_EDITOR
@@ -298,11 +298,12 @@ FPCGTaskId UPCGRuntimeGenExecutionSource::Cleanup(const IPCGGraphExecutionState:
 		return InvalidPCGTaskId;
 	}
 
-	PCGGeneratedResourcesLogging::LogCleanupLocalImmediate(this, InCleanupParams.bReleaseManagedResources, ManagedResourceContainer.GetManagedResourcesCopy());
+	FPCGManagedResourceContainerHelper ContainerHelper(this);
+	PCGGeneratedResourcesLogging::LogCleanupLocalImmediate(this, InCleanupParams.bReleaseManagedResources, ContainerHelper.GetManagedResourcesCopy());
 
 	Cancel();
 
-	ManagedResourceContainer.Cleanup(InCleanupParams.bReleaseManagedResources);
+	ContainerHelper.Cleanup(InCleanupParams.bReleaseManagedResources);
 
 	if (InCleanupParams.bReleaseManagedResources)
 	{
@@ -322,7 +323,7 @@ FPCGTaskId UPCGRuntimeGenExecutionSource::Cleanup(const IPCGGraphExecutionState:
 		}
 	}
 
-	PCGGeneratedResourcesLogging::LogCleanupLocalImmediateFinished(this, ManagedResourceContainer.GetManagedResourcesCopy());
+	PCGGeneratedResourcesLogging::LogCleanupLocalImmediateFinished(this, ContainerHelper.GetManagedResourcesCopy());
 
 	return InvalidPCGTaskId;
 }
diff --git a/Engine/Plugins/PCG/Source/PCG/Public/PCGComponent.h b/Engine/Plugins/PCG/Source/PCG/Public/PCGComponent.h
--- a/Engine/Plugins/PCG/Source/PCG/Public/PCGComponent.h
+++ b/Engine/Plugins/PCG/Source/PCG/Public/PCGComponent.h
@@ -204,7 +204,7 @@ public:
 	UE_API void AddActorsToManagedResources(const TArray<AActor*>& InActors);
 
 public:
-	bool AreManagedResourcesAccessible() const { return ManagedResourceContainer.AreResourcesAccessible(); }
+	UE_API bool AreManagedResourcesAccessible() const;
 	UE_API void ForEachManagedResource(TFunctionRef<void(UPCGManagedResource*)> InFunction);
 	UE_API void ForEachConstManagedResource(TFunctionRef<void(const UPCGManagedResource*)> InFunction) const;
 
diff --git a/Engine/Plugins/PCG/Source/PCG/Public/PCGManagedResourceContainer.h b/Engine/Plugins/PCG/Source/PCG/Public/PCGManagedResourceContainer.h
--- a/Engine/Plugins/PCG/Source/PCG/Public/PCGManagedResourceContainer.h
+++ b/Engine/Plugins/PCG/Source/PCG/Public/PCGManagedResourceContainer.h
@@ -18,24 +18,69 @@ struct FPCGManagedResourceContainer
 {
 	GENERATED_BODY()
 
-public:
-	PCG_API void SetOwner(IPCGGraphExecutionSource* InOwner);
-	
-	bool AreResourcesAccessible() const { return !bAreResourcesInaccessible; }
-	bool IsEmpty() { return GeneratedResources.IsEmpty(); }
-	
-	PCG_API void AddManagedResource(UPCGManagedResource* InManagedResource);
-	PCG_API void CleanupUnusedManagedResources();
-	PCG_API void SetManagedResources(const TArray<TObjectPtr<UPCGManagedResource>>& InManagedResources);
+private:
+	friend struct FPCGManagedResourceConstContainerHelper;
+	friend struct FPCGManagedResourceContainerHelper;
+
+	bool bAreResourcesInaccessible = false;
+
+	UPROPERTY()
+	TArray<TObjectPtr<UPCGManagedResource>> GeneratedResources;
+
+#if WITH_EDITORONLY_DATA
+	UPROPERTY()
+	TArray<TObjectPtr<UPCGManagedResource>> LoadedPreviewResources;
+#endif
+};
+
+// Const and Non-const API to interface with FPCGManagedResourceContainer
+//
+// The helper is not tracking the execution source becoming invalid so avoid storing a copy.
+// Instead store a weak ptr to the execution source which allows creating a helper on the stack when needed.
+struct FPCGManagedResourceConstContainerHelper
+{
+	FPCGManagedResourceConstContainerHelper() = delete;
+
+	// const_cast here to store a single pointer for the const and non const versions of the helper.
+	explicit FPCGManagedResourceConstContainerHelper(const IPCGGraphExecutionSource* InExecutionSource)
+		: Owner(const_cast<IPCGGraphExecutionSource*>(InExecutionSource))
+	{
+		Container = Owner ? Owner->GetExecutionState().GetManagedResourceContainer() : nullptr;
+	}
+
+	bool IsValid() const { return Owner != nullptr && Container != nullptr; }
+	bool AreResourcesAccessible() const { return IsValid() && !Container->bAreResourcesInaccessible; }
+	bool IsEmpty() const { return !Container || Container->GeneratedResources.IsEmpty(); }
+
 	PCG_API TArray<TObjectPtr<UPCGManagedResource>> GetManagedResourcesCopy() const;
-	
-	const TArray<TObjectPtr<UPCGManagedResource>>& GetConstManagedResourcesNoLock() const { return GeneratedResources; }
-	
-	PCG_API void ForEachManagedResource(TFunctionRef<void(UPCGManagedResource*)> InFunction);
+
+	const TArray<TObjectPtr<UPCGManagedResource>>& GetConstManagedResourcesNoLock() const { check(IsValid()); return Container->GeneratedResources; }
+
 	PCG_API void ForEachConstManagedResource(TFunctionRef<void(const UPCGManagedResource*)> InFunction) const;
 	PCG_API bool IsAnyObjectManagedByResource(const TArrayView<const UObject*> InObjects) const;
 
-	void EndPlay();
+#if WITH_EDITOR
+	const TArray<TObjectPtr<UPCGManagedResource>>& GetConstLoadedPreviewResourcesNoLock() const { check(IsValid()); return Container->LoadedPreviewResources; }
+#endif
+
+protected:
+	IPCGGraphExecutionSource* Owner = nullptr;
+	FPCGManagedResourceContainer* Container = nullptr;
+};
+
+struct FPCGManagedResourceContainerHelper final : public FPCGManagedResourceConstContainerHelper
+{
+	explicit FPCGManagedResourceContainerHelper(IPCGGraphExecutionSource* InExecutionSource) 
+		: FPCGManagedResourceConstContainerHelper(InExecutionSource)
+	{
+	}
+
+	PCG_API void AddManagedResource(UPCGManagedResource* InManagedResource);
+	PCG_API void CleanupUnusedManagedResources();
+	PCG_API void SetManagedResources(TArray<TObjectPtr<UPCGManagedResource>> InManagedResources);
+	
+	PCG_API void ForEachManagedResource(TFunctionRef<void(UPCGManagedResource*)> InFunction);
+	PCG_API void EndPlay();
 
 	struct FCleanupTask : public FPCGManagedActorLoadingScope
 	{
@@ -47,70 +92,27 @@ public:
 		PCG_API void Abort();
 		PCG_API bool Execute(bool bReleaseManagedResources);
 	};
-	
+
 	PCG_API TSharedPtr<FCleanupTask> CreateCleanupTask();
 	PCG_API void Cleanup(bool bReleaseManagedResources);
 
 #if WITH_EDITOR
-	void SetLoadedPreviewResourcesNoLock(const TArray<TObjectPtr<UPCGManagedResource>>& InManagedResources) { LoadedPreviewResources = InManagedResources; }
-	const TArray<TObjectPtr<UPCGManagedResource>>& GetConstLoadedPreviewResourcesNoLock() const { return LoadedPreviewResources; }
+	void SetLoadedPreviewResourcesNoLock(TArray<TObjectPtr<UPCGManagedResource>> InManagedResources) { check(IsValid()); Container->LoadedPreviewResources = MoveTemp(InManagedResources); }
 
-	void MarkResourcesAsTransientOnLoad();
-	bool ChangeTransientState(EPCGEditorDirtyMode NewEditingMode);
-	bool DeletePreviewResources();
+	PCG_API void MarkResourcesAsTransientOnLoad();
+	PCG_API bool ChangeTransientState(EPCGEditorDirtyMode NewEditingMode);
+	PCG_API bool DeletePreviewResources();
 
-	void BeginPreviewSave(TArray<TObjectPtr<UPCGManagedResource>>& OutGeneratedCopy);
-	void EndPreviewSave(const TArray<TObjectPtr<UPCGManagedResource>>& InGeneratedCopy);
+	PCG_API void BeginPreviewSave(TArray<TObjectPtr<UPCGManagedResource>>& OutGeneratedCopy);
+	PCG_API void EndPreviewSave(const TArray<TObjectPtr<UPCGManagedResource>>& InGeneratedCopy);
 #endif	
 
 private:
-	friend struct FPCGManagerResourceContainerHelper;
+	friend class UPCGComponent;
 
 	void AddManagedResourceNoLock(UPCGManagedResource* InManagedResource);
 	void RemoveManagedResourceNoLock(UPCGManagedResource* InManagedResource);
 	void RemoveManagedResourceAtNoLock(int32 ResourceIndex);
 
 	void AddManagedResourceInternal(UPCGManagedResource* InManagedResource, FTransactionallySafeCriticalSection* InLock);
-
-	bool bAreResourcesInaccessible = false;
-
-	UPROPERTY()
-	TArray<TObjectPtr<UPCGManagedResource>> GeneratedResources;
-
-#if WITH_EDITORONLY_DATA
-	UPROPERTY()
-	TArray<TObjectPtr<UPCGManagedResource>> LoadedPreviewResources;
-#endif
-
-	IPCGGraphExecutionSource* Owner = nullptr;
-};
-
-// Gives access to NoLock methods for friend classes only
-struct FPCGManagerResourceContainerHelper
-{
-	FPCGManagerResourceContainerHelper() = delete;
-	explicit FPCGManagerResourceContainerHelper(FPCGManagedResourceContainer& InContainer)
-		: Container(InContainer)
-	{
-	}
-
-private:
-	friend class UPCGComponent;
-
-	void AddManagedResourceNoLock(UPCGManagedResource* InManagedResource)
-	{
-		Container.AddManagedResourceNoLock(InManagedResource);
-	}
-
-	void RemoveManagedResourceNoLock(UPCGManagedResource* InManagedResource)
-	{
-		Container.RemoveManagedResourceNoLock(InManagedResource);
-	}
-
-	void RemoveManagedResourceAtNoLock(int32 ResourceIndex)
-	{
-		Container.RemoveManagedResourceAtNoLock(ResourceIndex);
-	}
-
-	FPCGManagedResourceContainer& Container;
 };
\ No newline at end of file
