Fixed instance parameters in pcg editor mode not accurately representing the selected graph if selected via preset buttons

#Synced-CL 50904874

diff --git a/Engine/Plugins/PCG/Source/PCGEditor/Private/Details/PCGInteractiveToolSettingsDetails.cpp b/Engine/Plugins/PCG/Source/PCGEditor/Private/Details/PCGInteractiveToolSettingsDetails.cpp
--- a/Engine/Plugins/PCG/Source/PCGEditor/Private/Details/PCGInteractiveToolSettingsDetails.cpp
+++ b/Engine/Plugins/PCG/Source/PCGEditor/Private/Details/PCGInteractiveToolSettingsDetails.cpp
@@ -33,12 +33,12 @@ FPCGInteractiveToolSettingsDetails::~FPCGInteractiveToolSettingsDetails()
 
 void FPCGInteractiveToolSettingsDetails::RegisterDelegates()
 {
-	FCoreUObjectDelegates::OnPostObjectPropertyChanged.AddSP(this, &FPCGInteractiveToolSettingsDetails::Rebuild);
+	FCoreUObjectDelegates::OnObjectPropertyChanged.AddSP(this, &FPCGInteractiveToolSettingsDetails::Rebuild);
 }
 
 void FPCGInteractiveToolSettingsDetails::UnregisterDelegates()
 {
-	FCoreUObjectDelegates::OnPostObjectPropertyChanged.RemoveAll(this);
+	FCoreUObjectDelegates::OnObjectPropertyChanged.RemoveAll(this);
 }
 
 void FPCGInteractiveToolSettingsDetails::CustomizeDetails(IDetailLayoutBuilder& DetailBuilder)
@@ -119,11 +119,11 @@ void FPCGInteractiveToolSettingsDetails::CustomizeDetails(const TSharedPtr<IDeta
 	IDetailCustomization::CustomizeDetails(DetailBuilder);
 }
 
-void FPCGInteractiveToolSettingsDetails::Rebuild(UObject* Object, const FPropertyChangedChainEvent& PropertyChangedChainEvent)
+void FPCGInteractiveToolSettingsDetails::Rebuild(UObject* Object, FPropertyChangedEvent& PropertyChangedEvent)
 {
 	if(Object == ToolSettings)
 	{
-		if(PropertyChangedChainEvent.Property && PropertyChangedChainEvent.Property->GetFName() == GET_MEMBER_NAME_CHECKED(UPCGInteractiveToolSettings, ToolGraph))
+		if(PropertyChangedEvent.Property && PropertyChangedEvent.Property->GetFName() == GET_MEMBER_NAME_CHECKED(UPCGInteractiveToolSettings, ToolGraph))
 		{
 			if(DetailLayoutBuilder.IsValid())
 			{
diff --git a/Engine/Plugins/PCG/Source/PCGEditor/Private/EditorMode/Tools/PCGInteractiveToolSettings.cpp b/Engine/Plugins/PCG/Source/PCGEditor/Private/EditorMode/Tools/PCGInteractiveToolSettings.cpp
--- a/Engine/Plugins/PCG/Source/PCGEditor/Private/EditorMode/Tools/PCGInteractiveToolSettings.cpp
+++ b/Engine/Plugins/PCG/Source/PCGEditor/Private/EditorMode/Tools/PCGInteractiveToolSettings.cpp
@@ -277,6 +277,15 @@ void UPCGInteractiveToolSettings::Shutdown()
 	}
 }
 
+void UPCGInteractiveToolSettings::PostEditChangeProperty(FPropertyChangedEvent& PropertyChangedEvent)
+{
+	// We flip the order compared to the Super function so that OnModified can trigger downstream changes before regular calls are made.
+	// This makes it so that when global PostEditChange delegates are called, we have already made the modifications.
+	OnModified.Broadcast(this, PropertyChangedEvent.Property);
+
+	UInteractiveToolPropertySet::PostEditChangeProperty(PropertyChangedEvent);
+}
+
 void UPCGInteractiveToolSettings::RestoreProperties(UInteractiveTool* RestoreToTool, const FString& CacheIdentifier)
 {
 	Super::RestoreProperties(RestoreToTool, CacheIdentifier);
@@ -753,7 +762,7 @@ void UPCGInteractiveToolSettings::PostWorkingDataInitialized(FPCGInteractiveTool
 
 void UPCGInteractiveToolSettings::OnPropertyModified(UObject* Object, FProperty* Property)
 {
-	if (Object != this)
+	if (Object != this || Property == nullptr)
 	{
 		return;
 	}
diff --git a/Engine/Plugins/PCG/Source/PCGEditor/Private/EditorMode/Tools/PCGInteractiveToolSettings.h b/Engine/Plugins/PCG/Source/PCGEditor/Private/EditorMode/Tools/PCGInteractiveToolSettings.h
--- a/Engine/Plugins/PCG/Source/PCGEditor/Private/EditorMode/Tools/PCGInteractiveToolSettings.h
+++ b/Engine/Plugins/PCG/Source/PCGEditor/Private/EditorMode/Tools/PCGInteractiveToolSettings.h
@@ -58,6 +58,7 @@ public:
 	virtual void Apply(UInteractiveTool* OwningTool) override;
 	virtual void Cancel(UInteractiveTool* OwningTool) override;
 	virtual void Shutdown() override;
+	virtual void PostEditChangeProperty(FPropertyChangedEvent& PropertyChangedEvent) override;
 
 	//~Begin interface UInteractiveToolPropertySet
 	virtual void RestoreProperties(UInteractiveTool* RestoreToTool, const FString& CacheIdentifier = TEXT("")) override;
diff --git a/Engine/Plugins/PCG/Source/PCGEditor/Public/Details/PCGInteractiveToolSettingsDetails.h b/Engine/Plugins/PCG/Source/PCGEditor/Public/Details/PCGInteractiveToolSettingsDetails.h
--- a/Engine/Plugins/PCG/Source/PCGEditor/Public/Details/PCGInteractiveToolSettingsDetails.h
+++ b/Engine/Plugins/PCG/Source/PCGEditor/Public/Details/PCGInteractiveToolSettingsDetails.h
@@ -24,7 +24,7 @@ public:
 	virtual void CustomizeDetails(IDetailLayoutBuilder& DetailBuilder) override;
 	virtual void CustomizeDetails(const TSharedPtr<IDetailLayoutBuilder>& DetailBuilder) override;
 private:
-	void Rebuild(UObject* Object, const FPropertyChangedChainEvent& PropertyChangedChainEvent);
+	void Rebuild(UObject* Object, FPropertyChangedEvent& PropertyChangedEvent);
 	
 	FText GetResetToolDataButtonText() const;
 	EVisibility IsResetToolDataButtonVisible() const;
