[PCGEditor] Rename SPCGEditorNodeSource to SPCGCodeEditor

#rnx

#Synced-CL 50811096

diff --git a/Engine/Plugins/PCG/Source/PCG/Internal/Compute/IPCGCodeEditorTextProvider.h b/Engine/Plugins/PCG/Source/PCG/Internal/Compute/IPCGCodeEditorTextProvider.h
new file mode 100644
--- /dev/null
+++ b/Engine/Plugins/PCG/Source/PCG/Internal/Compute/IPCGCodeEditorTextProvider.h
@@ -0,0 +1,41 @@
+// Copyright Epic Games, Inc. All Rights Reserved.
+
+#pragma once
+
+#include "UObject/Interface.h"
+#include "IPCGCodeEditorTextProvider.generated.h"
+
+UINTERFACE(MinimalAPI)
+class UPCGCodeEditorTextProvider : public UInterface
+{
+	GENERATED_BODY()
+};
+
+/**
+* Interface for objects that provide source code to the code editor widget.
+*/
+class IPCGCodeEditorTextProvider
+{
+	GENERATED_BODY()
+
+#if WITH_EDITOR
+public:
+	/** Get text for declarations pane. */
+	virtual FString GetDeclarationsText() const = 0;
+
+	/** Get text for functions pane. */
+	virtual FString GetFunctionsText() const = 0;
+
+	/** Get text for source pane. */
+	virtual FString GetSourceText() const = 0;
+
+	/** Set functions text after edit modifications. */
+	virtual void SetFunctionsText(const FString& InText) = 0;
+
+	/** Set source text after edit modifications. */
+	virtual void SetSourceText(const FString& InText) = 0;
+
+	/** Whether the text can be edited. */
+	virtual bool IsReadOnly() const = 0;
+#endif
+};
diff --git a/Engine/Plugins/PCG/Source/PCG/Internal/Compute/IPCGNodeSourceTextProvider.h b/Engine/Plugins/PCG/Source/PCG/Internal/Compute/IPCGNodeSourceTextProvider.h
deleted file mode 100644
diff --git a/Engine/Plugins/PCG/Source/PCG/Private/Compute/Elements/PCGCustomHLSL.cpp b/Engine/Plugins/PCG/Source/PCG/Private/Compute/Elements/PCGCustomHLSL.cpp
--- a/Engine/Plugins/PCG/Source/PCG/Private/Compute/Elements/PCGCustomHLSL.cpp
+++ b/Engine/Plugins/PCG/Source/PCG/Private/Compute/Elements/PCGCustomHLSL.cpp
@@ -300,53 +300,53 @@ FString UPCGCustomHLSLSettings::GetDeclarationsText() const
 	return InputDeclarations + TEXT("\n\n") + OutputDeclarations + TEXT("\n\n") + HelperDeclarations;
 }
 
-FString UPCGCustomHLSLSettings::GetShaderFunctionsText() const
+FString UPCGCustomHLSLSettings::GetFunctionsText() const
 {
 	return ShaderFunctions;
 }
 
-FString UPCGCustomHLSLSettings::GetShaderText() const
+FString UPCGCustomHLSLSettings::GetSourceText() const
 {
 	return ShaderSource;
 }
 
-void UPCGCustomHLSLSettings::SetShaderFunctionsText(const FString& NewFunctionsText)
+void UPCGCustomHLSLSettings::SetFunctionsText(const FString& InText)
 {
 	const FName PropertyName = GET_MEMBER_NAME_CHECKED(UPCGCustomHLSLSettings, ShaderFunctions);
 	FProperty* Property = FindFProperty<FProperty>(StaticClass(), PropertyName);
 	FPropertyChangedEvent PropertyChangedEvent(Property, EPropertyChangeType::ValueSet);
 
 	{
-		FScopedTransaction Transaction(LOCTEXT("OnSetShaderFunctionsText", "Set Shader Functions Text"));
+		FScopedTransaction Transaction(LOCTEXT("OnSetShaderFunctionsText", "Set Functions Text"));
 	
 		PreEditChange(Property);
 		Modify();
-		ShaderFunctions = NewFunctionsText;
+		ShaderFunctions = InText;
 		PostEditChangeProperty(PropertyChangedEvent);
 	}
 	
 	OnSettingsChangedDelegate.Broadcast(this, GetChangeTypeForProperty(PropertyName));
 }
 
-void UPCGCustomHLSLSettings::SetShaderText(const FString& NewText)
+void UPCGCustomHLSLSettings::SetSourceText(const FString& InText)
 {
 	const FName PropertyName = GET_MEMBER_NAME_CHECKED(UPCGCustomHLSLSettings, ShaderSource);
 	FProperty* Property = FindFProperty<FProperty>(StaticClass(), PropertyName);
 	FPropertyChangedEvent PropertyChangedEvent(Property, EPropertyChangeType::ValueSet);
 	
 	{
-		FScopedTransaction Transaction(LOCTEXT("OnSetShaderSourceText", "Set Shader Source Text"));
+		FScopedTransaction Transaction(LOCTEXT("OnSetShaderSourceText", "Set Source Text"));
 
 		PreEditChange(Property);
 		Modify();
-		ShaderSource = NewText;
+		ShaderSource = InText;
 		PostEditChangeProperty(PropertyChangedEvent);
 	}
 
 	OnSettingsChangedDelegate.Broadcast(this, GetChangeTypeForProperty(PropertyName));
 }
 
-bool UPCGCustomHLSLSettings::IsShaderTextReadOnly() const
+bool UPCGCustomHLSLSettings::IsReadOnly() const
 {
 	return KernelSourceOverride != nullptr;
 }
diff --git a/Engine/Plugins/PCG/Source/PCG/Private/Compute/Elements/PCGCustomHLSL.h b/Engine/Plugins/PCG/Source/PCG/Private/Compute/Elements/PCGCustomHLSL.h
--- a/Engine/Plugins/PCG/Source/PCG/Private/Compute/Elements/PCGCustomHLSL.h
+++ b/Engine/Plugins/PCG/Source/PCG/Private/Compute/Elements/PCGCustomHLSL.h
@@ -3,7 +3,7 @@
 #pragma once
 
 #include "PCGSettings.h"
-#include "Compute/IPCGNodeSourceTextProvider.h"
+#include "Compute/IPCGCodeEditorTextProvider.h"
 
 #include "Compute/PCGPinPropertiesGPU.h"
 
@@ -41,7 +41,7 @@ enum class EPCGDispatchThreadCount : uint8
 UCLASS(MinimalAPI, BlueprintType, ClassGroup = (Procedural))
 class UPCGCustomHLSLSettings
 	: public UPCGSettings
-	, public IPCGNodeSourceTextProvider
+	, public IPCGCodeEditorTextProvider
 {
 	GENERATED_BODY()
 
@@ -86,14 +86,14 @@ protected:
 
 public:
 #if WITH_EDITOR
-	//~Begin IPCGNodeSourceTextProvider interface
-	FString GetShaderText() const override;
+	//~Begin IPCGCodeEditorTextProvider interface
+	FString GetSourceText() const override;
 	FString GetDeclarationsText() const override;
-	FString GetShaderFunctionsText() const override;
-	void SetShaderFunctionsText(const FString& NewFunctionsText) override;
-	void SetShaderText(const FString& NewText) override;
-	bool IsShaderTextReadOnly() const override;
-	//~End IPCGNodeSourceTextProvider interface
+	FString GetFunctionsText() const override;
+	void SetFunctionsText(const FString& InText) override;
+	void SetSourceText(const FString& InText) override;
+	bool IsReadOnly() const override;
+	//~End IPCGCodeEditorTextProvider interface
 #endif
 
 #if WITH_EDITOR
diff --git a/Engine/Plugins/PCG/Source/PCG/Private/Compute/PCGComputeSource.cpp b/Engine/Plugins/PCG/Source/PCG/Private/Compute/PCGComputeSource.cpp
--- a/Engine/Plugins/PCG/Source/PCG/Private/Compute/PCGComputeSource.cpp
+++ b/Engine/Plugins/PCG/Source/PCG/Private/Compute/PCGComputeSource.cpp
@@ -47,13 +47,13 @@ FString UPCGComputeSource::GetVirtualPath() const
 }
 
 #if WITH_EDITOR
-void UPCGComputeSource::SetShaderText(const FString& NewText)
+void UPCGComputeSource::SetSourceText(const FString& InText)
 {
-	if (!Source.Equals(NewText, ESearchCase::CaseSensitive))
+	if (!Source.Equals(InText, ESearchCase::CaseSensitive))
 	{
 		Modify();
 
-		Source = NewText;
+		Source = InText;
 
 		OnModifiedDelegate.Broadcast(this);
 	}
diff --git a/Engine/Plugins/PCG/Source/PCG/Public/Compute/PCGComputeSource.h b/Engine/Plugins/PCG/Source/PCG/Public/Compute/PCGComputeSource.h
--- a/Engine/Plugins/PCG/Source/PCG/Public/Compute/PCGComputeSource.h
+++ b/Engine/Plugins/PCG/Source/PCG/Public/Compute/PCGComputeSource.h
@@ -2,7 +2,7 @@
 
 #pragma once
 
-#include "Compute/IPCGNodeSourceTextProvider.h"
+#include "Compute/IPCGCodeEditorTextProvider.h"
 
 #include "ComputeFramework/ComputeSource.h"
 
@@ -16,7 +16,7 @@ DECLARE_MULTICAST_DELEGATE_OneParam(FOnPCGComputeSourceModified, const UPCGCompu
 #define UE_API PCG_API
 
 UCLASS(MinimalAPI)
-class UPCGComputeSource : public UComputeSource, public IPCGNodeSourceTextProvider
+class UPCGComputeSource : public UComputeSource, public IPCGCodeEditorTextProvider
 {
 	GENERATED_BODY()
 
@@ -34,14 +34,14 @@ public:
 	//~ End UComputeSource Interface.
 
 #if WITH_EDITOR
-	//~Begin IPCGNodeSourceTextProvider interface
-	FString GetShaderText() const override { return GetSource(); }
+	//~Begin IPCGCodeEditorTextProvider interface
 	FString GetDeclarationsText() const override { return {}; }
-	FString GetShaderFunctionsText() const override { return {}; }
-	void SetShaderFunctionsText(const FString& NewFunctionsText) override {}
-	void SetShaderText(const FString& NewText) override;
-	bool IsShaderTextReadOnly() const override { return false; }
-	//~End IPCGNodeSourceTextProvider interface
+	FString GetSourceText() const override { return GetSource(); }
+	FString GetFunctionsText() const override { return {}; }
+	void SetSourceText(const FString& InText) override;
+	void SetFunctionsText(const FString& InText) override {}
+	bool IsReadOnly() const override { return false; }
+	//~End IPCGCodeEditorTextProvider interface
 	
 	UE_API void SetSource(const FString& InSource);
 
@@ -53,6 +53,19 @@ protected:
 	UPROPERTY(EditAnywhere, Category = "Source", meta = (DisplayAfter = "AdditionalSources"))
 	FString Source;
 #endif
+
+	// Deprecated section
+#if WITH_EDITOR
+public:
+	UE_DEPRECATED(5.8, "Use GetSourceText() instead")
+	FString GetShaderText() const { return GetSource(); }
+	UE_DEPRECATED(5.8, "Use GetFunctionsText() instead")
+	FString GetShaderFunctionsText() const { return {}; }
+	UE_DEPRECATED(5.8, "Use SetFunctionsText() instead")
+	void SetShaderFunctionsText(const FString& InText) {}
+	UE_DEPRECATED(5.8, "Use IsReadOnly() instead")
+	bool IsShaderTextReadOnly() const { return false; }
+#endif
 };
 
 #undef UE_API
diff --git a/Engine/Plugins/PCG/Source/PCGEditor/Private/Details/PCGComputeSourceDetails.cpp b/Engine/Plugins/PCG/Source/PCGEditor/Private/Details/PCGComputeSourceDetails.cpp
--- a/Engine/Plugins/PCG/Source/PCGEditor/Private/Details/PCGComputeSourceDetails.cpp
+++ b/Engine/Plugins/PCG/Source/PCGEditor/Private/Details/PCGComputeSourceDetails.cpp
@@ -7,8 +7,8 @@
 #include "PCGEditor.h"
 #include "PCGEditorGraph.h"
 #include "PCGHLSLSyntaxHighlighter.h"
-#include "Widgets/SPCGEditorNodeSource.h"
-#include "Widgets/SPCGNodeSourceTextBox.h"
+#include "Widgets/SPCGCodeEditor.h"
+#include "Widgets/SPCGCodeEditorTextBox.h"
 
 #include "DetailWidgetRow.h"
 #include "DetailLayoutBuilder.h"
@@ -47,9 +47,11 @@ void FPCGComputeSourceDetails::CustomizeDetails(IDetailLayoutBuilder& DetailBuil
 		DetailsView->HideFilterArea(true);
 	}
 
+	// @todo_pcg: This should use the SPCGCodeEditor widget.
+
 	SExpandableArea::FArguments ExpandableAreaArgs;
 	ExpandableAreaArgs.AreaTitleFont(FAppStyle::GetFontStyle("DetailsView.CategoryFontStyle"));
-	ExpandableAreaArgs.AreaTitle(LOCTEXT("PCGNodeSource_ShaderText_Title", "Shader Source"));
+	ExpandableAreaArgs.AreaTitle(LOCTEXT("PCGCodeEditor_ShaderText_Title", "Source"));
 
 	TSharedRef<IPropertyHandle> SourcePropertyHandle = DetailBuilder.GetProperty(TEXT("Source"));
 	DetailBuilder.EditDefaultProperty(SourcePropertyHandle)->CustomWidget()
@@ -59,8 +61,7 @@ void FPCGComputeSourceDetails::CustomizeDetails(IDetailLayoutBuilder& DetailBuil
 		+SVerticalBox::Slot()
 		.AutoHeight()
 		[
-			// @todo_pcg: Share functionality with the NodeSourceEditor
-			//ConstructNonExpandableHeaderWidget(ExpandableAreaArgs.AreaTitle(LOCTEXT("PCGNodeSource_ShaderText_Title", "Shader Source")))
+			//ConstructNonExpandableHeaderWidget(ExpandableAreaArgs.AreaTitle(LOCTEXT("PCGCodeEditor_ShaderText_Title", "Source")))
 			SNew(SBorder)
 				.BorderImage(FStyleDefaults::GetNoBrush())
 				.BorderBackgroundColor(FLinearColor::Transparent)
@@ -98,7 +99,7 @@ void FPCGComputeSourceDetails::CustomizeDetails(IDetailLayoutBuilder& DetailBuil
 			.MinDesiredHeight(200)
 			.MaxDesiredHeight(800)
 			[
-				SAssignNew(SourceTextBox, SPCGNodeSourceTextBox)
+				SAssignNew(SourceTextBox, SPCGCodeEditorTextBox)
 				.Text(this, &FPCGComputeSourceDetails::GetSourceText)
 				.IsReadOnly(false)
 				.OnTextChanged(this, &FPCGComputeSourceDetails::OnSourceTextChanged)
@@ -115,7 +116,7 @@ void FPCGComputeSourceDetails::CustomizeDetails(IDetailLayoutBuilder& DetailBuil
 FText FPCGComputeSourceDetails::GetSourceText() const
 {
 	UPCGComputeSource* ComputeSource = ComputeSourceWeakPtr.Get();
-	return ComputeSource ? FText::FromString(ComputeSource->GetShaderText()) : FText::GetEmpty();
+	return ComputeSource ? FText::FromString(ComputeSource->GetSourceText()) : FText::GetEmpty();
 }
 
 void FPCGComputeSourceDetails::OnSourceTextChanged(const FText& InText)
@@ -140,9 +141,9 @@ void FPCGComputeSourceDetails::SetComputeSourceText() const
 	{
 		const FString SourceString = SourceText.ToString();
 
-		if (!SourceString.Equals(ComputeSource->GetShaderText(), ESearchCase::CaseSensitive))
+		if (!SourceString.Equals(ComputeSource->GetSourceText(), ESearchCase::CaseSensitive))
 		{
-			ComputeSource->SetShaderText(SourceString);
+			ComputeSource->SetSourceText(SourceString);
 		}
 	}
 }
diff --git a/Engine/Plugins/PCG/Source/PCGEditor/Private/Details/PCGComputeSourceDetails.h b/Engine/Plugins/PCG/Source/PCGEditor/Private/Details/PCGComputeSourceDetails.h
--- a/Engine/Plugins/PCG/Source/PCGEditor/Private/Details/PCGComputeSourceDetails.h
+++ b/Engine/Plugins/PCG/Source/PCGEditor/Private/Details/PCGComputeSourceDetails.h
@@ -8,7 +8,7 @@
 #include "UObject/WeakObjectPtr.h"
 
 class FPCGHLSLSyntaxHighlighter;
-class SPCGNodeSourceTextBox;
+class SPCGCodeEditorTextBox;
 class UPCGComputeSource;
 
 /**
@@ -31,7 +31,7 @@ protected:
 
 	TWeakObjectPtr<UPCGComputeSource> ComputeSourceWeakPtr;
 	TSharedRef<FPCGHLSLSyntaxHighlighter> SyntaxHighlighter;
-	TSharedPtr<SPCGNodeSourceTextBox> SourceTextBox;
+	TSharedPtr<SPCGCodeEditorTextBox> SourceTextBox;
 
 	FText SourceText;
 };
diff --git a/Engine/Plugins/PCG/Source/PCGEditor/Private/Details/PCGCustomHLSLSettingsDetails.cpp b/Engine/Plugins/PCG/Source/PCGEditor/Private/Details/PCGCustomHLSLSettingsDetails.cpp
--- a/Engine/Plugins/PCG/Source/PCGEditor/Private/Details/PCGCustomHLSLSettingsDetails.cpp
+++ b/Engine/Plugins/PCG/Source/PCGEditor/Private/Details/PCGCustomHLSLSettingsDetails.cpp
@@ -29,7 +29,7 @@ void FPCGCustomHLSLSettingsDetails::CustomizeDetails(IDetailLayoutBuilder& Detai
 		const TSharedPtr<const FPCGEditor> PCGEditorSharedPtr = PCGEditorWeakPtr.Pin();
 
 		// Add a "open source editor" button if the source editor tab is not already visible.
-		if (PCGEditorSharedPtr.IsValid() && !PCGEditorSharedPtr->IsPanelCurrentlyOpen(PCGEditorTabs::HLSLSourceID))
+		if (PCGEditorSharedPtr.IsValid() && !PCGEditorSharedPtr->IsPanelCurrentlyOpen(PCGEditorTabs::CodeEditorID))
 		{
 			SettingsCategory.AddCustomRow(FText::GetEmpty())
 				.Visibility(
@@ -37,7 +37,7 @@ void FPCGCustomHLSLSettingsDetails::CustomizeDetails(IDetailLayoutBuilder& Detai
 				[PCGEditorWeakPtr]
 				{
 					const TSharedPtr<const FPCGEditor> SharedPtr = PCGEditorWeakPtr.Pin();
-					return (SharedPtr.IsValid() && SharedPtr->IsPanelCurrentlyOpen(PCGEditorTabs::HLSLSourceID)) ? EVisibility::Collapsed : EVisibility::Visible;
+					return (SharedPtr.IsValid() && SharedPtr->IsPanelCurrentlyOpen(PCGEditorTabs::CodeEditorID)) ? EVisibility::Collapsed : EVisibility::Visible;
 				}))
 				.ValueContent()
 				.MaxDesiredWidth(120.f)
@@ -53,7 +53,7 @@ void FPCGCustomHLSLSettingsDetails::CustomizeDetails(IDetailLayoutBuilder& Detai
 						{
 							if (const TSharedPtr<const FPCGEditor> SharedPtr = PCGEditorWeakPtr.Pin())
 							{
-								SharedPtr->BringFocusToPanel(PCGEditorTabs::HLSLSourceID);
+								SharedPtr->BringFocusToPanel(PCGEditorTabs::CodeEditorID);
 							}
 
 							return FReply::Handled();
diff --git a/Engine/Plugins/PCG/Source/PCGEditor/Private/PCGEditor.cpp b/Engine/Plugins/PCG/Source/PCGEditor/Private/PCGEditor.cpp
--- a/Engine/Plugins/PCG/Source/PCGEditor/Private/PCGEditor.cpp
+++ b/Engine/Plugins/PCG/Source/PCGEditor/Private/PCGEditor.cpp
@@ -48,7 +48,7 @@
 #include "Widgets/SPCGEditorGraphParamsView.h"
 #include "Widgets/SPCGEditorGraphProfilingView.h"
 #include "Widgets/SPCGEditorGraphTitleBar.h"
-#include "Widgets/SPCGEditorNodeSource.h"
+#include "Widgets/SPCGCodeEditor.h"
 #include "Widgets/AssetEditorViewport/SPCGEditorViewport.h"
 
 #include "AssetDefinition.h"
@@ -107,7 +107,7 @@ namespace FPCGEditor_private
 		PCGEditorTabs::DeterminismID,
 		PCGEditorTabs::ProfilingID,
 		PCGEditorTabs::LogID,
-		PCGEditorTabs::HLSLSourceID,
+		PCGEditorTabs::CodeEditorID,
 		PCGEditorTabs::UserParamsID,
 		PCGEditorTabs::ViewportID[0],
 		PCGEditorTabs::ViewportID[1],
@@ -118,7 +118,7 @@ namespace FPCGEditor_private
 
 	const FName AllRestrictedPanels[] =
 	{
-		PCGEditorTabs::HLSLSourceID,
+		PCGEditorTabs::CodeEditorID,
 	};
 }
 
@@ -192,7 +192,7 @@ void FPCGEditor::Initialize(const EToolkitMode::Type InMode, const TSharedPtr<cl
 
 	InitializePCGEditorGraph(MainGraph);
 
-	GetOrCreateNodeSourceWidget();
+	GetOrCreateCodeEditorWidget();
 
 	BindCommands();
 
@@ -556,9 +556,9 @@ const FPCGStack* FPCGEditor::GetStackBeingInspected() const
 
 void FPCGEditor::SetSourceEditorTargetObject(UObject* InObject)
 {
-	if (NodeSourceWidget.IsValid())
+	if (CodeEditorWidget.IsValid())
 	{
-		NodeSourceWidget->SetTextProviderObject(InObject);
+		CodeEditorWidget->SetTextProviderObject(InObject);
 	}
 }
 
@@ -631,7 +631,7 @@ FName FPCGEditor::GetPanelID(const EPCGEditorPanel Panel) const
 		case EPCGEditorPanel::NodePalette:
 			return PCGEditorTabs::PaletteID;
 		case EPCGEditorPanel::NodeSource:
-			return PCGEditorTabs::HLSLSourceID;
+			return PCGEditorTabs::CodeEditorID;
 		case EPCGEditorPanel::Profiling:
 			return PCGEditorTabs::ProfilingID;
 		case EPCGEditorPanel::PropertyDetails1:
@@ -3378,21 +3378,27 @@ TSharedRef<SPCGEditorGraphLogView> FPCGEditor::GetOrCreateLogWidget()
 	return LogWidget.ToSharedRef();
 }
 
-TSharedPtr<FWorkflowTabFactory> FPCGEditor::CreateNodeSourceTabFactory()
+TSharedPtr<FWorkflowTabFactory> FPCGEditor::CreateCodeEditorTabFactory()
 {
-	FPCGGenericTabFactoryParams Params(PCGEditorTabs::HLSLSourceID, SharedThis<FPCGEditor>(this), [this]() { return GetOrCreateNodeSourceWidget(); });
-	Params.Label = LOCTEXT("NodeSourceLabel", "HLSL Source");
+	FPCGGenericTabFactoryParams Params(PCGEditorTabs::CodeEditorID, SharedThis<FPCGEditor>(this), [this]() { return GetOrCreateCodeEditorWidget(); });
+	Params.Label = LOCTEXT("CodeEditorLabel", "Code Editor");
 	return MakeShared<FPCGGenericTabFactory>(Params);
 }
 
-TSharedRef<SPCGEditorNodeSource> FPCGEditor::GetOrCreateNodeSourceWidget()
+// Deprecated 5.8
+TSharedPtr<FWorkflowTabFactory> FPCGEditor::CreateNodeSourceTabFactory()
+{
+	return CreateCodeEditorTabFactory();
+}
+
+TSharedRef<SPCGCodeEditor> FPCGEditor::GetOrCreateCodeEditorWidget()
 {
-	if (!NodeSourceWidget.IsValid())
+	if (!CodeEditorWidget.IsValid())
 	{
-		NodeSourceWidget = SNew(SPCGEditorNodeSource);
+		CodeEditorWidget = SNew(SPCGCodeEditor);
 	}
 	
-	return NodeSourceWidget.ToSharedRef();
+	return CodeEditorWidget.ToSharedRef();
 }
 
 TSharedPtr<FWorkflowTabFactory> FPCGEditor::CreateGraphParamsTabFactory()
@@ -3827,12 +3833,12 @@ void FPCGEditor::OnGraphChanged(UPCGGraphInterface* InGraph, EPCGChangeType Chan
 
 void FPCGEditor::OnNodeSourceCompiled(const UPCGNode* InNode, const FPCGCompilerDiagnostics& InDiagnostics)
 {
-	check(NodeSourceWidget);
+	check(CodeEditorWidget);
 
 	const UPCGSettings* Settings = InNode ? InNode->GetSettings() : nullptr;
-	if (Settings && NodeSourceWidget->GetTextProviderObject() == Settings)
+	if (Settings && CodeEditorWidget->GetTextProviderObject() == Settings)
 	{
-		NodeSourceWidget->OnDiagnosticsUpdated(InDiagnostics);
+		CodeEditorWidget->OnDiagnosticsUpdated(InDiagnostics);
 	}
 }
 
diff --git a/Engine/Plugins/PCG/Source/PCGEditor/Private/PCGEditorDefaultMode.cpp b/Engine/Plugins/PCG/Source/PCGEditor/Private/PCGEditorDefaultMode.cpp
--- a/Engine/Plugins/PCG/Source/PCGEditor/Private/PCGEditorDefaultMode.cpp
+++ b/Engine/Plugins/PCG/Source/PCGEditor/Private/PCGEditorDefaultMode.cpp
@@ -38,7 +38,7 @@ namespace PCGEditorDefaultMode
 						FTabManager::NewStack()
 						->SetSizeCoefficient(0.2f)
 						->AddTab(PCGEditorTabs::ViewportID[0], ETabState::OpenedTab)
-						->AddTab(PCGEditorTabs::HLSLSourceID, ETabState::OpenedTab)
+						->AddTab(PCGEditorTabs::CodeEditorID, ETabState::OpenedTab)
 						->SetForegroundTab(FName(PCGEditorTabs::ViewportID[0]))
 					)
 					->Split // Node Graph
@@ -118,7 +118,7 @@ FPCGEditorDefaultMode::FPCGEditorDefaultMode(TSharedPtr<FPCGEditor> InPCGEditor)
 		{PCGEditorTabs::DeterminismID, PCGEditorDefaultMode::FCreateTabFactoryDelegate::CreateSP(InPCGEditor.Get(), &FPCGEditor::CreateDeterminismTabFactory)},
 		{PCGEditorTabs::ProfilingID, PCGEditorDefaultMode::FCreateTabFactoryDelegate::CreateSP(InPCGEditor.Get(), &FPCGEditor::CreateProfilingTabFactory)},
 		{PCGEditorTabs::LogID, PCGEditorDefaultMode::FCreateTabFactoryDelegate::CreateSP(InPCGEditor.Get(), &FPCGEditor::CreateLogTabFactory)},
-		{PCGEditorTabs::HLSLSourceID, PCGEditorDefaultMode::FCreateTabFactoryDelegate::CreateSP(InPCGEditor.Get(), &FPCGEditor::CreateNodeSourceTabFactory)},
+		{PCGEditorTabs::CodeEditorID, PCGEditorDefaultMode::FCreateTabFactoryDelegate::CreateSP(InPCGEditor.Get(), &FPCGEditor::CreateCodeEditorTabFactory)},
 		{PCGEditorTabs::UserParamsID, PCGEditorDefaultMode::FCreateTabFactoryDelegate::CreateSP(InPCGEditor.Get(), &FPCGEditor::CreateGraphParamsTabFactory)},
 		{PCGEditorTabs::ViewportID[0], PCGEditorDefaultMode::FCreateTabFactoryDelegate::CreateSP(InPCGEditor.Get(), &FPCGEditor::CreateViewportTabFactory, 0, ViewportGroup)},
 		{PCGEditorTabs::ViewportID[1], PCGEditorDefaultMode::FCreateTabFactoryDelegate::CreateSP(InPCGEditor.Get(), &FPCGEditor::CreateViewportTabFactory, 1, ViewportGroup)},
diff --git a/Engine/Plugins/PCG/Source/PCGEditor/Private/PCGEditorModule.cpp b/Engine/Plugins/PCG/Source/PCGEditor/Private/PCGEditorModule.cpp
--- a/Engine/Plugins/PCG/Source/PCGEditor/Private/PCGEditorModule.cpp
+++ b/Engine/Plugins/PCG/Source/PCGEditor/Private/PCGEditorModule.cpp
@@ -94,7 +94,7 @@
 #include "HierarchyEditor/PropertyBagHierarchyViewModel.h"
 #include "Schema/PCGGraphParametersSchema.h"
 #include "UObject/ObjectSaveContext.h"
-#include "Widgets/SPCGNodeSourceTextBox.h"
+#include "Widgets/SPCGCodeEditorTextBox.h"
 
 #define LOCTEXT_NAMESPACE "FPCGEditorModule"
 
@@ -127,7 +127,7 @@ void FPCGEditorModule::StartupModule()
 	FPCGEditorStyle::Register();
 	FPCGEditorCommands::Register();
 	FPCGEditorSpawnNodeCommands::Register();
-	FPCGNodeSourceEditorTextBoxCommands::Register();
+	FPCGCodeEditorTextBoxCommands::Register();
 
 	UPCGEditorMode::RegisterEditorMode();
 
@@ -174,7 +174,7 @@ void FPCGEditorModule::ShutdownModule()
 
 	UPCGEditorMode::UnregisterEditorMode();
 
-	FPCGNodeSourceEditorTextBoxCommands::Unregister();
+	FPCGCodeEditorTextBoxCommands::Unregister();
 	FPCGEditorSpawnNodeCommands::Unregister();
 	FPCGEditorCommands::Unregister();
 	FPCGEditorStyle::Unregister();
diff --git a/Engine/Plugins/PCG/Source/PCGEditor/Private/Widgets/SPCGCodeEditor.cpp b/Engine/Plugins/PCG/Source/PCGEditor/Private/Widgets/SPCGCodeEditor.cpp
new file mode 100644
--- /dev/null
+++ b/Engine/Plugins/PCG/Source/PCGEditor/Private/Widgets/SPCGCodeEditor.cpp
@@ -0,0 +1,286 @@
+// Copyright Epic Games, Inc. All Rights Reserved.
+
+#include "SPCGCodeEditor.h"
+
+#include "PCGSettings.h"
+#include "Compute/IPCGCodeEditorTextProvider.h"
+#include "Helpers/PCGHelpers.h"
+
+#include "PCGEditor.h"
+#include "PCGHLSLSyntaxHighlighter.h"
+#include "SPCGCodeEditorTextBox.h"
+
+#include "Widgets/Input/SButton.h"
+#include "Widgets/Layout/SSpacer.h"
+
+#define LOCTEXT_NAMESPACE "SPCGCodeEditor"
+
+SPCGCodeEditor::SPCGCodeEditor()
+	: SyntaxHighlighterDeclarations(FPCGHLSLSyntaxHighlighter::Create())
+	, SyntaxHighlighterText(FPCGHLSLSyntaxHighlighter::Create())
+{
+}
+
+void SPCGCodeEditor::Construct(const FArguments& InArgs)
+{
+	SExpandableArea::FArguments ExpandableAreaArgs;
+	ExpandableAreaArgs.AreaTitleFont(FAppStyle::GetFontStyle("DetailsView.CategoryFontStyle"));
+	
+	ChildSlot
+	[
+		SNew(SSplitter)
+		.Orientation(Orient_Vertical)
+		.PhysicalSplitterHandleSize(4.0f)
+		.HitDetectionSplitterHandleSize(6.0f)
+		+SSplitter::Slot()
+		.Value(1.0)
+		.SizeRule(this, &SPCGCodeEditor::GetDeclarationsSectionSizeRule)
+		[
+			SAssignNew(DeclarationsExpandableArea, SExpandableArea) = ExpandableAreaArgs
+			.AreaTitle(LOCTEXT("PCGCodeEditor_Declarations_Title", "Declarations (Read-Only)"))
+			.InitiallyCollapsed(true)
+			.AllowAnimatedTransition(false)
+			.BodyContent()
+			[
+				SAssignNew(DeclarationsTextBox, SPCGCodeEditorTextBox)
+				.Text(this, &SPCGCodeEditor::GetDeclarationsText)
+				.IsReadOnly(true)
+				.Marshaller(SyntaxHighlighterDeclarations)
+			]
+		]
+		+SSplitter::Slot()
+		.Value(1.0)
+		.SizeRule(this, &SPCGCodeEditor::GetFunctionsSectionSizeRule)
+		[
+			SAssignNew(FunctionsExpandableArea, SExpandableArea) = ExpandableAreaArgs
+			.AreaTitle(LOCTEXT("PCGCodeEditor_Functions_Title", "Functions"))
+			.InitiallyCollapsed(true)
+			.AllowAnimatedTransition(false)
+			.BodyContent()
+			[
+				SAssignNew(FunctionsTextBox, SPCGCodeEditorTextBox)
+				.Text(this, &SPCGCodeEditor::GetFunctionsText)
+				.IsReadOnly(this, &SPCGCodeEditor::IsReadOnly)
+				.Marshaller(SyntaxHighlighterDeclarations)
+				.OnTextChanged(this, &SPCGCodeEditor::OnFunctionsTextChanged)
+				.OnTextCommitted(this, &SPCGCodeEditor::OnFunctionsTextCommitted)
+				.OnTextChangesApplied(this, &SPCGCodeEditor::OnFunctionsTextChangesApplied)
+			]
+		]
+		+SSplitter::Slot()
+		.Value(1.0)
+		[
+			SNew(SVerticalBox)
+			+SVerticalBox::Slot()
+			.AutoHeight()
+			[
+				ConstructNonExpandableHeaderWidget(ExpandableAreaArgs.AreaTitle(LOCTEXT("PCGCodeEditor_Text_Title", "Source")))
+			]
+			+SVerticalBox::Slot()
+			.FillHeight(1.0f)
+			[
+				SAssignNew(SourceTextBox, SPCGCodeEditorTextBox)
+				.Text(this, &SPCGCodeEditor::GetSourceText)
+				.IsReadOnly(this, &SPCGCodeEditor::IsReadOnly)
+				.Marshaller(SyntaxHighlighterText)
+				.OnTextChanged(this, &SPCGCodeEditor::OnSourceTextChanged)
+				.OnTextCommitted(this, &SPCGCodeEditor::OnSourceTextCommitted)
+				.OnTextChangesApplied(this, &SPCGCodeEditor::OnSourceTextChangesApplied)
+			]
+		]
+	];
+}
+
+void SPCGCodeEditor::SetTextProviderObject(UObject* InProviderObject)
+{
+	TextProviderObject = InProviderObject;
+
+	FunctionsTextBox->SetTextProviderObject(InProviderObject);
+	SourceTextBox->SetTextProviderObject(InProviderObject);
+
+	FunctionsText = GetFunctionsText();
+	SourceText = GetSourceText();
+
+	SyntaxHighlighterText->ClearCompilerMessages();
+
+	Refresh();
+}
+
+void SPCGCodeEditor::Refresh()
+{
+	DeclarationsTextBox->Refresh();
+	SourceTextBox->Refresh();
+}
+
+SSplitter::ESizeRule SPCGCodeEditor::GetDeclarationsSectionSizeRule() const
+{
+	SSplitter::ESizeRule SizeRule = SSplitter::ESizeRule::FractionOfParent;
+
+	if (ensure(DeclarationsExpandableArea.IsValid()))
+	{
+		SizeRule = DeclarationsExpandableArea->IsExpanded() ? SSplitter::ESizeRule::FractionOfParent : SSplitter::ESizeRule::SizeToContent;
+	}
+
+	return SizeRule;
+}
+
+SSplitter::ESizeRule SPCGCodeEditor::GetFunctionsSectionSizeRule() const
+{
+	SSplitter::ESizeRule SizeRule = SSplitter::ESizeRule::FractionOfParent;
+
+	if (ensure(FunctionsExpandableArea.IsValid()))
+	{
+		SizeRule = FunctionsExpandableArea->IsExpanded() ? SSplitter::ESizeRule::FractionOfParent : SSplitter::ESizeRule::SizeToContent;
+	}
+
+	return SizeRule;
+}
+
+IPCGCodeEditorTextProvider* SPCGCodeEditor::GetTextProviderInterface() const
+{
+	return Cast<IPCGCodeEditorTextProvider>(TextProviderObject);
+}
+
+FText SPCGCodeEditor::GetDeclarationsText() const
+{
+	if (GetTextProviderInterface())
+	{
+		return FText::FromString(GetTextProviderInterface()->GetDeclarationsText());
+	}
+
+	return FText::GetEmpty();
+}
+
+FText SPCGCodeEditor::GetFunctionsText() const
+{
+	if (GetTextProviderInterface())
+	{
+		return FText::FromString(GetTextProviderInterface()->GetFunctionsText());
+	}
+
+	return FText::GetEmpty();
+}
+
+FText SPCGCodeEditor::GetSourceText() const
+{
+	if (GetTextProviderInterface())
+	{
+		return FText::FromString(GetTextProviderInterface()->GetSourceText());
+	}
+
+	return FText::GetEmpty();
+}
+
+bool SPCGCodeEditor::IsReadOnly() const
+{
+	if (IPCGCodeEditorTextProvider* Provider = Cast<IPCGCodeEditorTextProvider>(TextProviderObject))
+	{
+		return Provider->IsReadOnly();
+	}
+
+	return false;
+}
+
+void SPCGCodeEditor::OnFunctionsTextChanged(const FText& InText)
+{
+	// Always clear compiler messages when text is edited so that markup doesn't lurk in random places.
+	SyntaxHighlighterText->ClearCompilerMessages();
+
+	FunctionsText = InText;
+}
+
+void SPCGCodeEditor::OnSourceTextChanged(const FText& InText)
+{
+	// Always clear compiler messages when text is edited so that markup doesn't lurk in random places.
+	SyntaxHighlighterText->ClearCompilerMessages();
+
+	SourceText = InText;
+}
+
+void SPCGCodeEditor::OnFunctionsTextCommitted(const FText& InText, ETextCommit::Type InCommitInfo)
+{
+	FunctionsText = InText;
+	SetFunctionsText();
+}
+
+void SPCGCodeEditor::OnSourceTextCommitted(const FText& InText, ETextCommit::Type InCommitInfo)
+{
+	SourceText = InText;
+	SetSourceText();
+}
+
+void SPCGCodeEditor::OnFunctionsTextChangesApplied() const
+{
+	SetFunctionsText();
+}
+
+void SPCGCodeEditor::OnSourceTextChangesApplied() const
+{
+	SetSourceText();
+}
+
+void SPCGCodeEditor::SetFunctionsText() const
+{
+	if (IPCGCodeEditorTextProvider* Provider = GetTextProviderInterface())
+	{
+		const FString FunctionsString = FunctionsText.ToString();
+
+		if (!FunctionsString.Equals(Provider->GetFunctionsText(), ESearchCase::CaseSensitive))
+		{
+			Provider->SetFunctionsText(FunctionsString);
+		}
+	}
+}
+
+void SPCGCodeEditor::SetSourceText() const
+{
+	if (IPCGCodeEditorTextProvider* Provider = GetTextProviderInterface())
+	{
+		const FString SourceString = SourceText.ToString();
+
+		if (!SourceString.Equals(Provider->GetSourceText(), ESearchCase::CaseSensitive))
+		{
+			Provider->SetSourceText(SourceString);
+		}
+	}
+}
+
+void SPCGCodeEditor::OnDiagnosticsUpdated(const FPCGCompilerDiagnostics& InDiagnostics) const
+{
+	SyntaxHighlighterText->SetCompilerMessages(InDiagnostics);
+}
+
+TSharedRef<SWidget> SPCGCodeEditor::ConstructNonExpandableHeaderWidget(const SExpandableArea::FArguments& InArgs) const
+{
+	return SNew(SBorder)
+		.BorderImage(FStyleDefaults::GetNoBrush())
+		.BorderBackgroundColor(FLinearColor::Transparent)
+		.Padding(0.0f)
+		[
+			SNew(SButton)
+			.ButtonStyle(FCoreStyle::Get(), "NoBorder")
+			.ContentPadding(InArgs._HeaderPadding)
+			.ForegroundColor(FSlateColor::UseForeground())
+			[
+				SNew(SHorizontalBox)
+				+SHorizontalBox::Slot()
+				.AutoWidth()
+				.Padding(InArgs._AreaTitlePadding)
+				.VAlign(VAlign_Center)
+				[
+					SNew(SSpacer)
+					.Size(InArgs._Style->CollapsedImage.GetImageSize())
+				]
+				+SHorizontalBox::Slot()
+				.FillWidth(1.0f)
+				.VAlign(VAlign_Center)
+				[
+					SNew(STextBlock)
+					.Text(InArgs._AreaTitle)
+					.Font(InArgs._AreaTitleFont)
+				]
+			]
+		];
+}
+
+#undef LOCTEXT_NAMESPACE
diff --git a/Engine/Plugins/PCG/Source/PCGEditor/Private/Widgets/SPCGCodeEditor.h b/Engine/Plugins/PCG/Source/PCGEditor/Private/Widgets/SPCGCodeEditor.h
new file mode 100644
--- /dev/null
+++ b/Engine/Plugins/PCG/Source/PCGEditor/Private/Widgets/SPCGCodeEditor.h
@@ -0,0 +1,75 @@
+// Copyright Epic Games, Inc. All Rights Reserved.
+
+#pragma once
+
+#include "Widgets/SCompoundWidget.h"
+
+#include "Widgets/Layout/SExpandableArea.h"
+#include "Widgets/Views/ITableRow.h"
+#include "Widgets/Views/SListView.h"
+#include "Widgets/Views/STableRow.h"
+
+class FPCGEditor;
+class FPCGHLSLSyntaxHighlighter;
+class IPCGCodeEditorTextProvider;
+class SPCGCodeEditorTextBox;
+struct FPCGCompilerDiagnostics;
+
+/** Window for editing source code. */
+class SPCGCodeEditor : public SCompoundWidget
+{
+public:
+	SLATE_BEGIN_ARGS(SPCGCodeEditor) { }
+	SLATE_END_ARGS()
+
+	SPCGCodeEditor();
+
+	void Construct(const FArguments& InArgs);
+
+	void SetTextProviderObject(UObject* InProviderObject);
+	const UObject* GetTextProviderObject() const { return TextProviderObject.IsValid() ? TextProviderObject.Get() : nullptr; }
+
+	void OnDiagnosticsUpdated(const FPCGCompilerDiagnostics& InDiagnostics) const;
+
+private:
+	void Refresh();
+
+	SSplitter::ESizeRule GetDeclarationsSectionSizeRule() const;
+	SSplitter::ESizeRule GetFunctionsSectionSizeRule() const;
+	IPCGCodeEditorTextProvider* GetTextProviderInterface() const;
+
+	FText GetDeclarationsText() const;
+	FText GetFunctionsText() const;
+	FText GetSourceText() const;
+
+	bool IsReadOnly() const;
+
+	void OnFunctionsTextChanged(const FText& InText);
+	void OnSourceTextChanged(const FText& InText);
+
+	void OnFunctionsTextCommitted(const FText& InText, ETextCommit::Type InCommitInfo);
+	void OnSourceTextCommitted(const FText& InText, ETextCommit::Type InCommitInfo);
+
+	void OnFunctionsTextChangesApplied() const;
+	void OnSourceTextChangesApplied() const;
+
+	void SetFunctionsText() const;
+	void SetSourceText() const;
+
+	TSharedRef<SWidget> ConstructNonExpandableHeaderWidget(const SExpandableArea::FArguments& InArgs) const;
+
+	TSharedRef<FPCGHLSLSyntaxHighlighter> SyntaxHighlighterDeclarations;
+	TSharedRef<FPCGHLSLSyntaxHighlighter> SyntaxHighlighterText;
+
+	TWeakObjectPtr<UObject> TextProviderObject;
+
+	TSharedPtr<SExpandableArea> DeclarationsExpandableArea;
+	TSharedPtr<SExpandableArea> FunctionsExpandableArea;
+
+	TSharedPtr<SPCGCodeEditorTextBox> DeclarationsTextBox;
+	TSharedPtr<SPCGCodeEditorTextBox> FunctionsTextBox;
+	TSharedPtr<SPCGCodeEditorTextBox> SourceTextBox;
+
+	FText FunctionsText;
+	FText SourceText;
+};
diff --git a/Engine/Plugins/PCG/Source/PCGEditor/Private/Widgets/SPCGCodeEditorTextBox.cpp b/Engine/Plugins/PCG/Source/PCGEditor/Private/Widgets/SPCGCodeEditorTextBox.cpp
new file mode 100644
--- /dev/null
+++ b/Engine/Plugins/PCG/Source/PCGEditor/Private/Widgets/SPCGCodeEditorTextBox.cpp
@@ -0,0 +1,761 @@
+// Copyright Epic Games, Inc. All Rights Reserved.
+
+#include "SPCGCodeEditorTextBox.h"
+
+#include "Compute/IPCGCodeEditorTextProvider.h"
+
+#include "PCGEditorStyle.h"
+#include "SPCGShaderTextSearchWidget.h"
+
+#include "Framework/Application/SlateApplication.h"
+#include "Framework/Application/SlateUser.h"
+#include "Framework/MultiBox/MultiBoxBuilder.h"
+#include "Styling/AppStyle.h" 
+#include "Widgets/Layout/SGridPanel.h"
+#include "Widgets/Layout/SScrollBar.h"
+#include "Widgets/Text/SMultiLineEditableText.h"
+
+#define LOCTEXT_NAMESPACE "PCGShaderTextDocumentTextBox"
+
+static int32 GetNumSpacesAtStartOfLine(const FString& InLine)
+{
+	int32 NumSpaces = 0;
+	for (const TCHAR Char : InLine)
+	{
+		if ((Char != TEXT(' ')))
+		{
+			break;
+		}
+				
+		NumSpaces++;
+	}
+
+	return NumSpaces;
+}
+
+static bool IsOpenBrace(const TCHAR& InCharacter)
+{
+	return (InCharacter == TEXT('{') || InCharacter == TEXT('[') || InCharacter == TEXT('('));
+}
+
+static bool IsCloseBrace(const TCHAR& InCharacter)
+{
+	return (InCharacter == TEXT('}') || InCharacter == TEXT(']') || InCharacter == TEXT(')'));
+}
+
+static TCHAR GetMatchedCloseBrace(const TCHAR& InCharacter)
+{
+	return InCharacter == TEXT('{') ? TEXT('}') :
+		(InCharacter == TEXT('[') ? TEXT(']') : TEXT(')'));
+}
+
+static bool IsWhiteSpace(const TCHAR& InCharacter)
+{
+	return InCharacter == TEXT(' ') || InCharacter == TEXT('\r') || InCharacter == TEXT('\n');
+}
+
+FPCGCodeEditorTextBoxCommands::FPCGCodeEditorTextBoxCommands() 
+	: TCommands<FPCGCodeEditorTextBoxCommands>(
+		"PCGShaderTextEditorDocumentTextBox", // Context name for fast lookup
+		NSLOCTEXT("Contexts", "PCGShaderTextEditorDocumentTextBox", "PCG HLSL Source Editor Document TextBox"), // Localized context name for display
+		NAME_None,
+		FAppStyle::GetAppStyleSetName()
+	)
+{
+}
+
+void FPCGCodeEditorTextBoxCommands::RegisterCommands()
+{
+	UI_COMMAND(ApplyChanges, "Apply Changes", "Apply Changes", EUserInterfaceActionType::Button, FInputChord(EKeys::Enter, EModifierKey::Alt));
+
+	UI_COMMAND(Search, "Search", "Search for a String", EUserInterfaceActionType::Button, FInputChord(EKeys::F, EModifierKey::Control));
+	UI_COMMAND(NextOccurrence, "Next Occurrence", "Go to Next Occurrence", EUserInterfaceActionType::Button, FInputChord(EKeys::F3, EModifierKey::None));
+	UI_COMMAND(PreviousOccurrence, "Previous Occurrence", "Go to Previous Occurrence", EUserInterfaceActionType::Button, FInputChord(EKeys::F3, EModifierKey::Shift));
+
+	UI_COMMAND(ToggleComment, "Toggle Comment", "Comment/Uncomment selected lines", EUserInterfaceActionType::Button, FInputChord(EKeys::Slash, EModifierKey::Control));
+}
+
+SPCGCodeEditorTextBox::SPCGCodeEditorTextBox()
+	: bIsSearchBarHidden(true)
+	, TopLevelCommandList(MakeShared<FUICommandList>())
+	, TextCommandList(MakeShared<FUICommandList>())
+{
+}
+
+SPCGCodeEditorTextBox::~SPCGCodeEditorTextBox()
+{
+}
+
+void SPCGCodeEditorTextBox::Construct(const FArguments& InArgs)
+{
+	RegisterCommands();
+	
+	OnTextChangesAppliedCallback = InArgs._OnTextChangesApplied;
+
+	const TSharedPtr<SScrollBar> HScrollBar = SNew(SScrollBar)
+		.Orientation(EOrientation::Orient_Horizontal);
+	
+	const TSharedPtr<SScrollBar> VScrollBar = SNew(SScrollBar)
+		.Orientation(EOrientation::Orient_Vertical);
+	
+	const FTextBlockStyle& TextStyle = FAppStyle::Get().GetWidgetStyle<FTextBlockStyle>("MessageLog");
+	const FSlateFontInfo &Font = TextStyle.Font;
+
+	const bool bReadOnly = InArgs._IsReadOnly.Get();
+	const FPCGCodeEditorTextBoxCommands& Commands = FPCGCodeEditorTextBoxCommands::Get();
+	
+	Text = SNew(SMultiLineEditableText)
+			.Font(Font)
+			.TextStyle(&TextStyle)
+			.Text(InArgs._Text)
+			.OnTextChanged(InArgs._OnTextChanged)
+			.OnTextCommitted(InArgs._OnTextCommitted)
+			.OnKeyCharHandler(this, &SPCGCodeEditorTextBox::OnTextKeyChar)
+			.OnKeyDownHandler(this, &SPCGCodeEditorTextBox::OnTextKeyDown)
+			// By default, the Tab key gets routed to "next widget". We want to disable that behaviour.
+			.OnIsTypedCharValid_Lambda([](const TCHAR InChar) { return true; })
+			.Marshaller(InArgs._Marshaller)
+			.AutoWrapText(false)
+			.ClearTextSelectionOnFocusLoss(false)
+			.AllowContextMenu(true)
+			.ContextMenuExtender(
+				FMenuExtensionDelegate::CreateLambda(
+					[this, bReadOnly, Commands](FMenuBuilder& InBuilder)
+					{
+						if (!bReadOnly)
+						{
+							InBuilder.PushCommandList(TextCommandList);
+							InBuilder.AddMenuEntry(Commands.ApplyChanges);
+							InBuilder.AddMenuEntry(Commands.ToggleComment);
+							InBuilder.PopCommandList();
+						}
+						
+					}))
+			.IsReadOnly(this, &SPCGCodeEditorTextBox::IsReadOnly)
+			.HScrollBar(HScrollBar)
+			.VScrollBar(VScrollBar);
+
+	SearchBar = SNew(SPCGShaderTextSearchWidget)
+		.OnTextChanged(this, &SPCGCodeEditorTextBox::OnSearchTextChanged)
+		.OnTextCommitted(this, &SPCGCodeEditorTextBox::OnSearchTextCommitted)
+		.SearchResultData(this, &SPCGCodeEditorTextBox::GetSearchResultData)
+		.OnResultNavigationButtonClicked(this, &SPCGCodeEditorTextBox::OnSearchResultNavigationButtonClicked);
+
+	ChildSlot
+	[
+		SAssignNew(TabBody, SVerticalBox)
+		+SVerticalBox::Slot()
+		[
+			SNew(SBorder)
+			.BorderImage(FPCGEditorStyle::Get().GetBrush("TextEditor.Border"))
+			.BorderBackgroundColor(FLinearColor(0.55f, 0.55f, 0.55f, 1.0f))
+			[
+				SNew(SGridPanel)
+				.FillColumn(0,1.0f)
+				.FillRow(0,1.0f)
+				+SGridPanel::Slot(0,0)
+				[
+					Text.ToSharedRef()
+				]
+				+SGridPanel::Slot(1,0)
+				[
+					VScrollBar.ToSharedRef()
+				]
+				+SGridPanel::Slot(0,1)
+				[
+					HScrollBar.ToSharedRef()
+				]
+			]
+		]
+	];
+}
+
+void SPCGCodeEditorTextBox::SetTextProviderObject(UObject* InProviderObject)
+{
+	ShaderTextProviderObject = InProviderObject;
+
+	Refresh();
+}
+
+void SPCGCodeEditorTextBox::RegisterCommands()
+{
+	const FPCGCodeEditorTextBoxCommands& Commands = FPCGCodeEditorTextBoxCommands::Get();
+	
+	TopLevelCommandList->MapAction(Commands.Search, FExecuteAction::CreateSP(this, &SPCGCodeEditorTextBox::OnTriggerSearch));
+	TopLevelCommandList->MapAction(Commands.NextOccurrence, FExecuteAction::CreateSP(this, &SPCGCodeEditorTextBox::OnGoToNextOccurrence));
+	TopLevelCommandList->MapAction(Commands.PreviousOccurrence, FExecuteAction::CreateSP(this, &SPCGCodeEditorTextBox::OnGoToPreviousOccurrence));
+	TextCommandList->MapAction(Commands.ApplyChanges, FExecuteAction::CreateSP(this, &SPCGCodeEditorTextBox::OnApplyChanges));
+	TextCommandList->MapAction(Commands.ToggleComment, FExecuteAction::CreateSP(this, &SPCGCodeEditorTextBox::OnToggleComment));
+}
+
+FReply SPCGCodeEditorTextBox::OnPreviewKeyDown(const FGeometry& MyGeometry, const FKeyEvent& InKeyEvent)
+{
+	const FKey Key = InKeyEvent.GetKey();
+	
+	if (Key == EKeys::Escape)
+	{
+		if (HandleEscape())
+		{
+			return FReply::Handled();
+		}
+	}
+
+	if (TopLevelCommandList->ProcessCommandBindings(InKeyEvent))
+	{
+		return FReply::Handled();
+	}
+	
+	return SCompoundWidget::OnPreviewKeyDown(MyGeometry, InKeyEvent);
+}
+
+bool SPCGCodeEditorTextBox::HandleEscape()
+{
+	if (HideSearchBar())
+	{
+		return true;
+	}
+
+	return false;
+}
+
+void SPCGCodeEditorTextBox::ShowSearchBar()
+{
+	if (bIsSearchBarHidden)
+	{
+		bIsSearchBarHidden = false;
+		
+		TabBody->InsertSlot(0)
+		.AutoHeight()
+		[
+			SearchBar.ToSharedRef()
+		];
+	}
+}
+
+bool SPCGCodeEditorTextBox::HideSearchBar() 
+{
+	if (!bIsSearchBarHidden)
+	{
+		bIsSearchBarHidden = true;
+		SearchBar->ClearSearchText();
+		TabBody->RemoveSlot(SearchBar.ToSharedRef());
+			
+		FSlateApplication::Get().ForEachUser([&](FSlateUser& User) {
+			User.SetFocus(Text.ToSharedRef(), EFocusCause::SetDirectly);
+		});
+
+		return true;
+	}
+	
+	return false;
+}
+
+void SPCGCodeEditorTextBox::OnTriggerSearch()
+{
+	ShowSearchBar();
+
+	FText SelectedText = Text->GetSelectedText();
+	
+	SearchBar->TriggerSearch(SelectedText);
+}
+
+void SPCGCodeEditorTextBox::Refresh() const
+{
+	check(IsInGameThread());
+	Text->Refresh();
+}
+
+void SPCGCodeEditorTextBox::OnSearchTextChanged(const FText& InTextToSearch)
+{
+	// we start the search from the beginning of current selection.
+	// goto clears the selection, but it will be restored by the first search
+	Text->GoTo(Text->GetSelection().GetBeginning());
+
+	Text->SetSearchText(InTextToSearch);
+}
+
+void SPCGCodeEditorTextBox::OnSearchTextCommitted(const FText& InTextToSearch, ETextCommit::Type InCommitType)
+{
+	if (!InTextToSearch.EqualTo(Text->GetSearchText()))
+	{
+		Text->SetSearchText(InTextToSearch);
+	}
+	else
+	{
+		if (InCommitType == ETextCommit::Type::OnEnter)
+		{
+			OnSearchResultNavigationButtonClicked(SSearchBox::SearchDirection::Next);
+		}
+	}
+}
+
+TOptional<SSearchBox::FSearchResultData> SPCGCodeEditorTextBox::GetSearchResultData() const
+{
+	FText SearchText = Text->GetSearchText();
+	
+	if (!SearchText.IsEmpty())
+	{
+		SSearchBox::FSearchResultData Result;
+		Result.CurrentSearchResultIndex = Text->GetSearchResultIndex();
+		Result.NumSearchResults = Text->GetNumSearchResults();
+		
+		return Result;
+	}
+		
+	return TOptional<SSearchBox::FSearchResultData>();
+}
+
+void SPCGCodeEditorTextBox::OnSearchResultNavigationButtonClicked(SSearchBox::SearchDirection InDirection)
+{
+	Text->AdvanceSearch(InDirection == SSearchBox::SearchDirection::Previous);
+}
+
+void SPCGCodeEditorTextBox::OnGoToNextOccurrence()
+{
+	OnSearchResultNavigationButtonClicked(SSearchBox::SearchDirection::Next);
+}
+
+void SPCGCodeEditorTextBox::OnGoToPreviousOccurrence()
+{
+	OnSearchResultNavigationButtonClicked(SSearchBox::SearchDirection::Previous);
+}
+
+FReply SPCGCodeEditorTextBox::OnTextKeyDown(
+	const FGeometry& MyGeometry,
+	const FKeyEvent& InKeyEvent) const
+{
+	if (TextCommandList->ProcessCommandBindings(InKeyEvent))
+	{
+		return FReply::Handled();
+	}
+
+	// Let SMultiLineEditableText::OnKeyDown handle it.
+	return FReply::Unhandled();
+}
+
+void SPCGCodeEditorTextBox::OnApplyChanges() const
+{
+	OnTextChangesAppliedCallback.ExecuteIfBound();
+}
+
+void SPCGCodeEditorTextBox::OnToggleComment() const
+{
+	SMultiLineEditableText::FScopedEditableTextTransaction Transaction(Text);
+	
+	const FTextLocation CursorLocation = Text->GetCursorLocation();
+	
+	const FTextSelection Selection = Text->GetSelection();
+	const FTextLocation SelectionStart =
+		CursorLocation == Selection.GetBeginning() ? Selection.GetEnd() : Selection.GetBeginning();
+	
+	// Need to shift the selection according to the new indentation
+	FTextLocation NewCursorLocation = CursorLocation;
+	FTextLocation NewSelectionStart = SelectionStart;
+		
+	const int32 StartLine = Selection.GetBeginning().GetLineIndex();
+	const int32 EndLine = Selection.GetEnd().GetLineIndex();
+
+	bool bShouldComment = false;
+	bool bAreAllLinesEmpty = true;
+	int32 MinNumLeadingSpaces = INDEX_NONE;
+	for (int32 Index = StartLine; Index <= EndLine; Index++)
+	{
+		FString Line;
+		Text->GetTextLine(Index, Line);
+		// Empty lines are not considered when deciding whether to comment/uncomment
+		const int32 NumSpaces = GetNumSpacesAtStartOfLine(Line);
+		if (NumSpaces >= Line.Len())	
+		{
+			continue;
+		}
+
+		bAreAllLinesEmpty = false;
+		
+		FStringView LineView(&Line[NumSpaces]);
+		if (!LineView.StartsWith(TEXT("//")))
+		{
+			bShouldComment = true;
+		}
+
+		MinNumLeadingSpaces =
+			MinNumLeadingSpaces == INDEX_NONE ?
+			NumSpaces : FMath::Min(MinNumLeadingSpaces, NumSpaces);
+	}
+
+	// Find nearest tab
+	MinNumLeadingSpaces = MinNumLeadingSpaces / 4 * 4;
+
+	// Single line comment or single line no-op should move the cursor down
+	const bool bShouldGoToNewLine = !Text->AnyTextSelected() && (bShouldComment || bAreAllLinesEmpty);
+	
+	for (int32 Index = StartLine; Index <= EndLine; Index++)
+	{
+		FString Line;
+		Text->GetTextLine(Index, Line);
+
+		if (!Line.IsEmpty())
+		{
+			const int32 NumSpaces = GetNumSpacesAtStartOfLine(Line);
+			if (NumSpaces < Line.Len())
+			{
+				int32 LineShift = 0;
+				if (bShouldComment)
+				{
+					Text->GoTo(FTextLocation(Index, MinNumLeadingSpaces));
+					Text->InsertTextAtCursor(TEXT("// "));
+					LineShift = 3;
+				}
+				else
+				{
+					int32 CommentOffset = Line.Find("//");
+					if (ensure(CommentOffset != INDEX_NONE))
+					{
+						int32 ContentOffset =
+							Line.IsValidIndex(CommentOffset+2) && Line[CommentOffset+2] == TEXT(' ') ?
+								3 : 2;
+						Text->SelectText({Index, CommentOffset}, {Index, CommentOffset+ContentOffset});
+						Text->DeleteSelectedText();
+						LineShift = -ContentOffset;
+					}
+				}
+
+				if (Index == CursorLocation.GetLineIndex())
+				{
+					NewCursorLocation = FTextLocation(CursorLocation, LineShift);
+				}
+					
+				if (Index == SelectionStart.GetLineIndex())
+				{
+					NewSelectionStart = FTextLocation(SelectionStart, LineShift);
+				}	
+			}
+		}
+	}
+
+	if (bShouldGoToNewLine)
+	{
+		// Automatically go to the next line if there is no selection
+		int32 LineIndex = CursorLocation.GetLineIndex() + 1;
+		if (LineIndex < const_cast<const SMultiLineEditableText*>(Text.Get())->GetTextLineCount())
+		{
+			FString Line;
+			Text->GetTextLine(LineIndex, Line);
+			
+			int32 Offset = FMath::Min(CursorLocation.GetOffset(), Line.Len());
+			
+			Text->GoTo({LineIndex, Offset});
+		}
+		else
+		{
+			Text->GoTo(NewCursorLocation);
+		}
+	}
+	else
+	{
+		Text->SelectText(NewSelectionStart, NewCursorLocation);
+	}
+}
+
+FReply SPCGCodeEditorTextBox::OnTextKeyChar(
+	const FGeometry& MyGeometry,
+	const FCharacterEvent& InCharacterEvent) const
+{
+	if (Text->IsTextReadOnly())
+	{
+		return FReply::Unhandled();
+	}
+
+	const TCHAR Character = InCharacterEvent.GetCharacter();
+
+	if (Character == TEXT('\b'))
+	{
+		if (!Text->AnyTextSelected())
+		{
+			// if we are deleting a single open brace
+			// look for a matching close brace following it and delete the close brace as well
+			const FTextLocation CursorLocation = Text->GetCursorLocation();
+			const int Offset = CursorLocation.GetOffset();
+			FString Line;
+			Text->GetTextLine(CursorLocation.GetLineIndex(), Line);
+
+			if (Line.IsValidIndex(Offset) && Line.IsValidIndex(Offset-1))
+			{
+				if (IsOpenBrace(Line[Offset-1]))
+				{
+					if (Line[Offset] == GetMatchedCloseBrace(Line[Offset-1]))
+					{
+						SMultiLineEditableText::FScopedEditableTextTransaction ScopedTransaction(Text);
+						Text->SelectText(FTextLocation(CursorLocation, -1),FTextLocation(CursorLocation, 1));
+						Text->DeleteSelectedText();
+						Text->GoTo(FTextLocation(CursorLocation, -1));
+						return FReply::Handled();
+					}
+				}
+			}
+		}
+
+		return FReply::Unhandled();
+	}
+	else if (Character == TEXT('\t'))
+	{
+		SMultiLineEditableText::FScopedEditableTextTransaction Transaction(Text);
+
+		const FTextLocation CursorLocation = Text->GetCursorLocation();
+
+		const bool bShouldIncreaseIndentation = InCharacterEvent.GetModifierKeys().IsShiftDown() ? false : true;
+		
+		// When there is no text selected, shift tab should also decrease line indentation
+		const bool bShouldIndentLine = Text->AnyTextSelected() || (!Text->AnyTextSelected() && !bShouldIncreaseIndentation);
+		
+		if (bShouldIndentLine)
+		{
+			// Indent the whole line if there is a text selection
+			const FTextSelection Selection = Text->GetSelection();
+			const FTextLocation SelectionStart =
+				CursorLocation == Selection.GetBeginning() ? Selection.GetEnd() : Selection.GetBeginning();
+
+			// Shift the selection according to the new indentation
+			FTextLocation NewCursorLocation;
+			FTextLocation NewSelectionStart;
+		
+			const int32 StartLine = Selection.GetBeginning().GetLineIndex();
+			const int32 EndLine = Selection.GetEnd().GetLineIndex();
+		
+			for (int32 Index = StartLine; Index <= EndLine; Index++)
+			{
+				const FTextLocation LineStart(Index, 0);
+				Text->GoTo(LineStart);
+				
+				FString Line;
+				Text->GetTextLine(Index, Line);
+				const int32 NumSpaces = GetNumSpacesAtStartOfLine(Line);
+				const int32 NumExtraSpaces = NumSpaces % 4;
+
+				// Tab to nearest 4.
+				int32 NumSpacesForIndentation;
+				if (bShouldIncreaseIndentation)
+				{
+					NumSpacesForIndentation = NumExtraSpaces == 0 ? 4 : 4 - NumExtraSpaces ;
+					Text->InsertTextAtCursor(FString::ChrN(NumSpacesForIndentation, TEXT(' ')));
+				}
+				else
+				{
+					NumSpacesForIndentation = NumExtraSpaces == 0 ? FMath::Min(4, NumSpaces) : NumExtraSpaces;
+					Text->SelectText(LineStart,FTextLocation(LineStart, NumSpacesForIndentation));
+					Text->DeleteSelectedText();
+				}
+
+				const int32 CursorShiftDirection = bShouldIncreaseIndentation ? 1 : -1;
+				const int32 CursorShift = NumSpacesForIndentation * CursorShiftDirection;
+				
+				if (Index == CursorLocation.GetLineIndex())
+				{
+					NewCursorLocation = FTextLocation(CursorLocation, CursorShift);
+				}
+				
+				if (Index == SelectionStart.GetLineIndex())
+				{
+					NewSelectionStart = FTextLocation(SelectionStart, CursorShift);
+				}
+			}
+			
+			Text->SelectText(NewSelectionStart, NewCursorLocation);
+		}
+		else
+		{
+			FString Line;
+			Text->GetCurrentTextLine(Line);
+
+			const int32 Offset = CursorLocation.GetOffset();
+
+			// Tab to nearest 4.
+			if (ensure(bShouldIncreaseIndentation))
+			{
+				const int32 NumSpacesForIndentation = 4 - Offset % 4;
+				Text->InsertTextAtCursor(FString::ChrN(NumSpacesForIndentation, TEXT(' ')));
+			}
+		}
+
+		return FReply::Handled();
+	}
+	else if (Character == TEXT('\n') || Character == TEXT('\r'))
+	{
+		SMultiLineEditableText::FScopedEditableTextTransaction Transaction(Text);
+		
+		// at this point, the text after the text cursor is already in a new line
+		HandleAutoIndent();
+		
+		return FReply::Handled();
+	}
+	else if (IsOpenBrace(Character))
+	{
+		const TCHAR CloseBrace = GetMatchedCloseBrace(Character);
+		const FTextLocation CursorLocation = Text->GetCursorLocation();
+		FString Line;
+		Text->GetCurrentTextLine(Line);
+		
+		bool bShouldAutoInsertBraces = false;
+
+		if (CursorLocation.GetOffset() < Line.Len())
+		{
+			const TCHAR NextChar = Text->GetCharacterAt(CursorLocation);
+
+			if (IsWhiteSpace(NextChar))
+			{
+				bShouldAutoInsertBraces = true;
+			}
+			else if (IsCloseBrace(NextChar))
+			{
+				int32 BraceBalancePrior = 0;
+				for (int32 Index = 0; Index < CursorLocation.GetOffset() && Index < Line.Len(); Index++)
+				{
+					BraceBalancePrior += (Line[Index] == Character);
+					BraceBalancePrior -= (Line[Index] == CloseBrace);
+				}
+				
+				int32 BraceBalanceLater = 0;
+				for (int32 Index = CursorLocation.GetOffset(); Index < Line.Len(); Index++)
+				{
+					BraceBalanceLater += (Line[Index] == Character);
+					BraceBalanceLater -= (Line[Index] == CloseBrace);
+				}
+
+				if (BraceBalancePrior >= -BraceBalanceLater)
+				{
+					bShouldAutoInsertBraces = true;
+				}
+			}
+		}
+		else
+		{
+			bShouldAutoInsertBraces = true;
+		}
+		
+		// auto insert if we have more open braces
+		// on the left side than close braces on the right side
+		if (bShouldAutoInsertBraces)
+		{
+			// auto insert the matched close brace
+			SMultiLineEditableText::FScopedEditableTextTransaction Transaction(Text);
+			Text->InsertTextAtCursor(FString::Chr(Character));
+			Text->InsertTextAtCursor(FString::Chr(CloseBrace));
+			const FTextLocation NewCursorLocation(Text->GetCursorLocation(), -1);
+			Text->GoTo(NewCursorLocation);
+			return FReply::Handled();
+		}
+
+		return FReply::Unhandled();
+	}
+	else if (IsCloseBrace(Character))
+	{
+		if (!Text->AnyTextSelected())
+		{
+			const FTextLocation CursorLocation = Text->GetCursorLocation();	
+			FString Line;
+			Text->GetTextLine(CursorLocation.GetLineIndex(), Line);
+			
+			const int32 Offset =CursorLocation.GetOffset();
+			
+			if (Line.IsValidIndex(Offset))
+			{
+				if (Line[Offset] == Character)
+				{
+					// avoid creating a duplicated close brace and simply
+					// advance the cursor
+					Text->GoTo(FTextLocation(CursorLocation, 1));
+					return FReply::Handled();
+				}
+			}
+		}
+		
+		return FReply::Unhandled();
+	}
+	else
+	{
+		// Let SMultiLineEditableText::OnKeyChar handle it.
+		return FReply::Unhandled();
+	}
+}
+
+void SPCGCodeEditorTextBox::HandleAutoIndent() const
+{
+	const FTextLocation CursorLocation = Text->GetCursorLocation();
+	const int32 CurLineIndex = CursorLocation.GetLineIndex();
+	const int32 LastLineIndex = CurLineIndex - 1;
+
+	if (LastLineIndex > 0)
+	{
+		FString LastLine;
+		Text->GetTextLine(LastLineIndex, LastLine);
+
+		const int32 NumSpaces = GetNumSpacesAtStartOfLine(LastLine);
+
+		const int32 NumSpacesForCurrentIndentation = NumSpaces/4*4;
+		const FString CurrentIndentation = FString::ChrN(NumSpacesForCurrentIndentation, TEXT(' '));
+		const FString NextIndentation = FString::ChrN(NumSpacesForCurrentIndentation + 4, TEXT(' '));
+	
+		// See what the open/close curly brace balance is.
+		int32 BraceBalance = 0;
+		for (const TCHAR Char : LastLine)
+		{
+			BraceBalance += (Char == TEXT('{'));
+			BraceBalance -= (Char == TEXT('}'));
+		}
+
+		if (BraceBalance <= 0)
+		{
+			Text->InsertTextAtCursor(CurrentIndentation);
+		}
+		else
+		{
+			Text->InsertTextAtCursor(NextIndentation);
+		
+			// Look for an extra close curly brace and auto-indent it as well
+			FString CurLine;
+			Text->GetTextLine(CurLineIndex, CurLine);
+
+			BraceBalance = 0;
+			int32 CloseBraceOffset = 0;
+			for (const TCHAR Char : CurLine)
+			{
+				BraceBalance += (Char == TEXT('{'));
+				BraceBalance -= (Char == TEXT('}'));
+
+				// Found the first extra '}'
+				if (BraceBalance < 0)
+				{
+					break;
+				}
+
+				CloseBraceOffset++;
+			}
+
+			if (BraceBalance < 0)
+			{
+				const FTextLocation SavedCursorLocation = Text->GetCursorLocation();
+				const FTextLocation CloseBraceLocation(CurLineIndex, CloseBraceOffset);
+			
+				// Create a new line and apply indentation for the close curly brace
+				FString NewLineAndIndent(TEXT("\n"));
+				NewLineAndIndent.Append(CurrentIndentation);
+			
+				Text->GoTo(CloseBraceLocation);
+				Text->InsertTextAtCursor(NewLineAndIndent);
+				// Recover cursor location
+				Text->GoTo(SavedCursorLocation);
+			}
+		}
+	}
+}
+
+bool SPCGCodeEditorTextBox::IsReadOnly() const
+{
+	if (IPCGCodeEditorTextProvider* Provider = Cast<IPCGCodeEditorTextProvider>(ShaderTextProviderObject))
+	{
+		return Provider->IsReadOnly();
+	}
+
+	return true;
+}
+
+#undef LOCTEXT_NAMESPACE
diff --git a/Engine/Plugins/PCG/Source/PCGEditor/Private/Widgets/SPCGCodeEditorTextBox.h b/Engine/Plugins/PCG/Source/PCGEditor/Private/Widgets/SPCGCodeEditorTextBox.h
new file mode 100644
--- /dev/null
+++ b/Engine/Plugins/PCG/Source/PCGEditor/Private/Widgets/SPCGCodeEditorTextBox.h
@@ -0,0 +1,123 @@
+// Copyright Epic Games, Inc. All Rights Reserved.
+
+#pragma once
+
+#include "Widgets/SCompoundWidget.h"
+#include "Framework/Commands/Commands.h"
+
+#include "CoreMinimal.h"
+#include "Framework/SlateDelegates.h"
+#include "Widgets/Input/SSearchBox.h"
+
+class ITextLayoutMarshaller;
+class SMultiLineEditableText;
+class SPCGShaderTextSearchWidget;
+class SVerticalBox;
+
+class FPCGCodeEditorTextBoxCommands :
+	public TCommands<FPCGCodeEditorTextBoxCommands>
+{
+public:
+	FPCGCodeEditorTextBoxCommands();
+	
+	// TCommands<> overrides
+	virtual void RegisterCommands() override;
+	
+	TSharedPtr<FUICommandInfo> Search;
+	TSharedPtr<FUICommandInfo> NextOccurrence;
+	TSharedPtr<FUICommandInfo> PreviousOccurrence;
+	
+	TSharedPtr<FUICommandInfo> ApplyChanges;
+	TSharedPtr<FUICommandInfo> ToggleComment;
+};
+
+/** Text box for editing code. */
+class SPCGCodeEditorTextBox : public SCompoundWidget
+{
+public:
+	SPCGCodeEditorTextBox();
+	virtual ~SPCGCodeEditorTextBox() override;
+
+	DECLARE_DELEGATE(FOnTextChangesApplied);
+
+	SLATE_BEGIN_ARGS(SPCGCodeEditorTextBox) {};
+		/** The initial text that will appear in the widget. */
+		SLATE_ATTRIBUTE(FText, Text)
+
+		/** Text to search for (a new search is triggered whenever this text changes) */
+		SLATE_ATTRIBUTE(FText, SearchText)
+
+		/** The marshaller used to get/set the raw text to/from the text layout. */
+		SLATE_ARGUMENT(TSharedPtr<ITextLayoutMarshaller>, Marshaller)
+
+		/** Sets whether this text box can actually be modified interactively by the user */
+		SLATE_ATTRIBUTE(bool, IsReadOnly)
+
+		/** Called whenever the text is changed programmatically or interactively by the user */
+		SLATE_EVENT(FOnTextChanged, OnTextChanged)
+
+		/** Called whenever the text is committed. This happens when the text box loses focus. */
+		SLATE_EVENT(FOnTextCommitted, OnTextCommitted)
+
+		/** Called when user applies the changes */
+		SLATE_EVENT(FOnTextChangesApplied, OnTextChangesApplied)
+	SLATE_END_ARGS()
+
+	void Construct(const FArguments& InArgs);
+
+	void Refresh() const;
+
+	void SetTextProviderObject(UObject* InProviderObject);
+
+private:
+	//~ Begin SWidget interface
+	// Use Preview Key Down so that the commands related to this tab is given the priority
+	virtual FReply OnPreviewKeyDown(const FGeometry& MyGeometry, const FKeyEvent& InKeyEvent) override;
+	//~ End SWidget interface
+
+	void RegisterCommands();
+	
+	bool HandleEscape();
+	
+	void ShowSearchBar();
+	bool HideSearchBar();
+	
+	void OnTriggerSearch();
+	void OnSearchTextChanged(const FText& InTextToSearch);
+	void OnSearchTextCommitted(const FText& InTextToSearch, ETextCommit::Type InCommitType);
+	TOptional<SSearchBox::FSearchResultData> GetSearchResultData() const;
+	void OnSearchResultNavigationButtonClicked(SSearchBox::SearchDirection InDirection);
+	void OnGoToNextOccurrence();
+	void OnGoToPreviousOccurrence();
+
+	FReply OnTextKeyDown(const FGeometry& MyGeometry, const FKeyEvent& InCharacterEvent) const;
+
+	void OnApplyChanges() const;
+
+	void OnToggleComment() const;
+	
+	FReply OnTextKeyChar(const FGeometry& MyGeometry, const FCharacterEvent& InCharacterEvent) const;
+	
+	// Figure out if we need to auto-indent.
+	void HandleAutoIndent() const;
+
+	bool IsReadOnly() const;
+
+private:
+	TSharedPtr<SVerticalBox> TabBody;
+	
+	TSharedPtr<SMultiLineEditableText> Text;
+
+	TSharedPtr<SPCGShaderTextSearchWidget> SearchBar;
+	bool bIsSearchBarHidden = true;
+
+	// Top level command list is processed during preview key down
+	// such that its commands are given the priority
+	TSharedRef<FUICommandList> TopLevelCommandList;
+	
+	TSharedRef<FUICommandList> TextCommandList;
+	
+	FOnTextChangesApplied OnTextChangesAppliedCallback;
+
+	TWeakObjectPtr<UObject> ShaderTextProviderObject;
+};
diff --git a/Engine/Plugins/PCG/Source/PCGEditor/Private/Widgets/SPCGEditorNodeSource.cpp b/Engine/Plugins/PCG/Source/PCGEditor/Private/Widgets/SPCGEditorNodeSource.cpp
deleted file mode 100644
diff --git a/Engine/Plugins/PCG/Source/PCGEditor/Private/Widgets/SPCGEditorNodeSource.h b/Engine/Plugins/PCG/Source/PCGEditor/Private/Widgets/SPCGEditorNodeSource.h
deleted file mode 100644
diff --git a/Engine/Plugins/PCG/Source/PCGEditor/Private/Widgets/SPCGNodeSourceTextBox.cpp b/Engine/Plugins/PCG/Source/PCGEditor/Private/Widgets/SPCGNodeSourceTextBox.cpp
deleted file mode 100644
diff --git a/Engine/Plugins/PCG/Source/PCGEditor/Private/Widgets/SPCGNodeSourceTextBox.h b/Engine/Plugins/PCG/Source/PCGEditor/Private/Widgets/SPCGNodeSourceTextBox.h
deleted file mode 100644
diff --git a/Engine/Plugins/PCG/Source/PCGEditor/Public/PCGEditor.h b/Engine/Plugins/PCG/Source/PCGEditor/Public/PCGEditor.h
--- a/Engine/Plugins/PCG/Source/PCGEditor/Public/PCGEditor.h
+++ b/Engine/Plugins/PCG/Source/PCGEditor/Public/PCGEditor.h
@@ -31,7 +31,7 @@ class SPCGEditorGraphEmbeddedSubgraphsView;
 class SPCGEditorGraphLogView;
 class SPCGEditorGraphNodePalette;
 class SPCGEditorGraphProfilingView;
-class SPCGEditorNodeSource;
+class SPCGCodeEditor;
 class SPCGEditorGraphUserParametersView;
 class SPCGEditorViewport;
 class UEdGraphNode;
@@ -103,7 +103,6 @@ namespace PCGEditorTabs
 	constexpr FLazyName DeterminismID = TEXT("Determinism");
 	constexpr FLazyName ProfilingID = TEXT("Profiling");
 	constexpr FLazyName LogID = TEXT("Log");
-	constexpr FLazyName HLSLSourceID = TEXT("HLSLSource");
 	constexpr FLazyName UserParamsID = TEXT("UserParams");
 	constexpr FLazyName ViewportID[] = {
 		TEXT("Viewport"),
@@ -111,10 +110,14 @@ namespace PCGEditorTabs
 		TEXT("Viewport3"),
 		TEXT("Viewport4") };
 	constexpr FLazyName EmbeddedSubgraphsID = TEXT("EmbeddedSubgraphs");
+	constexpr FLazyName CodeEditorID = TEXT("CodeEditor");
 
 	constexpr int32 NumPropertyDetailTabs = sizeof(PropertyDetailsID) / sizeof(FLazyName);
 	constexpr int32 NumAttributeTabs = sizeof(AttributesID) / sizeof(FLazyName);
 	constexpr int32 NumViewportTabs = sizeof(ViewportID) / sizeof(FLazyName);
+
+	UE_DEPRECATED(5.8, "Use CodeEditorID instead.")
+	constexpr FLazyName HLSLSourceID = TEXT("HLSLSource");
 }
 
 class FPCGEditor : public FWorkflowCentricApplication, public FGCObject, public FSelfRegisteringEditorUndoClient
@@ -268,12 +271,15 @@ PRAGMA_ENABLE_DEPRECATION_WARNINGS
 	PCGEDITOR_API TSharedPtr<FWorkflowTabFactory> CreateDeterminismTabFactory();
 	PCGEDITOR_API TSharedPtr<FWorkflowTabFactory> CreateProfilingTabFactory();
 	PCGEDITOR_API TSharedPtr<FWorkflowTabFactory> CreateLogTabFactory();
-	PCGEDITOR_API TSharedPtr<FWorkflowTabFactory> CreateNodeSourceTabFactory();
+	PCGEDITOR_API TSharedPtr<FWorkflowTabFactory> CreateCodeEditorTabFactory();
 	PCGEDITOR_API TSharedPtr<FWorkflowTabFactory> CreateGraphParamsTabFactory();
 	PCGEDITOR_API TSharedPtr<FWorkflowTabFactory> CreatePropertyDetailsTabFactory(int32 Index, TSharedRef<FWorkspaceItem> Group);
 	PCGEDITOR_API TSharedPtr<FWorkflowTabFactory> CreateViewportTabFactory(int32 Index, TSharedRef<FWorkspaceItem> Group);
 	PCGEDITOR_API TSharedPtr<FWorkflowTabFactory> CreateEmbeddedSubgraphsTabFactory();
 
+	UE_DEPRECATED(5.8, "Use CreateCodeEditorTabFactory() instead")
+	PCGEDITOR_API TSharedPtr<FWorkflowTabFactory> CreateNodeSourceTabFactory();
+
 	/** Called when graph editor focus is changed */
 	PCGEDITOR_API void OnGraphEditorFocused(const TSharedRef<SGraphEditor>& InGraphEditor);
 
@@ -539,8 +545,8 @@ private:
 	/** Get or Create a new log capture tab widget */
 	TSharedRef<SPCGEditorGraphLogView> GetOrCreateLogWidget();
 
-	/** Get or Create a new node source editor tab widget */
-	TSharedRef<SPCGEditorNodeSource> GetOrCreateNodeSourceWidget();
+	/** Get or Create a new code editor tab widget */
+	TSharedRef<SPCGCodeEditor> GetOrCreateCodeEditorWidget();
 
 	/** Get or Create a new user graph parameters tab widget */
 	TSharedRef<SPCGEditorGraphUserParametersView> GetOrCreateGraphParamsWidget();
@@ -607,7 +613,7 @@ private:
 	TSharedPtr<SPCGEditorGraphDeterminismListView> DeterminismWidget;
 	TSharedPtr<SPCGEditorGraphProfilingView> ProfilingWidget;
 	TSharedPtr<SPCGEditorGraphLogView> LogWidget;
-	TSharedPtr<SPCGEditorNodeSource> NodeSourceWidget;
+	TSharedPtr<SPCGCodeEditor> CodeEditorWidget;
 	TSharedPtr<SPCGEditorGraphUserParametersView> UserParamsWidget;
 	TSharedPtr<SPCGEditorViewport> ViewportWidgets[PCGEditorTabs::NumViewportTabs];
 	TSharedPtr<SPCGEditorGraphEmbeddedSubgraphsView> EmbeddedSubgraphsWidget;
