[PCG] Fix ISM Reuse issues : missing or extra instances depending on situations
- ISM Reuse code was blindly reusing without validating the expected number of resources to be reused
- We now keep track of the number of ISM resources we need to find in order to consider reusing
- Applies to global reuse (when all dependencies crc match) or per input reuse
- This means that any resource that is shared between spawner execution can't reuse instance, only the component itself (MarkAsUsed)
- Reset the different Crcs in MarkAsUsed

#Synced-CL 50905358

diff --git a/Engine/Plugins/PCG/Source/PCG/Private/Elements/PCGStaticMeshSpawner.cpp b/Engine/Plugins/PCG/Source/PCG/Private/Elements/PCGStaticMeshSpawner.cpp
--- a/Engine/Plugins/PCG/Source/PCG/Private/Elements/PCGStaticMeshSpawner.cpp
+++ b/Engine/Plugins/PCG/Source/PCG/Private/Elements/PCGStaticMeshSpawner.cpp
@@ -33,6 +33,76 @@
 
 namespace PCGStaticMeshSpawner
 {
+	/**
+	 * ISM Reuse (all this applies if CVarAllowISMReuse is enabled)
+	 * 
+	 * There are 3 ways a SMS (Static Mesh Spawner) can do actual reuse:
+	 * 1. Full Reuse: Dependencies Crc matches the last execution (all inputs, settings, etc) meaning we can reuse all our resources and skip most processing except for generating an output when needed.
+	 * 2. Partial Reuse: 1-N input Crc matches the last execution in which case the resources from those inputs can be reused and the others generated.
+	 * 3. Component Reuse: No inputs match but the settings/descriptor crc on the components match and so we can reuse the resource but its instances will be reset and rebuild.
+	 * 
+	 * Also of note: bAllowMergeDifferentDataInSameInstancedComponents is a setting that will allow different SMS to output instances to the same components if those components share the same settings/descriptor crc.
+	 * 
+	 * When this 'bAllowMergeDifferentDataInSameInstancedComponents' is enabled and 2 or more SMS do indeed share some resources then actual reusing of Component+Instances becomes impossible as we do not support cleaning up invalid instances.
+	 * Ex: SMS1 outputs instances to Component1 and SMS2 also outputs instances to Component1, on the next execution SMS1 inputs changed and doesn't need to output to Component1 anymore but SMS2 still does. If we were to reuse Component1 as is
+	 * then instances from the previous executon of SMS1 would still be present in the result which is invalid.
+	 * 
+	 * So when a resource is used for more than one SMS then we only allow option 3 which reuses the component but requires that all SMS re-add their instances. This has a negative impact on performance but merging allows a great reduction on
+	 * number of generated resources for large scale generations. We could eventually support "cleaning" up instances from previous generations to get the best of both worlds.
+	 * 
+	 * How do we make sure that 1 and 2 can be achieved?
+	 * 
+	 * 1. Full Reuse uses
+	 *		- DependenciesCrc
+	 *		- FPCGStaticMeshSpawnerContext::UsedResources 
+	 *		- UPCGManagedISMComponent::CrcReuseCount
+	 *		- UPCGManagedISMComponent::Crc
+	 * 
+	 *		In FPCGStaticMeshSpawnerElement::PrepareDataInternal:
+	 *		- Compute 'DependenciesCrc'.
+	 *		- Iterate through all existing resources and accumulate all resources that have the same Crc.
+	 *		- If the accumulated resource count is equal to each of the individual resource CrcReuseCount it means we can do a full reuse,
+	 *			if the reuse count doesn't match it means one of our resources might have been used by another SMS making full reuse impossible.
+	 *      - Once we know we can do full reuse, we also push those resources into: UsedResources so we can update their CrcReuseCount at the end of execution and we also push
+	 *			those resources into the DataCrcToUsedResources (per input) since we can conclude that if we can do full reuse than our resoures would also be valid for partial reuse.
+	 *			The reason why we do this is that if the next execution makes full reuse impossible, we can still fallback on partial reuse as our resource will have been updated properly.
+	 * 
+	 * 2. Partial Reuse uses
+	 *		- DataCrc (per input Crc)
+	 *		- FPCGStaticMeshSpawnerContext::DataCrcToUsedResources
+	 *		- UPCGManagedISMComponent::DataCrcReuseCount
+	 *		- UPCGManagedISMComponent::DataCrc
+	 * 
+	 *		Still in FPCGStaticMeshSpawnerElement::PrepareDataInternal: (if we failed doing full reuse)
+	 *		- Iterate through inputs computing a per input DataCrc
+	 *		- On each input iterate all existing resources and accumulate all resources that have the same DataCrc (and settings crc)
+	 *		- If the accumulated resource count is equal to each of the individual resource DataCrcReuseCount it means we can do reuse for this input,
+	 *			if the reuse count doesn't match it means one of our resources might have been used by another SMS making partial reuse impossible.
+	 *		- Once we know we can do partial reuse, we also push those resources into: DataCrcToUsedResources so we can update their DataCrcReuseCount at the end of execution and we also push
+	 *			those resources into the UsedResources (full reuse) to update their CrcReuseCount and set the DependenciesCrc to make full reuse possible on next execution if nothing changes.
+	 * 
+	 *	This sums up the code that handles the reuse. If we are not reusing then it is important to still compute valid CrcReuseCount/DataCrcReuseCount/Crc/DataCrc for each execution of a SMS so that reuse can happen
+	 *  on subsequent executions.
+	 * 
+	 *  In FPCGStaticMeshSpawnerElement::SpawnStaticMeshInstances is where we do this:
+	 *		- This methods is responsible for setting the Crc (DependenciesCrc) on a resource.
+	 *		- If that resource already has a Crc that differs from our current DependenciesCrc it means that another SMS used this resource in which case we set UsedResources.bMergedResources = true.
+	 *		- We also update that resource's Crc by combining both Crcs so that no matches can happen anymore on this resource.
+	 *		- That flag will be read later one to impact the CrcReuseCount of that resource.
+	 *		- We also add that resource to the UsedResource so that in the end we have a list of all the resources needed for a full reuse.
+	 *		- This method also sets the DataCrc (per input) and adds that resource to the DataCrcToUsedResources map using the DataCrc as a key.
+	 *		- If that resource already has a DataCrc it means that another SMS used this resource in which case we set the bMergedResources to true for its previous DataCrc and the new DataCrc.
+	 *		- This is subtil but it allows flagging other resources assigned to the same input at the end of execution so that they also can't be reused.
+	 *		- We also add that resource to the DataCrcToUsedResources using the new DataCrc as the key.
+	 * 
+	 *	At this point either through FPCGStaticMeshSpawnerElement::PrepareDataInternal or FPCGStaticMeshSpawnerElement::SpawnStaticMeshInstances, we've gathered and updated resources for reuse.
+	 *  
+	 *	The last step is at the end of execution in FPCGStaticMeshSpawnerElement::ExecuteInternal where we are going to go through both UsedResources and DataCrcToUsedResources and based on the bMergedResources flag for each
+	 *  set of resources, are going to set the CrcReuseCount or DataCrcReuseCount to INDEX_NONE if that set has bMergedResources = true or to the actual number of resources if bMergedResources = false.
+	 * 
+	 *  In the case 'bAllowMergeDifferentDataInSameInstancedComponents' is false prior to calling UPCGActorHelpers::GetOrCreateManagedISMC we set DataCrc in the params struct so that we don't reuse an existing resource
+	 *  with the same settings crc.
+	 */
 	static TAutoConsoleVariable<bool> CVarAllowISMReuse(
 		TEXT("pcg.ISM.AllowReuse"),
 		true,
@@ -249,31 +319,41 @@ bool FPCGStaticMeshSpawnerElement::PrepareDataInternal(FPCGContext* InContext) c
 		
 		if (Context->DependenciesCrc.IsValid())
 		{
+			int32 CrcReuseCount = 0;
 			TArray<UPCGManagedISMComponent*> MISMCs;
-			SourceComponent->ForEachManagedResource([&MISMCs, &Context, Settings](UPCGManagedResource* InResource)
+			SourceComponent->ForEachManagedResource([&MISMCs, &Context, Settings, &CrcReuseCount](UPCGManagedResource* InResource)
 			{
 				if (UPCGManagedISMComponent* Resource = Cast<UPCGManagedISMComponent>(InResource))
 				{
-					if (Resource->GetCrc().IsValid() && Resource->GetCrc() == Context->DependenciesCrc)
+					if (Resource->GetCrc().IsValid() && Resource->GetCrc() == Context->DependenciesCrc && Resource->GetCrcReuseCount() != INDEX_NONE)
 					{
-						MISMCs.Add(Resource);
+						if (ensure(CrcReuseCount == 0 || CrcReuseCount == Resource->GetCrcReuseCount()))
+						{
+							CrcReuseCount = Resource->GetCrcReuseCount();
+							MISMCs.Add(Resource);
+						}
 					}
 				}
 			});
 
-			for (UPCGManagedISMComponent* MISMC : MISMCs)
+			if (CrcReuseCount > 0 && CrcReuseCount == MISMCs.Num())
 			{
-				if (!MISMC->IsMarkedUnused() && Settings->bWarnOnIdenticalSpawn)
+				for (UPCGManagedISMComponent* MISMC : MISMCs)
 				{
-					// TODO: Revisit if the stack is added to the managed components at creation
-					PCGLog::LogWarningOnGraph(LOCTEXT("IdenticalISMCSpawn", "Identical ISM Component spawn occurred. It may be beneficial to re-check graph logic for identical spawn conditions (same mesh descriptor at same location, etc) or repeated nodes."), Context);
-				}
+					if (!MISMC->IsMarkedUnused() && Settings->bWarnOnIdenticalSpawn)
+					{
+						// TODO: Revisit if the stack is added to the managed components at creation
+						PCGLog::LogWarningOnGraph(LOCTEXT("IdenticalISMCSpawn", "Identical ISM Component spawn occurred. It may be beneficial to re-check graph logic for identical spawn conditions (same mesh descriptor at same location, etc) or repeated nodes."), Context);
+					}
 
-				MISMC->MarkAsReused();
-			}
+					MISMC->MarkAsReused();
+					Context->UsedResources.Resources.Add(MISMC);
 
-			if (!MISMCs.IsEmpty())
-			{
+					if (ensure(MISMC->GetDataCrc().IsValid()))
+					{
+						Context->DataCrcToUsedResources.FindOrAdd(MISMC->GetDataCrc().GetValue()).Resources.Add(MISMC);
+					}
+				}
 				bSkippedDueToReuse = true;
 			}
 		}
@@ -335,7 +415,7 @@ bool FPCGStaticMeshSpawnerElement::PrepareDataInternal(FPCGContext* InContext) c
 			FPCGCrc DataCrc;
 			if(!bSkippedDueToReuse && PCGStaticMeshSpawner::CVarAllowISMReuse.GetValueOnAnyThread())
 			{
-				check(Context->DependenciesCrc.IsValid()); // should have been done earlier
+				ensure(Context->DependenciesCrc.IsValid()); // should have been done earlier
 
 				const FPCGCrc SettingsCrc = Settings->GetSettingsCrc();
 				ensure(SettingsCrc.IsValid());
@@ -354,35 +434,45 @@ bool FPCGStaticMeshSpawnerElement::PrepareDataInternal(FPCGContext* InContext) c
 					GetDependenciesCrc(FPCGGetDependenciesCrcParams(&SubCollection, Settings, Context->ExecutionSource.Get()), DataCrc);
 				}
 
+				int32 DataCrcReuseCount = 0;
 				TArray<UPCGManagedISMComponent*> MISMCs;
-				SourceComponent->ForEachManagedResource([&MISMCs, &SettingsCrc, &DataCrc, &Context](UPCGManagedResource* InResource)
+				SourceComponent->ForEachManagedResource([&MISMCs, &SettingsCrc, &DataCrc, &Context, &DataCrcReuseCount](UPCGManagedResource* InResource)
 				{
 					if (UPCGManagedISMComponent* Resource = Cast<UPCGManagedISMComponent>(InResource))
 					{
+						// All inputs are reused so just make sure to track per data crcs used resources and update data crcs
 						if (Resource->GetSettingsCrc().IsValid() && Resource->GetSettingsCrc() == SettingsCrc &&
-							Resource->GetDataCrc().IsValid() && Resource->GetDataCrc() == DataCrc)
+							Resource->GetDataCrc().IsValid() && Resource->GetDataCrc() == DataCrc &&
+							Resource->GetDataCrcReuseCount() != INDEX_NONE)
 						{
-							MISMCs.Add(Resource);
+							if (ensure(DataCrcReuseCount == 0 || DataCrcReuseCount == Resource->GetDataCrcReuseCount()))
+							{
+								DataCrcReuseCount = Resource->GetDataCrcReuseCount();
+								MISMCs.Add(Resource);
+							}
 						}
 					}
 				});
 
-				for (UPCGManagedISMComponent* MISMC : MISMCs)
+				if (DataCrcReuseCount > 0 && MISMCs.Num() == DataCrcReuseCount)
 				{
-					if (!MISMC->IsMarkedUnused() && Settings->bWarnOnIdenticalSpawn)
+					for (UPCGManagedISMComponent* MISMC : MISMCs)
 					{
-						// TODO: Revisit if the stack is added to the managed components at creation
-						PCGLog::LogWarningOnGraph(LOCTEXT("IdenticalISMCSpawn", "Identical ISM Component spawn occurred. It may be beneficial to re-check graph logic for identical spawn conditions (same mesh descriptor at same location, etc) or repeated nodes."), Context);
-					}
+						if (!MISMC->IsMarkedUnused() && Settings->bWarnOnIdenticalSpawn)
+						{
+							// TODO: Revisit if the stack is added to the managed components at creation
+							PCGLog::LogWarningOnGraph(LOCTEXT("IdenticalISMCSpawn", "Identical ISM Component spawn occurred. It may be beneficial to re-check graph logic for identical spawn conditions (same mesh descriptor at same location, etc) or repeated nodes."), Context);
+						}
 
-					MISMC->MarkAsReused();
+						MISMC->MarkAsReused();
 
-					// Update global crc, otherwise these resources wouldn't get picked up in a subsequent update
-					MISMC->SetCrc(Context->DependenciesCrc);
-				}
+						// Update global crc, otherwise these resources wouldn't get picked up in a subsequent update
+						MISMC->SetCrc(Context->DependenciesCrc);
+
+						Context->DataCrcToUsedResources.FindOrAdd(DataCrc.GetValue()).Resources.Emplace(MISMC);
+						Context->UsedResources.Resources.Add(MISMC);
+					}
 
-				if (!MISMCs.IsEmpty())
-				{
 					Context->bCurrentDataSkippedDueToReuse = true;
 				}
 			}
@@ -575,6 +665,35 @@ bool FPCGStaticMeshSpawnerElement::ExecuteInternal(FPCGContext* InContext) const
 	const bool bFinishedExecution = Context->MeshInstancesData.IsEmpty();
 	if (bFinishedExecution)
 	{
+		{
+			// We finished executing, make sure to add information needed to compare against full reuse
+			const int32 CrcReuseCount = Context->UsedResources.bMergedResources ? INDEX_NONE : Context->UsedResources.Resources.Num();
+
+			for (TWeakObjectPtr<UPCGManagedISMComponent> WeakManagedResource : Context->UsedResources.Resources)
+			{
+				if (UPCGManagedISMComponent* ManagedResource = WeakManagedResource.Get())
+				{
+					ManagedResource->SetCrcReuseCount(CrcReuseCount);
+				}
+			}
+		}
+
+		{
+			// We finished executing, make sure to add information needed to compare against per input reuse
+			for (auto& [DataCrc, UsedResources] : Context->DataCrcToUsedResources)
+			{
+				const int32 DataCrcReuseCount = UsedResources.bMergedResources ? INDEX_NONE : UsedResources.Resources.Num();
+
+				for (TWeakObjectPtr<UPCGManagedISMComponent> WeakManagedResource : UsedResources.Resources)
+				{
+					if (UPCGManagedISMComponent* ManagedResource = WeakManagedResource.Get())
+					{
+						ManagedResource->SetDataCrcReuseCount(DataCrcReuseCount);
+					}
+				}
+			}
+		}
+
 		if (AActor* TargetActor = Settings->TargetActor.Get() ? Settings->TargetActor.Get() : Context->GetTypedExecutionTarget<AActor>())
 		{
 			for (UFunction* Function : PCGHelpers::FindUserFunctions(TargetActor->GetClass(), Settings->PostProcessFunctionNames, { UPCGFunctionPrototypes::GetPrototypeWithNoParams() }, Context))
@@ -698,22 +817,43 @@ void FPCGStaticMeshSpawnerElement::SpawnStaticMeshInstances(FPCGStaticMeshSpawne
 	UPCGManagedISMComponent* MISMC = UPCGActorHelpers::GetOrCreateManagedISMC(TargetActor, SourceComponent, Params, Context);
 
 	check(MISMC);
-	MISMC->SetCrc(Context->DependenciesCrc);
+
+	// When spawning from different inputs, the dependencies crc will be the same on the same spawner but across multiple spawners it will change and needs
+	// to be updated to disable full reuse across multiple spawners.
+	if (MISMC->GetCrc().IsValid() && MISMC->GetCrc() != Context->DependenciesCrc)
+	{
+		FPCGCrc UpdatedCrc = MISMC->GetCrc();
+		UpdatedCrc.Combine(Context->DependenciesCrc);
+		MISMC->SetCrc(UpdatedCrc);
+		Context->UsedResources.bMergedResources = true;
+	}
+	else
+	{
+		MISMC->SetCrc(Context->DependenciesCrc);
+	}
 
 	ensure(Settings->bAllowMergeDifferentDataInSameInstancedComponents || !MISMC->GetDataCrc().IsValid() && MISMC->GetComponent() && MISMC->GetComponent()->GetNumInstances() == 0);
 
 	// Mutate MISMC data crc, so this will ensure that when we write multiple data to the same ISM that we don't trigger reuse.
 	if (MISMC->GetDataCrc().IsValid())
 	{
-		FPCGCrc UpdatedCrc = MISMC->GetDataCrc();
-		UpdatedCrc.Combine(DataCrc);
-		MISMC->SetDataCrc(UpdatedCrc);
+		FPCGCrc CurrentDataCrc = MISMC->GetDataCrc();
+
+		// This means we have different inputs sharing the same ISM in which case we can't do partial reuse (on both inputs, previous crc and current crc)
+		Context->DataCrcToUsedResources.FindOrAdd(CurrentDataCrc.GetValue()).bMergedResources = true;
+		Context->DataCrcToUsedResources.FindOrAdd(DataCrc.GetValue()).bMergedResources = true;
+
+		CurrentDataCrc.Combine(DataCrc);
+		MISMC->SetDataCrc(CurrentDataCrc);
 	}
 	else
 	{
 		MISMC->SetDataCrc(DataCrc);
 	}
 
+	Context->DataCrcToUsedResources.FindOrAdd(DataCrc.GetValue()).Resources.Emplace(MISMC);
+	Context->UsedResources.Resources.Emplace(MISMC);
+
 	// Keep track of all touched resources in the context, because if the execution is cancelled during the SMS execution
 	// we cannot easily guarantee that the state (esp. vs CRCs) is going to be entirely valid
 	Context->TouchedResources.Emplace(MISMC);
diff --git a/Engine/Plugins/PCG/Source/PCG/Private/PCGManagedResource.cpp b/Engine/Plugins/PCG/Source/PCG/Private/PCGManagedResource.cpp
--- a/Engine/Plugins/PCG/Source/PCG/Private/PCGManagedResource.cpp
+++ b/Engine/Plugins/PCG/Source/PCG/Private/PCGManagedResource.cpp
@@ -118,6 +118,13 @@ bool UPCGManagedResource::CanBeUsed() const
 #endif
 }
 
+void UPCGManagedResource::MarkAsUsed()
+{
+	ensure(CanBeUsed()); 
+	bIsMarkedUnused = false;
+	Crc = FPCGCrc();
+}
+
 bool UPCGManagedResource::DebugForcePurgeAllResourcesOnGenerate()
 {
 	return CVarForceReleaseResourcesOnGenerate.GetValueOnAnyThread();
@@ -991,6 +998,10 @@ void UPCGManagedISMComponent::MarkAsUsed()
 		return;
 	}
 
+	DataCrc = FPCGCrc();
+	DataCrcReuseCount = INDEX_NONE;
+	CrcReuseCount = INDEX_NONE;
+
 	if (UInstancedStaticMeshComponent* ISMC = GetComponent())
 	{
 		const bool bHasPreviousRootLocation = bHasRootLocation;
@@ -1024,6 +1035,9 @@ void UPCGManagedISMComponent::MarkAsReused()
 {
 	Super::MarkAsReused();
 
+	CrcReuseCount = INDEX_NONE;
+	DataCrcReuseCount = INDEX_NONE;
+
 	if (UInstancedStaticMeshComponent* ISMC = GetComponent())
 	{
 		// Reset the rotation/scale to be identity otherwise if the root component transform has changed, the final transform will be wrong
diff --git a/Engine/Plugins/PCG/Source/PCG/Public/Elements/PCGStaticMeshSpawnerContext.h b/Engine/Plugins/PCG/Source/PCG/Public/Elements/PCGStaticMeshSpawnerContext.h
--- a/Engine/Plugins/PCG/Source/PCG/Public/Elements/PCGStaticMeshSpawnerContext.h
+++ b/Engine/Plugins/PCG/Source/PCG/Public/Elements/PCGStaticMeshSpawnerContext.h
@@ -81,7 +81,20 @@ struct FPCGStaticMeshSpawnerContext : public FPCGContext
 	TArray<FPCGSoftISMComponentDescriptor> OverriddenDescriptors;
 
 	// Keeping track of all touched resources to allow for correct cleanup on abort
-	TArray<TWeakObjectPtr<UPCGManagedISMComponent>> TouchedResources;
+	TSet<TWeakObjectPtr<UPCGManagedISMComponent>> TouchedResources;
+	
+	// struct use to maintain the used resources and if those resources have been merged with another spawner
+	struct FResourceReuseInfo
+	{
+		bool bMergedResources = false;
+		TSet<TWeakObjectPtr<UPCGManagedISMComponent>> Resources;
+	};
+	
+	// information about all the resources used by this spawner for this execution
+	FResourceReuseInfo UsedResources;
+
+	// information about all the resources used per input (data crc) by this spawner for this execution
+	TMap<uint32, FResourceReuseInfo> DataCrcToUsedResources;
 
 	UE_API void ResetInputIterationData();
 };
diff --git a/Engine/Plugins/PCG/Source/PCG/Public/PCGManagedResource.h b/Engine/Plugins/PCG/Source/PCG/Public/PCGManagedResource.h
--- a/Engine/Plugins/PCG/Source/PCG/Public/PCGManagedResource.h
+++ b/Engine/Plugins/PCG/Source/PCG/Public/PCGManagedResource.h
@@ -100,7 +100,7 @@ public:
 	virtual bool ReleaseOnTeardown() const { return false; }
 
 	/** Marks the resources as being kept and changed through generation */
-	virtual void MarkAsUsed() { ensure(CanBeUsed()); bIsMarkedUnused = false; }
+	UE_API virtual void MarkAsUsed();
 	/** Marks the resource as being reused as-is during the generation */
 	virtual void MarkAsReused() { bIsMarkedUnused = false; }
 	bool IsMarkedUnused() const { return bIsMarkedUnused; }
@@ -382,6 +382,12 @@ public:
 	FPCGCrc GetDataCrc() const { return DataCrc; }
 	void SetDataCrc(const FPCGCrc& InDataCrc) { DataCrc = InDataCrc; }
 
+	void SetCrcReuseCount(int32 InCount) { CrcReuseCount = InCount; }
+	int32 GetCrcReuseCount() const { return CrcReuseCount; }
+
+	UE_API void SetDataCrcReuseCount(int32 InCount) { DataCrcReuseCount = InCount; }
+	UE_API int32 GetDataCrcReuseCount() const { return DataCrcReuseCount; }
+
 	UE_DEPRECATED(5.6, "Use GetSettingsCrc instead")
 	uint64 GetSettingsUID() const { return -1; }
 	
@@ -406,6 +412,12 @@ protected:
 	
 	UPROPERTY()
 	FPCGCrc DataCrc;
+
+	UPROPERTY()
+	int32 CrcReuseCount = INDEX_NONE;
+
+	UPROPERTY()
+	int32 DataCrcReuseCount = INDEX_NONE;
 };
 
 UCLASS(MinimalAPI, BlueprintType)
