[PCG] Improved behavior on split spline internal call for closed splines. Fixed a bad array access.

#Synced-CL 50443815

diff --git a/Engine/Plugins/PCG/Source/PCG/Private/Elements/PCGSplitSplines.cpp b/Engine/Plugins/PCG/Source/PCG/Private/Elements/PCGSplitSplines.cpp
--- a/Engine/Plugins/PCG/Source/PCG/Private/Elements/PCGSplitSplines.cpp
+++ b/Engine/Plugins/PCG/Source/PCG/Private/Elements/PCGSplitSplines.cpp
@@ -20,16 +20,67 @@ namespace PCGSplitSplineConstants
 
 namespace PCGSplitSpline
 {
-	void SplitSpline(const UPCGSplineData* Spline, UPCGSplineData* OutSplitSpline, TConstArrayView<FSplinePoint> SplinePoints, TConstArrayView<PCGMetadataEntryKey> SplinePointsEntryKeys, double StartKey, double EndKey)
+	EStatus SplitSpline(const UPCGSplineData* Spline, UPCGSplineData* OutSplitSpline, TConstArrayView<FSplinePoint> SplinePoints, TConstArrayView<PCGMetadataEntryKey> SplinePointsEntryKeys, double StartKey, double EndKey)
 	{
 		check(Spline);
 		check(OutSplitSpline);
 
+		if (SplinePoints.IsEmpty())
+		{
+			return EStatus::Skipped;
+		}
+
 		if (!Spline->IsClosed() && StartKey > EndKey)
 		{
-			double TempKey = StartKey;
-			StartKey = EndKey;
-			EndKey = TempKey;
+			return EStatus::Skipped;
+		}
+
+		const double LoopKeyLength = Spline->IsClosed() ? Spline->GetLoopKeyOffset() : 1.0;
+		const double SplineKeyLength = SplinePoints.Last().InputKey + (Spline->IsClosed() ? LoopKeyLength : 0.0);
+		constexpr double KeyEpsilon = 1.0e-6;
+
+		// Normalize start and end keys so that they are valid.
+		if (!Spline->IsClosed())
+		{
+			StartKey = FMath::Max(StartKey, 0.0);
+			EndKey = FMath::Min(EndKey, SplineKeyLength);
+
+			// Spline is fully clipped, so nothing to do.
+			if (StartKey >= EndKey)
+			{
+				return EStatus::Skipped;
+			}
+		}
+		else
+		{
+			// Validate that the end key and the start key are in valid positions with respect to each other.
+			// Basically, if the end key is smaller than the start key, we'll bump it once and once only.
+			// After that, if the difference is negative, then the range was too large (e.g. > spline key length)
+			if (EndKey < StartKey)
+			{
+				EndKey += SplineKeyLength;
+			}
+
+			// Validate length of split
+			double SplitLength = EndKey - StartKey;
+
+			if ((SplitLength <= -KeyEpsilon) || (SplitLength >= SplineKeyLength + KeyEpsilon))
+			{
+				return EStatus::Error;
+			}
+
+			SplitLength = FMath::Clamp(SplitLength, 0, SplineKeyLength);
+
+			// 1. Get into [0, spline key length[ range.
+			StartKey -= SplineKeyLength * FMath::Floor(StartKey / SplineKeyLength);
+
+			if (StartKey < 0)
+			{
+				StartKey += SplineKeyLength;
+			}
+
+			// 2. Get the end key in the [0, 2 * spline keyLength[ range.
+			EndKey = StartKey + SplitLength;
 		}
 
 		TArray<FSplinePoint> SplitSplinePoints;
@@ -45,17 +96,17 @@ namespace PCGSplitSpline
 			}
 		};
 
-		auto AddSplitPoint = [&OutSplitSpline, &SplitSplinePoints, &SplitSplinePointsEntryKeys, &Spline, &SplinePointsEntryKeys](const double Key, TEnumAsByte<ESplinePointType::Type> PointType, const double PreviousKey, const double NextKey)
+		auto AddSplitPoint = [&OutSplitSpline, &SplitSplinePoints, &SplitSplinePointsEntryKeys, &Spline, &SplinePointsEntryKeys](const double Key, const double NormalizedKey, TEnumAsByte<ESplinePointType::Type> PointType, const double PreviousKey, const double NextKey)
 		{
 			check(Key >= PreviousKey && Key <= NextKey && PreviousKey <= NextKey);
 
 			FSplinePoint SplitPoint;
 			SplitPoint.InputKey = Key;
-			SplitPoint.Position = Spline->SplineStruct.GetLocationAtSplineInputKey(Key, ESplineCoordinateSpace::Local);
-			SplitPoint.ArriveTangent = Spline->SplineStruct.GetTangentAtSplineInputKey(Key, ESplineCoordinateSpace::Local);
+			SplitPoint.Position = Spline->SplineStruct.GetLocationAtSplineInputKey(NormalizedKey, ESplineCoordinateSpace::Local);
+			SplitPoint.ArriveTangent = Spline->SplineStruct.GetTangentAtSplineInputKey(NormalizedKey, ESplineCoordinateSpace::Local);
 			SplitPoint.LeaveTangent = SplitPoint.ArriveTangent;
-			SplitPoint.Rotation = Spline->SplineStruct.GetQuaternionAtSplineInputKey(Key, ESplineCoordinateSpace::Local).Rotator();
-			SplitPoint.Scale = Spline->SplineStruct.GetScaleAtSplineInputKey(Key);
+			SplitPoint.Rotation = Spline->SplineStruct.GetQuaternionAtSplineInputKey(NormalizedKey, ESplineCoordinateSpace::Local).Rotator();
+			SplitPoint.Scale = Spline->SplineStruct.GetScaleAtSplineInputKey(NormalizedKey);
 			SplitPoint.Type = PointType;
 
 			// The tangent is always computed from the original spline's perspective so we will assume the key range is always one.
@@ -69,7 +120,7 @@ namespace PCGSplitSpline
 			if (!SplinePointsEntryKeys.IsEmpty())
 			{
 				PCGMetadataEntryKey& EntryKey = SplitSplinePointsEntryKeys.Add_GetRef(PCGInvalidEntryKey);
-				Spline->WriteMetadataToEntry(Key, EntryKey, OutSplitSpline->MutableMetadata());
+				Spline->WriteMetadataToEntry(NormalizedKey, EntryKey, OutSplitSpline->MutableMetadata());
 			}
 		};
 
@@ -87,114 +138,113 @@ namespace PCGSplitSpline
 			}
 		};
 
-		int FirstControlPointIndexInSegment = INDEX_NONE;
-		int LastControlPointIndexInSegment = INDEX_NONE;
-
-		for (int SplinePointIndex = 0; SplinePointIndex < SplinePoints.Num(); ++SplinePointIndex)
+		auto PrepareControlPointsInfo = [&SplinePoints, Spline, SplineKeyLength]()
 		{
-			if (SplinePoints[SplinePointIndex].InputKey >= StartKey)
+			// start key, end key, original point index
+			TArray<TTuple<double, double, int>> ControlPointsInfo;
+
+			for (int Index = 0; Index < SplinePoints.Num(); ++Index)
 			{
-				FirstControlPointIndexInSegment = SplinePointIndex;
-				break;
+				ControlPointsInfo.Emplace(SplinePoints[Index].InputKey, 0.0, Index);
 			}
-		}
 
-		for (int SplinePointIndex = SplinePoints.Num() - 1; SplinePointIndex >= 0; --SplinePointIndex)
-		{
-			if (SplinePoints[SplinePointIndex].InputKey <= EndKey)
+			if (Spline->IsClosed())
 			{
-				LastControlPointIndexInSegment = SplinePointIndex;
-				break;
+				for (int Index = 0; Index < SplinePoints.Num(); ++Index)
+				{
+					ControlPointsInfo.Emplace(SplinePoints[Index].InputKey + SplineKeyLength, 0.0, Index);
+				}
 			}
-		}
 
-		if (FirstControlPointIndexInSegment == INDEX_NONE || LastControlPointIndexInSegment == INDEX_NONE)
-		{
-			// Invalid keys - unclear what caller is expecting
-			ensure(false);
-			return;
-		}
+			for(int Index = 0; Index < ControlPointsInfo.Num() - 1; ++Index)
+			{
+				ControlPointsInfo[Index].Get<1>() = ControlPointsInfo[Index + 1].Get<0>();
+			}
+			ControlPointsInfo.Last().Get<1>() = (Spline->IsClosed() ? 2.0 : 1.0) * SplineKeyLength;
 
-		if (Spline->IsClosed() && EndKey <= StartKey) // allow for loop
-		{
-			LastControlPointIndexInSegment += SplinePoints.Num();
-		}
+			return ControlPointsInfo;
+		};
 
-		const double LoopKeyLength = Spline->IsClosed() ? Spline->GetLoopKeyOffset() : 1.0;
+		// Prepare control points information so we can so a simple traversal
+		// 0: start key, 1: end key, 2: original point index
+		TArray<TTuple<double, double, int>> ControlPointsInfo = PrepareControlPointsInfo();
 
-		const double FirstControlKey = SplinePoints[FirstControlPointIndexInSegment].InputKey;
-		const double PreviousControlKey = FirstControlPointIndexInSegment > 0 ? SplinePoints[FirstControlPointIndexInSegment - 1].InputKey : FirstControlKey - LoopKeyLength;
-		const double LastControlKey = SplinePoints[(LastControlPointIndexInSegment % SplinePoints.Num())].InputKey;
-		const double NextControlKey = ((LastControlPointIndexInSegment + 1) % SplinePoints.Num() == 0) ? (LastControlKey + LoopKeyLength) : SplinePoints[((LastControlPointIndexInSegment + 1) % SplinePoints.Num())].InputKey;
+		double CurrentKey = StartKey;
+		int CurrentSegmentIndex = 0;
+		double StartTangentAdjustmentRatio = -1.0;
+		double EndTangentAdjustmentRatio = -1.0;
 
-		bool bNeedsStartTangentAdjustment = false;
-		bool bNeedsEndTangentAdjustment = false;
+		// Iterate through control points and interpolate points when needed
+		while (CurrentKey <= EndKey && CurrentSegmentIndex < ControlPointsInfo.Num())
+		{
+			const TTuple<double, double, int>& Segment = ControlPointsInfo[CurrentSegmentIndex];
+			++CurrentSegmentIndex;
+			EndTangentAdjustmentRatio = -1.0;
 
-		// Add left split point, prioritize existing control points when possible
-		const double KeyEpsilon = 1.0e-6;
+			if (CurrentKey < (Segment.Get<0>() - KeyEpsilon) || CurrentKey > Segment.Get<1>() + KeyEpsilon)
+			{
+				continue;
+			}
 
-		const double StartKeyDeltaToPreviousControlPoint = StartKey - PreviousControlKey;
-		const double StartKeyDeltaToFirstControlPoint = FirstControlKey - StartKey;
+			const int StartControlPointIndex = Segment.Get<2>();
+			const int EndControlPointIndex = (StartControlPointIndex + 1) % SplinePoints.Num();
 
-		if (StartKeyDeltaToPreviousControlPoint < KeyEpsilon)
-		{
-			if (StartKeyDeltaToPreviousControlPoint > 0) // otherwise point is already added in the control point loop
+			// Now the current key can either be:
+			// (1) on the start control point
+			const double SegmentStartDelta = (CurrentKey - Segment.Get<0>());
+			if (FMath::Abs(SegmentStartDelta) <= KeyEpsilon)
 			{
-				AddExistingControlPoint((FirstControlPointIndexInSegment + SplinePoints.Num() - 1) % SplinePoints.Num());
+				AddExistingControlPoint(StartControlPointIndex);
 			}
-		}
-		else if (StartKeyDeltaToFirstControlPoint < KeyEpsilon)
-		{
-			// will be added by the control point loop
-		}
-		else
-		{
-			TEnumAsByte<ESplinePointType::Type> SplitPointType = GetPointTypeForTouchedPoint(SplinePoints[(FirstControlPointIndexInSegment + SplinePoints.Num() - 1) % SplinePoints.Num()].Type);
-			AddSplitPoint(StartKey, SplitPointType, PreviousControlKey, FirstControlKey);
-			bNeedsStartTangentAdjustment = true;
-		}
+			// (2) in between both
+			else
+			{
+				const double Ratio = (Segment.Get<1>() - Segment.Get<0>() > 0) ? (CurrentKey - Segment.Get<0>()) / (Segment.Get<1>() - Segment.Get<0>()) : 0.0;
 
-		// Add control points that are taken as-is from the original spline points
-		for (int ControlPointIndex = FirstControlPointIndexInSegment; ControlPointIndex <= LastControlPointIndexInSegment; ++ControlPointIndex)
-		{
-			AddExistingControlPoint(ControlPointIndex % SplinePoints.Num());
-		}
+				if (SplitSplinePoints.IsEmpty())
+				{
+					StartTangentAdjustmentRatio = Ratio;
+				}
 
-		// Add right split point
-		const double EndKeyDeltaToLastControlPoint = EndKey - LastControlKey;
-		const double EndKeyDeltaToNextControlPoint = NextControlKey - EndKey;
+				TEnumAsByte<ESplinePointType::Type> SplitPointType = GetPointTypeForTouchedPoint(SplinePoints[StartControlPointIndex].Type);
+				AddSplitPoint(CurrentKey, FMath::Fmod(CurrentKey, SplineKeyLength), SplitPointType, Segment.Get<0>(), Segment.Get<1>());
 
-		if (EndKeyDeltaToLastControlPoint < KeyEpsilon)
-		{
-			// will be already added by the control point loop
-		}
-		else if(EndKeyDeltaToNextControlPoint < KeyEpsilon)
-		{
-			AddExistingControlPoint((LastControlPointIndexInSegment + 1) % SplinePoints.Num());
-		}
-		else
-		{
-			TEnumAsByte<ESplinePointType::Type> SplitPointType = GetPointTypeForTouchedPoint(SplinePoints[LastControlPointIndexInSegment % SplinePoints.Num()].Type);
-			AddSplitPoint(EndKey, SplitPointType, LastControlKey, NextControlKey);
-			bNeedsEndTangentAdjustment = true;
+				EndTangentAdjustmentRatio = Ratio;
+			}
+
+			if (CurrentKey < EndKey)
+			{
+				if (Segment.Get<1>() <= EndKey)
+				{
+					CurrentKey = Segment.Get<1>();
+				}
+				else
+				{
+					CurrentKey = EndKey;
+					// Since we pre-incremented, we need to rollback once when we get to the final point, which would lie in the same segment as the last point.
+					--CurrentSegmentIndex;
+				}
+			}
+			else
+			{
+				break;
+			}
 		}
 
 		// Adjust tangents if required
 		if (SplitSplinePoints.Num() >= 2)
 		{
-			if (bNeedsStartTangentAdjustment)
+			if(StartTangentAdjustmentRatio >= 0.0)
 			{
-				const double Ratio = ((FirstControlKey - PreviousControlKey) > KeyEpsilon) ? (StartKey - PreviousControlKey) / (FirstControlKey - PreviousControlKey) : 0.0;
 				FSplinePoint& FirstPreexistingPoint = SplitSplinePoints[1];
-				FirstPreexistingPoint.ArriveTangent *= (1.0 - Ratio);
+				FirstPreexistingPoint.ArriveTangent *= (1.0 - StartTangentAdjustmentRatio);
 				FirstPreexistingPoint.Type = GetPointTypeForTouchedPoint(FirstPreexistingPoint.Type);
 			}
 
-			if (bNeedsEndTangentAdjustment)
+			if(EndTangentAdjustmentRatio >= 0.0)
 			{
-				const double Ratio = ((NextControlKey - LastControlKey) > KeyEpsilon) ? (EndKey - LastControlKey) / (NextControlKey - LastControlKey) : 1.0;
 				FSplinePoint& LastPreexistingPoint = SplitSplinePoints[SplitSplinePoints.Num() - 2];
-				LastPreexistingPoint.LeaveTangent *= Ratio;
+				LastPreexistingPoint.LeaveTangent *= EndTangentAdjustmentRatio;
 				LastPreexistingPoint.Type = GetPointTypeForTouchedPoint(LastPreexistingPoint.Type);
 			}
 		}
@@ -206,14 +256,22 @@ namespace PCGSplitSpline
 			SplinePoint.InputKey = static_cast<float>(SPIndex);
 		}
 
+		if (SplitSplinePoints.Num() > 0)
+		{
+			SplitSplinePoints[0].Type = GetPointTypeForTouchedPoint(SplitSplinePoints[0].Type);
+			SplitSplinePoints.Last().Type = GetPointTypeForTouchedPoint(SplitSplinePoints.Last().Type);
+		}
+
 		// Initialize data now
 		OutSplitSpline->Initialize(SplitSplinePoints,
 			/*bIsClosedLoop*/false,
 			Spline->SplineStruct.GetTransform(),
 			SplitSplinePointsEntryKeys);
+
+		return EStatus::Split;
 	}
 
-	void SplitSpline(const UPCGSplineData* Spline, UPCGSplineData* OutSplitSpline, double StartKey, double EndKey)
+	EStatus SplitSpline(const UPCGSplineData* Spline, UPCGSplineData* OutSplitSpline, double StartKey, double EndKey)
 	{
 		check(Spline);
 		check(OutSplitSpline);
@@ -221,7 +279,7 @@ namespace PCGSplitSpline
 		TArray<FSplinePoint> SplinePoints = Spline->GetSplinePoints();
 		TConstArrayView<PCGMetadataEntryKey> SplinePointsEntryKeys = Spline->GetConstVerticesEntryKeys();
 
-		SplitSpline(Spline, OutSplitSpline, SplinePoints, SplinePointsEntryKeys, StartKey, EndKey);
+		return SplitSpline(Spline, OutSplitSpline, SplinePoints, SplinePointsEntryKeys, StartKey, EndKey);
 	}
 }
 
@@ -361,7 +419,7 @@ bool FPCGSplitSplineElement::ExecuteInternal(FPCGContext* Context) const
 		}
 		else
 		{
-			TArray<double> Values;
+			TArray<double, TInlineAllocator<128>> Values;
 			if (Settings->bUseConstant)
 			{
 				Values.Add(Settings->Constant);
@@ -396,21 +454,29 @@ bool FPCGSplitSplineElement::ExecuteInternal(FPCGContext* Context) const
 			}
 		}
 
+		constexpr double KeyEpsilon = 1.0e-6f;
 		const bool bIsClosed = SplineData->IsClosed();
 
 		// Remove values that are outside of the spline
 		const float FirstKey = SplinePoints.IsEmpty() ? 0.0f : SplinePoints[0].InputKey;
 		const float LastKey = SplinePoints.IsEmpty() ? 1.0f : (SplinePoints.Last().InputKey + (bIsClosed ? SplineData->GetLoopKeyOffset() : 0.0f));
 
-		// Closed loops should properly cut at key = 0 and up to the last key + 1.
 		if (bIsClosed)
 		{
-			for (int SplitIndex = SplitKeys.Num() - 1; SplitIndex >= 0; --SplitIndex)
+			// Closed loops can be a bit more flexible on the input keys, but have a specific requirement:
+			// the range of keys should not be more than ONE spline length.
+			float MinKey = SplitKeys.IsEmpty() ? 0.0f : SplitKeys[0];
+			float MaxKey = MinKey;
+			for (int SplitIndex = 1; SplitIndex < SplitKeys.Num(); ++SplitIndex)
 			{
-				if (SplitKeys[SplitIndex] < FirstKey || SplitKeys[SplitIndex] > LastKey)
-				{
-					SplitKeys.RemoveAtSwap(SplitIndex);
-				}
+				MinKey = FMath::Min(MinKey, SplitKeys[SplitIndex]);
+				MaxKey = FMath::Max(MaxKey, SplitKeys[SplitIndex]);
+			}
+
+			if (MaxKey - MinKey > (LastKey - FirstKey) + KeyEpsilon)
+			{
+				PCGLog::LogWarningOnGraph(LOCTEXT("InvalidKeyRangeForClosedLoop", "Specified split keys are not valid for a closed spline in this operation."), Context);
+				continue;
 			}
 		}
 		else
@@ -428,15 +494,17 @@ bool FPCGSplitSplineElement::ExecuteInternal(FPCGContext* Context) const
 		// If there are no valid values, then there is nothing to do.
 		if (SplitKeys.IsEmpty())
 		{
-			Outputs.Add(Input);
+			if (Settings->CutDataBehavior == EPCGSplitSplineDataBehavior::KeepAllSplines)
+			{
+				Outputs.Add(Input);
+			}
+
 			continue;
 		}
 
 		// Sort values
 		SplitKeys.Sort([](const double& A, const double& B) { return A < B; });
 
-		const double KeyEpsilon = 1.0e-6f;
-
 		// Remove duplicates/values too close to each other
 		for (int SplitIndex = SplitKeys.Num() - 1; SplitIndex > 0; --SplitIndex)
 		{
@@ -454,9 +522,9 @@ bool FPCGSplitSplineElement::ExecuteInternal(FPCGContext* Context) const
 		if (bIsClosed)
 		{
 			const double FirstSplitKey = SplitKeys[0];
-			SplitKeys.Add(FirstSplitKey);
+			SplitKeys.Add(FirstSplitKey + LastKey);
 		}
-		else
+		else if(Settings->CutDataBehavior != EPCGSplitSplineDataBehavior::KeepInteriorSplinesOnly)
 		{
 			SplitKeys.Insert(SplinePoints[0].InputKey, 0);
 			SplitKeys.Add(SplinePoints.Last().InputKey);
@@ -470,26 +538,27 @@ bool FPCGSplitSplineElement::ExecuteInternal(FPCGContext* Context) const
 			UPCGSplineData* NewSplineData = FPCGContext::NewObject_AnyThread<UPCGSplineData>(Context);
 			NewSplineData->InitializeFromData(SplineData);
 
-			PCGSplitSpline::SplitSpline(SplineData, NewSplineData, SplinePoints, SplinePointsEntryKeys, StartKey, EndKey);
-
-			// Write originating spline index if required
-			if (Settings->bShouldOutputOriginatingSplineIndex)
+			if (PCGSplitSpline::SplitSpline(SplineData, NewSplineData, SplinePoints, SplinePointsEntryKeys, StartKey, EndKey) == PCGSplitSpline::EStatus::Split)
 			{
-				const FName IndexAttributeName = Settings->OutputOriginatingSplineIndex.GetName();
-				FPCGMetadataAttribute<int32>* OriginatingIndexAttribute = nullptr;
-
-				if (FPCGMetadataDomain* OutputMetadataDomain = NewSplineData->MutableMetadata()->GetMetadataDomainFromSelector(Settings->OutputOriginatingSplineIndex))
-				{
-					OriginatingIndexAttribute = OutputMetadataDomain->FindOrCreateAttribute(IndexAttributeName, InputIndex, /*bAllowInterpolation=*/false);
-				}
-
-				if (!OriginatingIndexAttribute)
+				// Write originating spline index if required
+				if (Settings->bShouldOutputOriginatingSplineIndex)
 				{
-					PCGLog::LogErrorOnGraph(FText::Format(LOCTEXT("FailedCreateIndexAttribute", "Failed to create the index attribute '{0}'."), FText::FromName(IndexAttributeName)), Context);
+					const FName IndexAttributeName = Settings->OutputOriginatingSplineIndex.GetName();
+					FPCGMetadataAttribute<int32>* OriginatingIndexAttribute = nullptr;
+
+					if (FPCGMetadataDomain* OutputMetadataDomain = NewSplineData->MutableMetadata()->GetMetadataDomainFromSelector(Settings->OutputOriginatingSplineIndex))
+					{
+						OriginatingIndexAttribute = OutputMetadataDomain->FindOrCreateAttribute(IndexAttributeName, InputIndex, /*bAllowInterpolation=*/false);
+					}
+
+					if (!OriginatingIndexAttribute)
+					{
+						PCGLog::LogErrorOnGraph(FText::Format(LOCTEXT("FailedCreateIndexAttribute", "Failed to create the index attribute '{0}'."), FText::FromName(IndexAttributeName)), Context);
+					}
 				}
-			}
 			
-			Outputs.Add_GetRef(Input).Data = NewSplineData;
+				Outputs.Add_GetRef(Input).Data = NewSplineData;
+			}
 		}
 	}
 
diff --git a/Engine/Plugins/PCG/Source/PCG/Private/Elements/Polygon/PCGClipPaths.cpp b/Engine/Plugins/PCG/Source/PCG/Private/Elements/Polygon/PCGClipPaths.cpp
--- a/Engine/Plugins/PCG/Source/PCG/Private/Elements/Polygon/PCGClipPaths.cpp
+++ b/Engine/Plugins/PCG/Source/PCG/Private/Elements/Polygon/PCGClipPaths.cpp
@@ -152,9 +152,10 @@ bool FPCGClipPathsElement::ExecuteInternal(FPCGContext* Context) const
 				UPCGSplineData* SplitSplineData = FPCGContext::NewObject_AnyThread<UPCGSplineData>(Context);
 				SplitSplineData->InitializeFromData(InputSpline);
 
-				PCGSplitSpline::SplitSpline(InputSpline, SplitSplineData, StartKey, EndKey);
-
-				PathData.Add(SplitSplineData);
+				if (PCGSplitSpline::SplitSpline(InputSpline, SplitSplineData, StartKey, EndKey) == PCGSplitSpline::EStatus::Split)
+				{
+					PathData.Add(SplitSplineData);
+				}
 			}
 			else if (const UPCGBasePointData* InputPointData = Cast<UPCGBasePointData>(Input.Data))
 			{
@@ -296,4 +297,4 @@ bool FPCGClipPathsElement::ExecuteInternal(FPCGContext* Context) const
 	return true;
 }
 
-#undef LOCTEXT_NAMESPACE
+#undef LOCTEXT_NAMESPACE
\ No newline at end of file
diff --git a/Engine/Plugins/PCG/Source/PCG/Private/Elements/Polygon/PCGGetSegment.cpp b/Engine/Plugins/PCG/Source/PCG/Private/Elements/Polygon/PCGGetSegment.cpp
--- a/Engine/Plugins/PCG/Source/PCG/Private/Elements/Polygon/PCGGetSegment.cpp
+++ b/Engine/Plugins/PCG/Source/PCG/Private/Elements/Polygon/PCGGetSegment.cpp
@@ -257,13 +257,10 @@ bool FPCGGetSegmentElement::ExecuteInternal(FPCGContext* Context) const
 				UPCGSplineData* OutSplineData = FPCGContext::NewObject_AnyThread<UPCGSplineData>(Context);
 				OutSplineData->InitializeFromData(SplineData);
 
-				PCGSplitSpline::SplitSpline(
-					SplineData,
-					OutSplineData,
-					StartKey,
-					EndKey);
-
-				OutputData = OutSplineData;
+				if (PCGSplitSpline::SplitSpline(SplineData, OutSplineData, StartKey, EndKey) == PCGSplitSpline::EStatus::Split)
+				{
+					OutputData = OutSplineData;
+				}
 			}
 			else
 			{
diff --git a/Engine/Plugins/PCG/Source/PCG/Public/Elements/PCGSplitSplines.h b/Engine/Plugins/PCG/Source/PCG/Public/Elements/PCGSplitSplines.h
--- a/Engine/Plugins/PCG/Source/PCG/Public/Elements/PCGSplitSplines.h
+++ b/Engine/Plugins/PCG/Source/PCG/Public/Elements/PCGSplitSplines.h
@@ -16,10 +16,25 @@ enum class EPCGSplitSplineMode
 	ByPredicateOnControlPoints,
 };
 
+UENUM()
+enum class EPCGSplitSplineDataBehavior
+{
+	KeepAllSplines,
+	KeepCutSplines,
+	KeepInteriorSplinesOnly
+};
+
 namespace PCGSplitSpline
 {
-	void SplitSpline(const UPCGSplineData* Spline, UPCGSplineData* SplitSpline, double StartKey, double EndKey);
-	void SplitSpline(const UPCGSplineData* Spline, UPCGSplineData* SplitSpline, TConstArrayView<FSplinePoint> SplinePoints, TConstArrayView<PCGMetadataEntryKey> SplinePointsEntryKeys, double StartKey, double EndKey);
+	enum class EStatus
+	{
+		Split,
+		Skipped,
+		Error
+	};
+
+	EStatus SplitSpline(const UPCGSplineData* Spline, UPCGSplineData* SplitSpline, double StartKey, double EndKey);
+	EStatus SplitSpline(const UPCGSplineData* Spline, UPCGSplineData* SplitSpline, TConstArrayView<FSplinePoint> SplinePoints, TConstArrayView<PCGMetadataEntryKey> SplinePointsEntryKeys, double StartKey, double EndKey);
 }
 
 /** Splits spline at a specific distance(s), key(s) or at certain values. */
@@ -50,6 +65,10 @@ public:
 	UPROPERTY(BlueprintReadWrite, EditAnywhere, Category = Settings)
 	EPCGSplitSplineMode Mode = EPCGSplitSplineMode::ByAlpha;
 
+	/** Behavior of whether we'll keep uncut, cut on one side or only interior splines. Note that for closed splines, all cut splines are interior splines. */
+	UPROPERTY(BlueprintReadWrite, EditAnywhere, Category = Settings)
+	EPCGSplitSplineDataBehavior CutDataBehavior = EPCGSplitSplineDataBehavior::KeepAllSplines;
+
 	/** Controls whether the input splines will be cut using a single constant or values driven either by an additional input or based on a predicate on the spline control points. */
 	UPROPERTY(BlueprintReadWrite, EditAnywhere, Category = Settings, meta = (EditCondition ="Mode != EPCGSplitSplineMode::ByPredicateOnControlPoints"))
 	bool bUseConstant = true;
