[PCG] Improved spline point type management
- Create spline now supports interpolation type.
- Get Spline control points returns interpolation type, also for polygon 2d.
- Clean spline now preserves segment types (at least the linear vs curve differentiation)
- Clean spline works on polygons

#Synced-CL 50443815

diff --git a/Engine/Plugins/PCG/Source/PCG/Private/Elements/PCGCleanSpline.cpp b/Engine/Plugins/PCG/Source/PCG/Private/Elements/PCGCleanSpline.cpp
--- a/Engine/Plugins/PCG/Source/PCG/Private/Elements/PCGCleanSpline.cpp
+++ b/Engine/Plugins/PCG/Source/PCG/Private/Elements/PCGCleanSpline.cpp
@@ -4,6 +4,7 @@
 
 #include "PCGContext.h"
 #include "Data/PCGSplineData.h"
+#include "Data/PCGPolygon2DData.h"
 
 #include UE_INLINE_GENERATED_CPP_BY_NAME(PCGCleanSpline)
 
@@ -16,13 +17,23 @@ namespace PCGCleanSplineHelpers
 		return FMath::Abs(FirstVector.GetSafeNormal() | SecondVector.GetSafeNormal()) >= Threshold;
 	}
 
-	bool ControlPointsAreColocated(const FSplinePoint& Point1, const FSplinePoint& Point2, const FTransform& SplineTransform, const double Threshold, const bool bUseLocalSpace = false)
+	bool ControlPointsAreColocated(
+		const FVector& Point1, 
+		const FVector& Point2, 
+		const FTransform& SplineTransform, 
+		const double Threshold, 
+		const bool bUseLocalSpace = false)
 	{
-		const FVector Location1 = bUseLocalSpace ? Point1.Position : SplineTransform.TransformPosition(Point1.Position);
-		const FVector Location2 = bUseLocalSpace ? Point2.Position : SplineTransform.TransformPosition(Point2.Position);
+		const FVector Location1 = bUseLocalSpace ? Point1: SplineTransform.TransformPosition(Point1);
+		const FVector Location2 = bUseLocalSpace ? Point2: SplineTransform.TransformPosition(Point2);
 		return FVector::DistSquared(Location1, Location2) < Threshold * Threshold;
 	}
 
+	bool ControlPointsAreCollinear(const FVector& Point1, const FVector& Point2, const FVector& Point3, const double Threshold)
+	{
+		return VectorsAreCollinear(Point2 - Point1, Point3 - Point2, Threshold);
+	}
+
 	// Implementation note: Co-located points will have a vector dot product of zero, regardless of tangents, and will be thus collinear.
 	bool ControlPointsAreCollinear(const FSplinePoint& Point1, const FSplinePoint& Point2, const FSplinePoint& Point3, const double Threshold)
 	{
@@ -70,14 +81,14 @@ void UPCGCleanSplineSettings::PostEditChangeProperty(FPropertyChangedEvent& Prop
 TArray<FPCGPinProperties> UPCGCleanSplineSettings::InputPinProperties() const
 {
 	TArray<FPCGPinProperties> Properties;
-	Properties.Emplace_GetRef(PCGPinConstants::DefaultInputLabel, EPCGDataType::Spline).SetRequiredPin();
+	Properties.Emplace_GetRef(PCGPinConstants::DefaultInputLabel, EPCGDataType::Spline | EPCGDataType::Polygon2D).SetRequiredPin();
 	return Properties;
 }
 
 TArray<FPCGPinProperties> UPCGCleanSplineSettings::OutputPinProperties() const
 {
 	TArray<FPCGPinProperties> Properties;
-	Properties.Emplace(PCGPinConstants::DefaultOutputLabel, EPCGDataType::Spline);
+	Properties.Emplace(PCGPinConstants::DefaultOutputLabel, EPCGDataType::Spline | EPCGDataType::Polygon2D);
 	return Properties;
 }
 
@@ -98,54 +109,78 @@ bool FPCGCleanSplineElement::ExecuteInternal(FPCGContext* InContext) const
 		return true;
 	}
 
+	for (const FPCGTaggedData& InputData : InContext->InputData.GetInputsByPin(PCGPinConstants::DefaultInputLabel))
+	{
+		if (const UPCGSplineData* InputSplineData = Cast<UPCGSplineData>(InputData.Data))
+		{
+			ProcessSpline(InContext, Settings, InputData, InputSplineData);
+		}
+		else if (const UPCGPolygon2DData* InputPolygonData = Cast<UPCGPolygon2DData>(InputData.Data))
+		{
+			ProcessPolygon(InContext, Settings, InputData, InputPolygonData);
+		}
+	}
+
+	return true;
+}
+
+void FPCGCleanSplineElement::ProcessSpline(FPCGContext* InContext, const UPCGCleanSplineSettings* Settings, const FPCGTaggedData& InputData, const UPCGSplineData* InputSplineData) const
+{
+	check(InContext && Settings && InputSplineData);
+
+	// Early out on splines that are already invalid
+	if (InputSplineData->GetNumSegments() < 1)
+	{
+		return;
+	}
+
 	// Pre-calculate the tolerance value from the user defined tolerance in radians/degrees to cross-product comparable value.
 	const double DotProductToleranceRad = Settings->bUseRadians ? Settings->CollinearAngleThreshold : FMath::DegreesToRadians(Settings->CollinearAngleThreshold);
 	// Max the tolerance with an epsilon to accomodate rounding errors.
 	const double DotProductTolerance = FMath::Max(FMath::Abs(FMath::Cos(DotProductToleranceRad)), UE_DOUBLE_SMALL_NUMBER);
 
-	for (const FPCGTaggedData& InputData : InContext->InputData.GetInputsByPin(PCGPinConstants::DefaultInputLabel))
+	const FTransform& SplineTransform = InputSplineData->SplineStruct.Transform;
+	const FInterpCurveVector& ControlPointsPosition = InputSplineData->SplineStruct.GetSplinePointsPosition();
+	const FInterpCurveQuat& ControlPointsRotation = InputSplineData->SplineStruct.GetSplinePointsRotation();
+	const FInterpCurveVector& ControlPointsScale = InputSplineData->SplineStruct.GetSplinePointsScale();
+	const TConstArrayView<PCGMetadataEntryKey> ControlPointsMetadataEntry = InputSplineData->SplineStruct.GetConstControlPointsEntryKeys();
+	const int32 NumControlPoints = ControlPointsPosition.Points.Num();
+	bool bControlPointWasRemoved = false;
+	const bool bIsClosed = InputSplineData->IsClosed();
+
+	TArray<FSplinePoint> ControlPoints;
+	ControlPoints.Reserve(NumControlPoints);
+
+	// Making sure that the number of metadata entry keys matches, if not, we are in an invalid state and will reset the metadata entry keys.
+	TArray<PCGMetadataEntryKey> ControlPointsKeys;
+	if (ControlPointsMetadataEntry.Num() == NumControlPoints)
 	{
-		const UPCGSplineData* InputSplineData = Cast<const UPCGSplineData>(InputData.Data);
-		if (!InputSplineData || InputSplineData->GetNumSegments() < 1)
-		{
-			continue;
-		}
+		ControlPointsKeys.Append(ControlPointsMetadataEntry);
+	}
 
-		const FTransform& SplineTransform = InputSplineData->SplineStruct.Transform;
-		const FInterpCurveVector& ControlPointsPosition = InputSplineData->SplineStruct.GetSplinePointsPosition();
-		const FInterpCurveQuat& ControlPointsRotation = InputSplineData->SplineStruct.GetSplinePointsRotation();
-		const FInterpCurveVector& ControlPointsScale = InputSplineData->SplineStruct.GetSplinePointsScale();
-		const TConstArrayView<PCGMetadataEntryKey> ControlPointsMetadataEntry = InputSplineData->SplineStruct.GetConstControlPointsEntryKeys();
-		const int32 NumControlPoints = ControlPointsPosition.Points.Num();
-		bool bControlPointWasRemoved = false;
-		const bool bIsClosed = InputSplineData->IsClosed();
-
-		TArray<FSplinePoint> ControlPoints;
-		ControlPoints.Reserve(NumControlPoints);
-
-		// Making sure that the number of metadata entry keys matches, if not, we are in an invalid state and will reset the metadata entry keys.
-		TArray<PCGMetadataEntryKey> ControlPointsKeys;
-		if (ControlPointsMetadataEntry.Num() == NumControlPoints)
-		{
-			ControlPointsKeys.Append(ControlPointsMetadataEntry);
-		}
+	// For code clarity, generate the points first and remove them as needed.
+	for (int32 i = 0; i < NumControlPoints; ++i)
+	{
+		/* Implementation note: Decay to custom tangents. They assist the user in building the spline, but the
+		 * interpolation modes will affect the recalculations unpredictably when control points are removed.
+		 */
+		ESplinePointType::Type ControlPointType = ConvertInterpCurveModeToSplinePointType(ControlPointsPosition.Points[i].InterpMode);
 
-		// For code clarity, generate the points first and remove them as needed.
-		for (int32 i = 0; i < NumControlPoints; ++i)
+		if ((ControlPointType == ESplinePointType::Curve || ControlPointType == ESplinePointType::CurveClamped || ControlPointType == ESplinePointType::CurveCustomTangent))
 		{
-			/* Implementation note: Decay to custom tangents. They assist the user in building the spline, but the
-			 * interpolation modes will affect the recalculations unpredictably when control points are removed.
-			 */
-			ControlPoints.Emplace(/*InputKey=*/i,
-			   ControlPointsPosition.Points[i].OutVal,
-			   ControlPointsPosition.Points[i].ArriveTangent,
-			   ControlPointsPosition.Points[i].LeaveTangent,
-			   ControlPointsRotation.Points[i].OutVal.Rotator(),
-			   ControlPointsScale.Points[i].OutVal,
-			   ESplinePointType::CurveCustomTangent);
+			ControlPointType = ESplinePointType::CurveCustomTangent;
 		}
 
-		auto RemovePoint = [&ControlPoints, &ControlPointsKeys, &bControlPointWasRemoved](int Index)
+		ControlPoints.Emplace(/*InputKey=*/i,
+			ControlPointsPosition.Points[i].OutVal,
+			ControlPointsPosition.Points[i].ArriveTangent,
+			ControlPointsPosition.Points[i].LeaveTangent,
+			ControlPointsRotation.Points[i].OutVal.Rotator(),
+			ControlPointsScale.Points[i].OutVal,
+			ControlPointType);
+	}
+
+	auto RemovePoint = [&ControlPoints, &ControlPointsKeys, &bControlPointWasRemoved](int Index)
 		{
 			if (!ControlPointsKeys.IsEmpty())
 			{
@@ -158,120 +193,253 @@ bool FPCGCleanSplineElement::ExecuteInternal(FPCGContext* InContext) const
 			bControlPointWasRemoved = true;
 		};
 
-		auto GetPreviousIndexFromOffset = [&ControlPoints](int32 Index, int32 Offset)
+	auto GetPreviousIndexFromOffset = [&ControlPoints](int32 Index, int32 Offset)
 		{
 			if (!ensure(!ControlPoints.IsEmpty()))
 			{
 				return 0;
 			}
-			
+
 			int32 Result = Index - Offset;
 			while (Result < 0)
 			{
 				Result += ControlPoints.Num();
 			}
-			
+
 			return Result;
 		};
 
-		if (Settings->bFuseColocatedControlPoints)
+	if (Settings->bFuseColocatedControlPoints)
+	{
+		const int32 MinIndex = bIsClosed ? 0 : 1;
+
+		// Will evaluate by pairs. Reverse order for optimizing RemoveAt.
+		for (int Index = ControlPoints.Num() - 1; Index >= MinIndex; --Index)
 		{
-			const int32 MinIndex = bIsClosed ? 0 : 1;
-			
-			// Will evaluate by pairs. Reverse order for optimizing RemoveAt.
-			for (int Index = ControlPoints.Num() - 1; Index >= MinIndex; --Index)
+			const int32 PreviousPointIndex = GetPreviousIndexFromOffset(Index, 1);
+			const int32 CurrentPointIndex = Index;
+
+			// Evaluate by pairs for colocated.
+			if (PCGCleanSplineHelpers::ControlPointsAreColocated(
+				ControlPoints[PreviousPointIndex].Position,
+				ControlPoints[CurrentPointIndex].Position,
+				SplineTransform,
+				Settings->ColocationDistanceThreshold,
+				Settings->bUseSplineLocalSpace))
 			{
-				const int32 PreviousPointIndex = GetPreviousIndexFromOffset(Index, 1);
-				const int32 CurrentPointIndex = Index;
-				
-				// Evaluate by pairs for colocated.
-				if (PCGCleanSplineHelpers::ControlPointsAreColocated(
-						ControlPoints[PreviousPointIndex],
-						ControlPoints[CurrentPointIndex],
-						SplineTransform,
-						Settings->ColocationDistanceThreshold,
-						Settings->bUseSplineLocalSpace))
+				EPCGControlPointFuseMode FuseMode = Settings->FuseMode;
+				// Generally, keep the previous control point (first), but preserve the final control point to maintain the spline's length (second) in non-closed splines.
+				if (FuseMode == EPCGControlPointFuseMode::Auto)
 				{
-					EPCGControlPointFuseMode FuseMode = Settings->FuseMode;
-					// Generally, keep the previous control point (first), but preserve the final control point to maintain the spline's length (second) in non-closed splines.
-					if (FuseMode == EPCGControlPointFuseMode::Auto)
-					{
-						FuseMode = (!bIsClosed && CurrentPointIndex == ControlPoints.Num() - 1) ? EPCGControlPointFuseMode::KeepSecond : EPCGControlPointFuseMode::KeepFirst;
-					}
-
-					switch (FuseMode)
-					{
-						// Keep the first control point
-						case EPCGControlPointFuseMode::KeepFirst:
-							ControlPoints[PreviousPointIndex].LeaveTangent = ControlPoints[CurrentPointIndex].LeaveTangent;
-							RemovePoint(CurrentPointIndex);
-							break;
-						// Keep the second control point
-						case EPCGControlPointFuseMode::KeepSecond:
-							ControlPoints[CurrentPointIndex].ArriveTangent = ControlPoints[PreviousPointIndex].ArriveTangent;
-							RemovePoint(PreviousPointIndex);
-							break;
-						// Average the two control points' transforms and update the leave tangent
-						case EPCGControlPointFuseMode::Merge:
-							ControlPoints[PreviousPointIndex].Position = FMath::Lerp(ControlPoints[PreviousPointIndex].Position, ControlPoints[CurrentPointIndex].Position, 0.5);
-							ControlPoints[PreviousPointIndex].Rotation = FMath::Lerp(ControlPoints[PreviousPointIndex].Rotation, ControlPoints[CurrentPointIndex].Rotation, 0.5);
-							ControlPoints[PreviousPointIndex].Scale = FMath::Lerp(ControlPoints[PreviousPointIndex].Scale, ControlPoints[CurrentPointIndex].Scale, 0.5);
-							ControlPoints[PreviousPointIndex].LeaveTangent = ControlPoints[CurrentPointIndex].LeaveTangent;
-							RemovePoint(CurrentPointIndex);
-							break;
-						case EPCGControlPointFuseMode::Auto: // Should've been picked up by now. Fallthrough...
-							checkNoEntry();
-						default:
-							break;
-					}
+					FuseMode = (!bIsClosed && CurrentPointIndex == ControlPoints.Num() - 1) ? EPCGControlPointFuseMode::KeepSecond : EPCGControlPointFuseMode::KeepFirst;
+				}
+
+				switch (FuseMode)
+				{
+					// Keep the first control point
+				case EPCGControlPointFuseMode::KeepFirst:
+					ControlPoints[PreviousPointIndex].LeaveTangent = ControlPoints[CurrentPointIndex].LeaveTangent;
+					RemovePoint(CurrentPointIndex);
+					break;
+					// Keep the second control point
+				case EPCGControlPointFuseMode::KeepSecond:
+					ControlPoints[CurrentPointIndex].ArriveTangent = ControlPoints[PreviousPointIndex].ArriveTangent;
+					RemovePoint(PreviousPointIndex);
+					break;
+					// Average the two control points' transforms and update the leave tangent
+				case EPCGControlPointFuseMode::Merge:
+					ControlPoints[PreviousPointIndex].Position = FMath::Lerp(ControlPoints[PreviousPointIndex].Position, ControlPoints[CurrentPointIndex].Position, 0.5);
+					ControlPoints[PreviousPointIndex].Rotation = FMath::Lerp(ControlPoints[PreviousPointIndex].Rotation, ControlPoints[CurrentPointIndex].Rotation, 0.5);
+					ControlPoints[PreviousPointIndex].Scale = FMath::Lerp(ControlPoints[PreviousPointIndex].Scale, ControlPoints[CurrentPointIndex].Scale, 0.5);
+					ControlPoints[PreviousPointIndex].LeaveTangent = ControlPoints[CurrentPointIndex].LeaveTangent;
+					RemovePoint(CurrentPointIndex);
+					break;
+				case EPCGControlPointFuseMode::Auto: // Should've been picked up by now. Fallthrough...
+					checkNoEntry();
+				default:
+					break;
 				}
 			}
 		}
+	}
+
+	if (Settings->bRemoveCollinearControlPoints)
+	{
+		const int32 MinIndex = bIsClosed ? 0 : 2;
 
-		if (Settings->bRemoveCollinearControlPoints)
+		// Will evaluate by triplets. Reverse order for optimizing RemoveAt.
+		for (int Index = ControlPoints.Num() - 1; Index >= MinIndex; --Index)
 		{
-			const int32 MinIndex = bIsClosed ? 0 : 2;
-			
-			// Will evaluate by triplets. Reverse order for optimizing RemoveAt.
-			for (int Index = ControlPoints.Num() - 1; Index >= MinIndex; --Index)
+			const int32 SecondPreviousPointIndex = GetPreviousIndexFromOffset(Index, 2);
+			const int32 PreviousPointIndex = GetPreviousIndexFromOffset(Index, 1);
+			const int32 CurrentPointIndex = Index;
+
+			if (PCGCleanSplineHelpers::ControlPointsAreCollinear(
+				ControlPoints[SecondPreviousPointIndex],
+				ControlPoints[PreviousPointIndex],
+				ControlPoints[CurrentPointIndex],
+				DotProductTolerance))
 			{
-				const int32 SecondPreviousPointIndex = GetPreviousIndexFromOffset(Index, 2);
-				const int32 PreviousPointIndex = GetPreviousIndexFromOffset(Index, 1);
-				const int32 CurrentPointIndex = Index;
-				
-				if (PCGCleanSplineHelpers::ControlPointsAreCollinear(
-						ControlPoints[SecondPreviousPointIndex],
-						ControlPoints[PreviousPointIndex],
-						ControlPoints[CurrentPointIndex],
-						DotProductTolerance))
-				{
-					ControlPoints[SecondPreviousPointIndex].LeaveTangent = ControlPoints[CurrentPointIndex].Position - ControlPoints[SecondPreviousPointIndex].Position;
-					ControlPoints[CurrentPointIndex].ArriveTangent = ControlPoints[SecondPreviousPointIndex].LeaveTangent;
-					RemovePoint(PreviousPointIndex);
-				}
+				ControlPoints[SecondPreviousPointIndex].LeaveTangent = ControlPoints[CurrentPointIndex].Position - ControlPoints[SecondPreviousPointIndex].Position;
+				ControlPoints[CurrentPointIndex].ArriveTangent = ControlPoints[SecondPreviousPointIndex].LeaveTangent;
+				RemovePoint(PreviousPointIndex);
 			}
 		}
+	}
+
+	FPCGTaggedData& Output = InContext->OutputData.TaggedData.Emplace_GetRef(InputData);
+
+	// We only need to create a new data if a point was removed.
+	if (bControlPointWasRemoved)
+	{
+		// Update input keys if a point was removed, to keep them monotonically incremental
+		for (int i = 0; i < ControlPoints.Num(); ++i)
+		{
+			ControlPoints[i].InputKey = i;
+		}
+
+		UPCGSplineData* NewSplineData = FPCGContext::NewObject_AnyThread<UPCGSplineData>(InContext);
+		NewSplineData->InitializeFromData(InputSplineData);
+		NewSplineData->Initialize(ControlPoints, InputSplineData->IsClosed(), InputSplineData->GetTransform(), std::move(ControlPointsKeys));
+
+		Output.Data = NewSplineData;
+	}
+}
+
+void FPCGCleanSplineElement::ProcessPolygon(FPCGContext* Context, const UPCGCleanSplineSettings* Settings, const FPCGTaggedData& Input, const UPCGPolygon2DData* PolygonData) const
+{
+	check(Context && Settings && PolygonData);
+
+	// Pre-calculate the tolerance value from the user defined tolerance in radians/degrees to cross-product comparable value.
+	const double DotProductToleranceRad = Settings->bUseRadians ? Settings->CollinearAngleThreshold : FMath::DegreesToRadians(Settings->CollinearAngleThreshold);
+	// Max the tolerance with an epsilon to accomodate rounding errors.
+	const double DotProductTolerance = FMath::Max(FMath::Abs(FMath::Cos(DotProductToleranceRad)), UE_DOUBLE_SMALL_NUMBER);
+
+	const EPCGControlPointFuseMode FuseMode = Settings->FuseMode;
+
+	// Get data from the polygon
+	const UE::Geometry::FGeneralPolygon2d& OriginalPolygon = PolygonData->GetPolygon();
+	const FTransform Transform = PolygonData->GetTransform();
+
+	// Early out if polygon is invalid
+	if (OriginalPolygon.GetOuter().VertexCount() < 3)
+	{
+		return;
+	}
 
-		FPCGTaggedData& Output = InContext->OutputData.TaggedData.Emplace_GetRef(InputData);
+	// Copy the polygon (outer + holes)
+	TArray<UE::Geometry::FPolygon2d> WorkingPolygons;
+	WorkingPolygons.Add(OriginalPolygon.GetOuter());
+	WorkingPolygons.Append(OriginalPolygon.GetHoles());
+	
+	// Get the metadata entry keys
+	TArray<PCGMetadataEntryKey> MetadataEntryKeys;
+	MetadataEntryKeys.Append(PolygonData->GetConstVerticesEntryKeys());
 
-		// We only need to create a new data if a point was removed.
-		if (bControlPointWasRemoved)
+	bool bHasChangedData = false;
+	int VertexOffset = 0;
+
+	// For all polygons...
+	for (int PolygonIndex = 0; PolygonIndex < WorkingPolygons.Num(); ++PolygonIndex)
+	{
+		UE::Geometry::FPolygon2d& Polygon = WorkingPolygons[PolygonIndex];
+
+		auto RemoveVertex = [&bHasChangedData, &Polygon, &MetadataEntryKeys, &VertexOffset](int Index)
 		{
-			// Update input keys if a point was removed, to keep them monotonically incremental
-			for (int i = 0; i < ControlPoints.Num(); ++i)
+			bHasChangedData = true;
+			Polygon.RemoveVertex(Index);
+			if (!MetadataEntryKeys.IsEmpty())
 			{
-				ControlPoints[i].InputKey = i;
+				MetadataEntryKeys.RemoveAt(Index + VertexOffset);
 			}
-			
-			UPCGSplineData* NewSplineData = FPCGContext::NewObject_AnyThread<UPCGSplineData>(InContext);
-			NewSplineData->InitializeFromData(InputSplineData);
-			NewSplineData->Initialize(ControlPoints, InputSplineData->IsClosed(), InputSplineData->GetTransform(), std::move(ControlPointsKeys));
+		};
+
+		// Start by removing colocated points
+		int VertexIndex = 0;
+		int CurrentMergeWeight = 2;
 
-			Output.Data = NewSplineData;
+		while (VertexIndex < Polygon.VertexCount() && Polygon.VertexCount() > 3)
+		{
+			const int NextVertexIndex = (VertexIndex + 1) % Polygon.VertexCount();
+
+			if (PCGCleanSplineHelpers::ControlPointsAreColocated(
+				FVector(Polygon[VertexIndex], 0.0),
+				FVector(Polygon[NextVertexIndex], 0.0),
+				Transform,
+				Settings->ColocationDistanceThreshold,
+				Settings->bUseSplineLocalSpace))
+			{
+				switch (FuseMode)
+				{
+				case EPCGControlPointFuseMode::Auto: // auto is always keep first in polygons
+				case EPCGControlPointFuseMode::KeepFirst:
+					RemoveVertex(NextVertexIndex);
+					break;
+				case EPCGControlPointFuseMode::KeepSecond:
+					RemoveVertex(VertexIndex);
+					break;
+				case EPCGControlPointFuseMode::Merge:
+					Polygon[VertexIndex] = FMath::Lerp(Polygon[VertexIndex], Polygon[NextVertexIndex], 1.0 / static_cast<double>(CurrentMergeWeight++));
+					// @todo_pcg: lerp metadata
+					RemoveVertex(NextVertexIndex);
+					break;
+				}
+			}
+			else
+			{
+				++VertexIndex;
+				CurrentMergeWeight = 2;
+			}
 		}
+
+		// Then remove collinear points
+		VertexIndex = 0;
+		while (VertexIndex < Polygon.VertexCount() && Polygon.VertexCount() > 3)
+		{
+			const int PreviousVertexIndex = (VertexIndex + Polygon.VertexCount() - 1) % Polygon.VertexCount();
+			const int NextVertexIndex = (VertexIndex + 1) % Polygon.VertexCount();
+
+			if (PCGCleanSplineHelpers::ControlPointsAreCollinear(
+				FVector(Polygon[PreviousVertexIndex], 0.0),
+				FVector(Polygon[VertexIndex], 0.0),
+				FVector(Polygon[NextVertexIndex], 0.0),
+				DotProductTolerance))
+			{
+				// @todo_pcg: Perform additional validation for metadata similarity, e.g. lerping between properties of previous vs next gives out similar properties to current vertex
+				RemoveVertex(VertexIndex);
+			}
+			else
+			{
+				++VertexIndex;
+			}
+		}
+
+		// @todo_pcg: final step for holes that could be removed if under a certain area/if all 3 remaining points are colocated.
+		// ...
+
+		VertexOffset += Polygon.VertexCount();
 	}
 
-	return true;
+	// Add to the output (even if trivial)
+	FPCGTaggedData& Output = Context->OutputData.TaggedData.Emplace_GetRef(Input);
+
+	// If we've made changes, create the data and push it on the output
+	if (bHasChangedData)
+	{
+		UPCGPolygon2DData* NewPolygonData = FPCGContext::NewObject_AnyThread<UPCGPolygon2DData>(Context);
+		NewPolygonData->InitializeFromData(PolygonData);
+		Output.Data = NewPolygonData;
+
+		UE::Geometry::FGeneralPolygon2d OutputPolygon(MoveTemp(WorkingPolygons[0]));
+		for (int HoleIndex = 1; HoleIndex < WorkingPolygons.Num(); ++HoleIndex)
+		{
+			OutputPolygon.AddHole(MoveTemp(WorkingPolygons[HoleIndex]), /*bCheckContainment=*/false, /*bCheckOrientation=*/false);
+		}
+
+		TConstArrayView<PCGMetadataEntryKey> EntryKeysView(MetadataEntryKeys);
+		NewPolygonData->SetPolygon(MoveTemp(OutputPolygon), EntryKeysView.IsEmpty() ? nullptr : &EntryKeysView);
+		NewPolygonData->SetTransform(Transform);
+	}
 }
 
 #undef LOCTEXT_NAMESPACE
diff --git a/Engine/Plugins/PCG/Source/PCG/Private/Elements/PCGCreateSpline.cpp b/Engine/Plugins/PCG/Source/PCG/Private/Elements/PCGCreateSpline.cpp
--- a/Engine/Plugins/PCG/Source/PCG/Private/Elements/PCGCreateSpline.cpp
+++ b/Engine/Plugins/PCG/Source/PCG/Private/Elements/PCGCreateSpline.cpp
@@ -29,6 +29,7 @@ UPCGCreateSplineSettings::UPCGCreateSplineSettings()
 	{
 		ArriveTangentAttribute = PCGSplineSamplerConstants::ArriveTangentAttributeName;
 		LeaveTangentAttribute = PCGSplineSamplerConstants::LeaveTangentAttributeName;
+		InterpTypeAttribute = PCGSplineSamplerConstants::InterpTypeAttributeName;
 	}
 }
 
@@ -104,6 +105,7 @@ bool FPCGCreateSplineElement::ExecuteInternal(FPCGContext* Context) const
 
 		const FPCGMetadataAttribute<FVector>* ArriveTangentAttribute = nullptr;
 		const FPCGMetadataAttribute<FVector>* LeaveTangentAttribute = nullptr;
+		const FPCGMetadataAttribute<int32>* InterpTypeAttribute = nullptr;
 
 		if (Settings->bApplyCustomTangents && !Settings->bLinear)
 		{
@@ -133,6 +135,22 @@ bool FPCGCreateSplineElement::ExecuteInternal(FPCGContext* Context) const
 			LeaveTangentAttribute = static_cast<const FPCGMetadataAttribute<FVector>*>(LeaveTangentBaseAttribute);
 		}
 
+		if (Settings->bUseInterpTypeAttribute)
+		{
+			const UPCGMetadata* PointMetadata = PointData->ConstMetadata();
+			check(PointMetadata);
+
+			FName LocalPointAttributeName = ((Settings->InterpTypeAttribute == NAME_None) ? PointMetadata->GetLatestAttributeNameOrNone() : Settings->InterpTypeAttribute);
+			const FPCGMetadataAttributeBase* InterpTypeBaseAttribute = PointMetadata->GetConstAttribute(LocalPointAttributeName);
+			if (!InterpTypeBaseAttribute || InterpTypeBaseAttribute->GetTypeId() != PCG::Private::MetadataTypes<int32>::Id)
+			{
+				PCGE_LOG(Error, GraphAndLog, FText::Format(LOCTEXT("PointTypeAttributeMissing", "Point Type Attribute '{0}' does not exist or is not an int32."), FText::FromString(LocalPointAttributeName.ToString())));
+				continue;
+			}
+
+			InterpTypeAttribute = static_cast<const FPCGMetadataAttribute<int32>*>(InterpTypeBaseAttribute);
+		}
+
 		UPCGSplineData* SplineData = FPCGContext::NewObject_AnyThread<UPCGSplineData>(Context);
 		SplineData->InitializeFromData(PointData);
 		AActor* SplineActor = TargetActor;
@@ -147,14 +165,14 @@ bool FPCGCreateSplineElement::ExecuteInternal(FPCGContext* Context) const
 
 		const FTransform SplineActorTransform = SplineActor->GetTransform();
 		
-		ESplinePointType::Type PointType = ESplinePointType::Curve;
+		ESplinePointType::Type FixedPointType = ESplinePointType::Curve;
 		if (Settings->bLinear)
 		{
-			PointType = ESplinePointType::Linear;
+			FixedPointType = ESplinePointType::Linear;
 		}
 		else if (Settings->bApplyCustomTangents)
 		{
-			PointType = ESplinePointType::CurveCustomTangent;
+			FixedPointType = ESplinePointType::CurveCustomTangent;
 		}
 
 		FConstPCGPointValueRanges InRanges(PointData);
@@ -171,7 +189,7 @@ bool FPCGCreateSplineElement::ExecuteInternal(FPCGContext* Context) const
 				LeaveTangentAttribute ? LeaveTangentAttribute->GetValueFromItemKey(PointMetadataEntry) : FVector::ZeroVector,
 				PointTransform.GetRotation().Rotator(),
 				PointTransform.GetScale3D(),
-				PointType);
+				InterpTypeAttribute ? static_cast<ESplinePointType::Type>(InterpTypeAttribute->GetValueFromItemKey(PointMetadataEntry)) : FixedPointType);
 
 			SplineEntryKeys.Emplace(PointMetadataEntry);
 			bHasAValidEntry |= (PointMetadataEntry != PCGInvalidEntryKey);
diff --git a/Engine/Plugins/PCG/Source/PCG/Private/Elements/PCGGetSplineControlPoints.cpp b/Engine/Plugins/PCG/Source/PCG/Private/Elements/PCGGetSplineControlPoints.cpp
--- a/Engine/Plugins/PCG/Source/PCG/Private/Elements/PCGGetSplineControlPoints.cpp
+++ b/Engine/Plugins/PCG/Source/PCG/Private/Elements/PCGGetSplineControlPoints.cpp
@@ -3,6 +3,8 @@
 #include "Elements/PCGGetSplineControlPoints.h"
 
 #include "Data/PCGPolyLineData.h"
+#include "Data/PCGSplineData.h"
+#include "Data/PCGPolygon2DData.h"
 #include "Data/PCGBasePointData.h"
 
 #include "PCGContext.h"
@@ -66,6 +68,8 @@ bool FPCGGetSplineControlPointsElement::ExecuteInternal(FPCGContext* InContext)
 			continue;
 		}
 
+		const UPCGSplineData* SplineData = Cast<UPCGSplineData>(Input.Data);
+
 		const int32 NumPoints = PolylineData->GetNumVertices();
 		const TConstArrayView<PCGMetadataEntryKey> MetadataEntries = PolylineData->GetConstVerticesEntryKeys();
 
@@ -94,6 +98,23 @@ bool FPCGGetSplineControlPointsElement::ExecuteInternal(FPCGContext* InContext)
 
 		FPCGMetadataAttribute<FVector>* LeaveTangentAttribute = CreateAndValidateAttribute(Settings->LeaveTangentAttributeName, FVector::ZeroVector);
 		FPCGMetadataAttribute<FVector>* ArriveTangentAttribute = CreateAndValidateAttribute(Settings->ArriveTangentAttributeName, FVector::ZeroVector);
+		FPCGMetadataAttribute<int32>* InterpTypeAttribute = nullptr;
+		
+		TArray<FSplinePoint> SplinePoints;
+		if (SplineData)
+		{
+			InterpTypeAttribute = CreateAndValidateAttribute(Settings->InterpTypeAttributeName, static_cast<int32>(ESplinePointType::CurveCustomTangent));
+			
+			if (InterpTypeAttribute)
+			{
+				SplinePoints = SplineData->GetSplinePoints();
+			}
+		}
+		else if (const UPCGPolygon2DData* PolygonData = Cast<UPCGPolygon2DData>(Input.Data))
+		{
+			// All segments in a polygon are linear
+			InterpTypeAttribute = CreateAndValidateAttribute(Settings->InterpTypeAttributeName, static_cast<int32>(ESplinePointType::Linear));
+		}
 
 		TPCGValueRange<FTransform> PointTransforms = PointData->GetTransformValueRange();
 		TPCGValueRange<int> PointSeeds = PointData->GetSeedValueRange();
@@ -121,6 +142,11 @@ bool FPCGGetSplineControlPointsElement::ExecuteInternal(FPCGContext* InContext)
 					ArriveTangentAttribute->SetValue(PointMetadataEntries[i], ArriveTangent);
 				}
 			}
+
+			if (SplineData && InterpTypeAttribute)
+			{
+				InterpTypeAttribute->SetValue(PointMetadataEntries[i], SplinePoints[i].Type);
+			}
 		}
 
 		InContext->OutputData.TaggedData.Emplace_GetRef(Input).Data = PointData;
diff --git a/Engine/Plugins/PCG/Source/PCG/Public/Elements/PCGCleanSpline.h b/Engine/Plugins/PCG/Source/PCG/Public/Elements/PCGCleanSpline.h
--- a/Engine/Plugins/PCG/Source/PCG/Public/Elements/PCGCleanSpline.h
+++ b/Engine/Plugins/PCG/Source/PCG/Public/Elements/PCGCleanSpline.h
@@ -6,6 +6,9 @@
 
 #include "PCGCleanSpline.generated.h"
 
+class UPCGSplineData;
+class UPCGPolygon2DData;
+
 UENUM()
 enum class EPCGControlPointFuseMode
 {
@@ -30,6 +33,7 @@ public:
 	virtual EPCGSettingsType GetType() const override { return EPCGSettingsType::Spatial; }
 	virtual void PostEditChangeProperty(FPropertyChangedEvent& PropertyChangedEvent) override;
 #endif // WITH_EDITOR
+	virtual bool HasDynamicPins() const override { return true; }
 
 protected:
 	virtual TArray<FPCGPinProperties> InputPinProperties() const override;
@@ -71,4 +75,8 @@ class FPCGCleanSplineElement : public IPCGElement
 {
 protected:
 	virtual bool ExecuteInternal(FPCGContext* InContext) const override;
+
+private:
+	void ProcessSpline(FPCGContext* InContext, const UPCGCleanSplineSettings* Settings, const FPCGTaggedData& InputData, const UPCGSplineData* InputSplineData) const;
+	void ProcessPolygon(FPCGContext* InContext, const UPCGCleanSplineSettings* Settings, const FPCGTaggedData& InputData, const UPCGPolygon2DData* InputPolygonData) const;
 };
diff --git a/Engine/Plugins/PCG/Source/PCG/Public/Elements/PCGCreateSpline.h b/Engine/Plugins/PCG/Source/PCG/Public/Elements/PCGCreateSpline.h
--- a/Engine/Plugins/PCG/Source/PCG/Public/Elements/PCGCreateSpline.h
+++ b/Engine/Plugins/PCG/Source/PCG/Public/Elements/PCGCreateSpline.h
@@ -63,6 +63,13 @@ public:
 	UPROPERTY(BlueprintReadWrite, EditAnywhere, Category = Settings, meta = (PCG_Overridable, EditCondition = "!bLinear && bApplyCustomTangents"))
 	FName LeaveTangentAttribute;
 
+	UPROPERTY(BlueprintReadWrite, EditAnywhere, Category = Settings, meta = (PCG_Overridable, InlineEditConditionToggle))
+	bool bUseInterpTypeAttribute = false;
+
+	/** Name of the attribute containing the interpolation type, if any. */
+	UPROPERTY(BlueprintReadWrite, EditAnywhere, Category = Settings, meta = (PCG_Overridable, EditCondition = "bUseInterpTypeAttribute"))
+	FName InterpTypeAttribute;
+
 	UPROPERTY(BlueprintReadWrite, EditAnywhere, Category = Settings, meta = (PCG_Overridable, EditCondition = "Mode!= EPCGCreateSplineMode::CreateDataOnly", EditConditionHides))
 	TArray<FName> TagsToAddOnComponents;
 
diff --git a/Engine/Plugins/PCG/Source/PCG/Public/Elements/PCGGetSplineControlPoints.h b/Engine/Plugins/PCG/Source/PCG/Public/Elements/PCGGetSplineControlPoints.h
--- a/Engine/Plugins/PCG/Source/PCG/Public/Elements/PCGGetSplineControlPoints.h
+++ b/Engine/Plugins/PCG/Source/PCG/Public/Elements/PCGGetSplineControlPoints.h
@@ -28,11 +28,15 @@ protected:
 	//~End UPCGSettings interface
 
 public:
-	UPROPERTY(EditAnywhere, BlueprintReadWrite, Category="Settings")
+	UPROPERTY(EditAnywhere, BlueprintReadWrite, Category = Settings)
 	FName ArriveTangentAttributeName = PCGSplineSamplerConstants::ArriveTangentAttributeName;
 	
-	UPROPERTY(EditAnywhere, BlueprintReadWrite, Category="Settings")
+	UPROPERTY(EditAnywhere, BlueprintReadWrite, Category = Settings)
 	FName LeaveTangentAttributeName = PCGSplineSamplerConstants::LeaveTangentAttributeName;
+
+	/** Name of the attribute that will contain the interpolation type. */
+	UPROPERTY(EditAnywhere, BlueprintReadWrite, Category = Settings)
+	FName InterpTypeAttributeName = PCGSplineSamplerConstants::InterpTypeAttributeName;
 };
 
 class FPCGGetSplineControlPointsElement : public IPCGElement
diff --git a/Engine/Plugins/PCG/Source/PCG/Public/Elements/PCGSplineSampler.h b/Engine/Plugins/PCG/Source/PCG/Public/Elements/PCGSplineSampler.h
--- a/Engine/Plugins/PCG/Source/PCG/Public/Elements/PCGSplineSampler.h
+++ b/Engine/Plugins/PCG/Source/PCG/Public/Elements/PCGSplineSampler.h
@@ -28,6 +28,7 @@ namespace PCGSplineSamplerConstants
 	const FName AlphaAttributeName = TEXT("Alpha");
 	const FName DistanceAttributeName = TEXT("Distance");
 	const FName InputKeyAttributeName = TEXT("InputKey");
+	const FName InterpTypeAttributeName = TEXT("InterpType");
 }
 
 UENUM()
