Fix FPCGGraphDocumentInfo::SavedViewOffset not properly initialized

#rnx

#Synced-CL 50918855

diff --git a/Engine/Plugins/PCG/Source/PCG/Public/PCGGraph.h b/Engine/Plugins/PCG/Source/PCG/Public/PCGGraph.h
--- a/Engine/Plugins/PCG/Source/PCG/Public/PCGGraph.h
+++ b/Engine/Plugins/PCG/Source/PCG/Public/PCGGraph.h
@@ -119,7 +119,7 @@ struct FPCGGraphDocumentInfo
 
 	/** Saved view position */
 	UPROPERTY()
-	FVector2f SavedViewOffset;
+	FVector2f SavedViewOffset = FVector2f::ZeroVector;
 
 	/** Saved zoom amount */
 	UPROPERTY()
