[PCG] PCGBasePointData Octree move to lazy allocated unique ptr (Empty octree is ~4kb)
- PCGBasePointData can easily be in the hundreds of thousands meaning ~400mb per 100k point data saved.

#Synced-CL 50653205

diff --git a/Engine/Plugins/PCG/Source/PCG/Private/Data/PCGBasePointData.cpp b/Engine/Plugins/PCG/Source/PCG/Private/Data/PCGBasePointData.cpp
--- a/Engine/Plugins/PCG/Source/PCG/Private/Data/PCGBasePointData.cpp
+++ b/Engine/Plugins/PCG/Source/PCG/Private/Data/PCGBasePointData.cpp
@@ -830,13 +830,14 @@ const PCGPointOctree::FPointOctree& UPCGBasePointData::GetPointOctree() const
 {
 	RebuildOctreeIfNeeded();
 	
-	return PCGPointOctree;
+	check(PCGPointOctree);
+	return *PCGPointOctree;
 }
 
 void UPCGBasePointData::RebuildOctree() const
 {
 	PCG::TScopeLock Lock(CachedDataLock);
-	if (!bOctreeIsDirty)
+	if (!bOctreeIsDirty && PCGPointOctree.IsValid())
 	{
 		return;
 	}
@@ -856,7 +857,7 @@ void UPCGBasePointData::RebuildOctree() const
 		NewOctree.AddElement(PCGPointOctree::FPointRef(PointIndex, PCGPointHelpers::GetDensityBounds(TransformRange[PointIndex], SteepnessRange[PointIndex], BoundsMinRange[PointIndex], BoundsMaxRange[PointIndex])));
 	}
 
-	PCGPointOctree = MoveTemp(NewOctree);
+	PCGPointOctree = MakeUnique<PCGPointOctree::FPointOctree>(MoveTemp(NewOctree));
 	bOctreeIsDirty = false;
 }
 
@@ -1063,7 +1064,9 @@ void UPCGBasePointData::GetResourceSizeEx(FResourceSizeEx& CumulativeResourceSiz
 {
 	Super::GetResourceSizeEx(CumulativeResourceSize);
 
-	CumulativeResourceSize.AddDedicatedSystemMemoryBytes(PCGPointOctree.GetSizeBytes() + sizeof(Bounds));
+	const SIZE_T OctreeSize = PCGPointOctree.IsValid() ? sizeof(*PCGPointOctree) + PCGPointOctree->GetSizeBytes() : 0;
+
+	CumulativeResourceSize.AddDedicatedSystemMemoryBytes(OctreeSize + sizeof(Bounds));
 }
 
 void UPCGBasePointData::InitializeFromActor(AActor* InActor, bool* bOutOptionalSanitizedTagAttributeName)
@@ -1208,6 +1211,7 @@ bool UPCGBasePointData::ProjectPoint(const FTransform& InTransform, const FBox&
 {
 	//TRACE_CPUPROFILER_EVENT_SCOPE(UPCGBasePointData::ProjectPoint);
 	RebuildOctreeIfNeeded();
+	check(PCGPointOctree.IsValid());
 
 	TArray<TPair<int32, FVector::FReal>, TInlineAllocator<4>> Contributions;
 	const bool bSampleInVolume = (InBounds.GetExtent() != FVector::ZeroVector);
@@ -1223,7 +1227,7 @@ bool UPCGBasePointData::ProjectPoint(const FTransform& InTransform, const FBox&
 	if (!bSampleInVolume)
 	{
 		const FVector InPosition = InTransform.GetLocation();
-		PCGPointOctree.FindElementsWithBoundsTest(FBoxCenterAndExtent(InPosition, FVector::Zero()), [this, &InPosition, &Contributions, &TransformRange, &BoundsMinRange, &BoundsMaxRange, &SteepnessRange](const PCGPointOctree::FPointRef& InPointRef)
+		PCGPointOctree->FindElementsWithBoundsTest(FBoxCenterAndExtent(InPosition, FVector::Zero()), [this, &InPosition, &Contributions, &TransformRange, &BoundsMinRange, &BoundsMaxRange, &SteepnessRange](const PCGPointOctree::FPointRef& InPointRef)
 		{
 			Contributions.Emplace(InPointRef.Index, PCGPointHelpers::InverseEuclidianDistance(TransformRange[InPointRef.Index], BoundsMinRange[InPointRef.Index], BoundsMaxRange[InPointRef.Index], SteepnessRange[InPointRef.Index], InPosition));
 		});
@@ -1233,7 +1237,7 @@ bool UPCGBasePointData::ProjectPoint(const FTransform& InTransform, const FBox&
 		FBox TransformedBounds = InBounds.TransformBy(InTransform);
 		FMatrix InTransformInverseMatrix = InTransform.ToMatrixWithScale().Inverse();
 
-		PCGPointOctree.FindElementsWithBoundsTest(FBoxCenterAndExtent(TransformedBounds.GetCenter(), TransformedBounds.GetExtent()), 
+		PCGPointOctree->FindElementsWithBoundsTest(FBoxCenterAndExtent(TransformedBounds.GetCenter(), TransformedBounds.GetExtent()), 
 			[this, bUseBounds, &InBounds, &InTransformInverseMatrix, &Contributions, &TransformRange, &BoundsMinRange, &BoundsMaxRange, &SteepnessRange](const PCGPointOctree::FPointRef& InPointRef)
 			{
 				const FVector::FReal Contribution = bUseBounds ? PCGPointHelpers::VolumeOverlap(TransformRange[InPointRef.Index], BoundsMinRange[InPointRef.Index], BoundsMaxRange[InPointRef.Index], SteepnessRange[InPointRef.Index], InBounds, InTransformInverseMatrix) : 1.0;
diff --git a/Engine/Plugins/PCG/Source/PCG/Public/Data/PCGBasePointData.h b/Engine/Plugins/PCG/Source/PCG/Public/Data/PCGBasePointData.h
--- a/Engine/Plugins/PCG/Source/PCG/Public/Data/PCGBasePointData.h
+++ b/Engine/Plugins/PCG/Source/PCG/Public/Data/PCGBasePointData.h
@@ -395,7 +395,7 @@ protected:
 
 	void RebuildOctreeIfNeeded() const
 	{
-		if (bOctreeIsDirty)
+		if (bOctreeIsDirty || !PCGPointOctree.IsValid())
 		{
 			RebuildOctree();
 		}
@@ -447,8 +447,10 @@ private:
 
 protected:
 
+	// Thread safety on Bounds and Octree is limited to lazy initialization, we do not support multi-threaded modification
+	// of UPCGBasePointData while another thread is using the last generated Octree and/or Bounds.
 	mutable PCG::FRecursiveLock CachedDataLock;
-	mutable PCGPointOctree::FPointOctree PCGPointOctree;
+	mutable TUniquePtr<PCGPointOctree::FPointOctree> PCGPointOctree;
 	mutable FBox Bounds;
 
 	mutable bool bOctreeIsDirty = true;
