[PCG] Modified SPCGEditorViewport to accept FEditorModeTools (for ITF in PVE)

#Synced-CL 50862275

diff --git a/Engine/Plugins/PCG/Source/PCGEditor/Private/Widgets/AssetEditorViewport/SPCGEditorViewport.cpp b/Engine/Plugins/PCG/Source/PCGEditor/Private/Widgets/AssetEditorViewport/SPCGEditorViewport.cpp
--- a/Engine/Plugins/PCG/Source/PCGEditor/Private/Widgets/AssetEditorViewport/SPCGEditorViewport.cpp
+++ b/Engine/Plugins/PCG/Source/PCGEditor/Private/Widgets/AssetEditorViewport/SPCGEditorViewport.cpp
@@ -26,7 +26,7 @@ namespace PCGEditorViewportConstants
 class FPCGEditorViewportClient : public FEditorViewportClient
 {
 public:
-	FPCGEditorViewportClient(const TSharedRef<SPCGEditorViewport>& InAssetEditorViewport);
+	FPCGEditorViewportClient(FEditorModeTools* InModeTools, const TSharedRef<SPCGEditorViewport>& InAssetEditorViewport);
 
 	//~ Begin FViewportClient Interface
 	virtual bool InputKey(const FInputKeyEventArgs& EventArgs) override;
@@ -41,8 +41,8 @@ public:
 	void Reset();
 };
 
-FPCGEditorViewportClient::FPCGEditorViewportClient(const TSharedRef<SPCGEditorViewport>& InAssetEditorViewport)
-	: FEditorViewportClient(/*InModeTools=*/nullptr, InAssetEditorViewport->GetAdvancedPreviewScene(), StaticCastSharedRef<SEditorViewport>(InAssetEditorViewport))
+FPCGEditorViewportClient::FPCGEditorViewportClient(FEditorModeTools* InModeTools, const TSharedRef<SPCGEditorViewport>& InAssetEditorViewport)
+	: FEditorViewportClient(InModeTools, InAssetEditorViewport->GetAdvancedPreviewScene(), StaticCastSharedRef<SEditorViewport>(InAssetEditorViewport))
 {
 	check(PreviewScene);
 
@@ -152,6 +152,8 @@ SPCGEditorViewport::~SPCGEditorViewport()
 
 void SPCGEditorViewport::Construct(const FArguments& InArgs)
 {
+	ModeTools = InArgs._ModeTools;
+	
 	SEditorViewport::Construct(SEditorViewport::FArguments());
 
 	UE::AdvancedPreviewScene::BindDefaultOnSettingsChangedHandler(AdvancedPreviewScene, EditorViewportClient);
@@ -237,7 +239,7 @@ void SPCGEditorViewport::ResetScene()
 
 TSharedRef<FEditorViewportClient> SPCGEditorViewport::MakeEditorViewportClient()
 {
-	EditorViewportClient = MakeShareable(new FPCGEditorViewportClient(SharedThis(this)));
+	EditorViewportClient = MakeShareable(new FPCGEditorViewportClient(ModeTools, SharedThis(this)));
 	EditorViewportClient->SetViewLocation(PCGEditorViewportConstants::DefaultViewLocation);
 	EditorViewportClient->SetViewRotation(PCGEditorViewportConstants::DefaultViewRotation);
 	EditorViewportClient->SetViewLocationForOrbiting(PCGEditorViewportConstants::DefaultViewLocation, PCGEditorViewportConstants::DefaultOrbitDistance);
diff --git a/Engine/Plugins/PCG/Source/PCGEditor/Public/Widgets/AssetEditorViewport/SPCGEditorViewport.h b/Engine/Plugins/PCG/Source/PCGEditor/Public/Widgets/AssetEditorViewport/SPCGEditorViewport.h
--- a/Engine/Plugins/PCG/Source/PCGEditor/Public/Widgets/AssetEditorViewport/SPCGEditorViewport.h
+++ b/Engine/Plugins/PCG/Source/PCGEditor/Public/Widgets/AssetEditorViewport/SPCGEditorViewport.h
@@ -13,7 +13,11 @@ class FPCGEditorViewportClient;
 class SPCGEditorViewport : public SAssetEditorViewport, public ICommonEditorViewportToolbarInfoProvider
 {
 public:
-	SLATE_BEGIN_ARGS(SPCGEditorViewport) {}
+	SLATE_BEGIN_ARGS(SPCGEditorViewport)
+			: _ModeTools(nullptr)
+		{}
+
+		SLATE_ARGUMENT(FEditorModeTools*, ModeTools)
 	SLATE_END_ARGS()
 
 	PCGEDITOR_API SPCGEditorViewport();
@@ -55,6 +59,8 @@ protected:
 protected:
 	TSharedPtr<FPCGEditorViewportClient> EditorViewportClient = nullptr;
 	TSharedPtr<FAdvancedPreviewScene> AdvancedPreviewScene = nullptr;
+	
+	FEditorModeTools* ModeTools = nullptr;
 
 	/** Objects used by the scene. Must be managed for GC. */
 	TArray<TObjectPtr<UObject>> ManagedResources;
