[PCG] Embedded subgraph coloring

#rnx

#Synced-CL 50971911

diff --git a/Engine/Plugins/PCG/Source/PCGEditor/Private/PCGEditorSettings.cpp b/Engine/Plugins/PCG/Source/PCGEditor/Private/PCGEditorSettings.cpp
--- a/Engine/Plugins/PCG/Source/PCGEditor/Private/PCGEditorSettings.cpp
+++ b/Engine/Plugins/PCG/Source/PCGEditor/Private/PCGEditorSettings.cpp
@@ -3,6 +3,7 @@
 #include "PCGEditorSettings.h"
 #include "PCGEditorCommon.h"
 #include "PCGModule.h"
+#include "PCGSubgraph.h"
 
 #include "EdGraph/EdGraphPin.h"
 #include "PCGSettings.h"
@@ -28,6 +29,7 @@ UPCGEditorSettings::UPCGEditorSettings(const FObjectInitializer& ObjectInitializ
 	SamplerNodeColor = FLinearColor(0.0f, 0.0f, 0.0f);
 	SpawnerNodeColor = FLinearColor(0.0f, 1.0f, 0.69f);
 	SubgraphNodeColor = FLinearColor(1.0f, 0.05f, 0.05f);
+	EmbeddedSubgraphNodeColor = FLinearColor(0.225f, 0.f, 0.f);
 	ParamDataNodeColor = FLinearColor(1.0f, 0.38f, 0.02f);
 	DebugNodeColor = FLinearColor(1.0f, 0.0f, 1.0f);
 	ControlFlowNodeColor = FLinearColor(0.0f, 1.0f, 0.0f);
@@ -93,7 +95,17 @@ FLinearColor UPCGEditorSettings::GetColor(UPCGSettings* Settings) const
 		case EPCGSettingsType::Spawner:
 			return SpawnerNodeColor;
 		case EPCGSettingsType::Subgraph:
+		{
+			if (UPCGBaseSubgraphSettings* SubGraphSettings = Cast<UPCGBaseSubgraphSettings>(Settings))
+			{
+				if (UPCGGraph* Subgraph = SubGraphSettings->GetSubgraph(); Subgraph && Subgraph->IsEmbeddedSubgraph())
+				{
+					return EmbeddedSubgraphNodeColor;
+				}
+			}
+			
 			return SubgraphNodeColor;
+		}
 		case EPCGSettingsType::Debug:
 			return DebugNodeColor;
 		case EPCGSettingsType::Param:
diff --git a/Engine/Plugins/PCG/Source/PCGEditor/Public/PCGEditorSettings.h b/Engine/Plugins/PCG/Source/PCGEditor/Public/PCGEditorSettings.h
--- a/Engine/Plugins/PCG/Source/PCGEditor/Public/PCGEditorSettings.h
+++ b/Engine/Plugins/PCG/Source/PCGEditor/Public/PCGEditorSettings.h
@@ -133,6 +133,10 @@ public:
 	UPROPERTY(EditAnywhere, config, Category = Node, meta = (HideAlphaChannel))
 	FLinearColor SubgraphNodeColor;
 
+	/** Color used for embedded subgraph-like operations */
+	UPROPERTY(EditAnywhere, config, Category = Node, meta = (HideAlphaChannel))
+	FLinearColor EmbeddedSubgraphNodeColor;
+
 	/** Color used for Attribute Set operations */
 	UPROPERTY(EditAnywhere, config, Category = Node, meta = (HideAlphaChannel, DisplayName = "Attribute Set Node Color"))
 	FLinearColor ParamDataNodeColor;
