[PCG] Subgraph drag on node support
- From Content browser or Embedded subgraph view

#Synced-CL 50971911

diff --git a/Engine/Plugins/PCG/Source/PCG/Private/PCGGraph.cpp b/Engine/Plugins/PCG/Source/PCG/Private/PCGGraph.cpp
--- a/Engine/Plugins/PCG/Source/PCG/Private/PCGGraph.cpp
+++ b/Engine/Plugins/PCG/Source/PCG/Private/PCGGraph.cpp
@@ -2875,6 +2875,8 @@ void UPCGGraphInstance::SetGraph(UPCGGraphInterface* InGraph)
 		return;
 	}
 
+	Modify();
+
 #if WITH_EDITOR
 	TeardownCallbacks();
 #endif // WITH_EDITOR
diff --git a/Engine/Plugins/PCG/Source/PCGEditor/Private/Schema/PCGEditorGraphSchema.cpp b/Engine/Plugins/PCG/Source/PCGEditor/Private/Schema/PCGEditorGraphSchema.cpp
--- a/Engine/Plugins/PCG/Source/PCGEditor/Private/Schema/PCGEditorGraphSchema.cpp
+++ b/Engine/Plugins/PCG/Source/PCGEditor/Private/Schema/PCGEditorGraphSchema.cpp
@@ -7,6 +7,7 @@
 #include "PCGEdge.h"
 #include "PCGGraph.h"
 #include "PCGPin.h"
+#include "PCGSubgraph.h"
 #include "Elements/PCGCollapseElement.h"
 #include "Elements/PCGCreateSurfaceFromSpline.h"
 #include "Elements/PCGExecuteBlueprint.h"
@@ -745,6 +746,76 @@ void UPCGEditorGraphSchema::GetDataAssetActions(FGraphActionMenuBuilder& ActionM
 	});
 }
 
+bool UPCGEditorGraphSchema::DroppedAssetsOnNodeInternal(const TArray<struct FAssetData>& Assets, const FVector2f& GraphPosition, const UEdGraphNode* Node, bool bValidate, FString& OutTooltipText) const
+{
+	OutTooltipText.Reset();
+
+	if (Assets.Num() != 1)
+	{
+		return false;
+	}
+
+	const UPCGEditorGraphNodeBase* EditorNode = Cast<const UPCGEditorGraphNodeBase>(Node);
+	if (!EditorNode)
+	{
+		return false;
+	}
+
+	const UPCGNode* PCGNode = EditorNode->GetPCGNode();
+	if (!PCGNode)
+	{
+		return false;
+	}
+
+	UPCGSubgraphSettings* SubgraphSettings = Cast<UPCGSubgraphSettings>(PCGNode->GetSettings());
+	if (!SubgraphSettings)
+	{
+		return false;
+	}
+
+	const FAssetData& AssetData = Assets[0];
+	const UClass* AssetClass = AssetData.GetClass();
+
+	if (!AssetClass || !AssetClass->IsChildOf(UPCGGraphInterface::StaticClass()))
+	{
+		return false;
+	}
+
+	if (SubgraphSettings->SubgraphAssetFilter(AssetData))
+	{
+		return false;
+	}
+
+	const TSharedPtr<IAssetReferenceFilter> AssetReferenceFilter = MakeAssetReferenceFilter(Node->GetGraph());
+	if (AssetReferenceFilter.IsValid() && !AssetReferenceFilter->PassesFilter(AssetData))
+	{
+		return false;
+	}
+
+	// Execute action if we are not validating
+	if (!bValidate)
+	{
+		if (UPCGGraphInterface* Graph = Cast<UPCGGraphInterface>(AssetData.GetAsset()))
+		{
+			FScopedTransaction SetGraphTransaction(LOCTEXT("ReplaceSubgraphTransaction", "Replace Subgraph"));
+			SubgraphSettings->SetSubgraph(Graph);
+		}
+	}
+
+	return true;
+}
+
+void UPCGEditorGraphSchema::DroppedAssetsOnNode(const TArray<struct FAssetData>& Assets, const FVector2f& GraphPosition, UEdGraphNode* Node) const
+{
+	FString TooltipText;
+	DroppedAssetsOnNodeInternal(Assets, GraphPosition, Node, /*bValidate=*/false, TooltipText);
+}
+
+void UPCGEditorGraphSchema::GetAssetsNodeHoverMessage(const TArray<struct FAssetData>& Assets, const UEdGraphNode* HoverNode, FString& OutTooltipText, bool& OutOkIcon) const
+{
+	OutOkIcon = DroppedAssetsOnNodeInternal(Assets, HoverNode->GetPosition(), HoverNode, /*bValidate=*/true, OutTooltipText);
+}
+
 void UPCGEditorGraphSchema::DroppedAssetsOnGraph(const TArray<FAssetData>& Assets, const FVector2f& GraphPosition, UEdGraph* Graph) const
 {
 	const TSharedPtr<IAssetReferenceFilter> AssetReferenceFilter = MakeAssetReferenceFilter(Graph);
diff --git a/Engine/Plugins/PCG/Source/PCGEditor/Private/Widgets/SPCGEditorGraphEmbeddedSubgraphsView.cpp b/Engine/Plugins/PCG/Source/PCGEditor/Private/Widgets/SPCGEditorGraphEmbeddedSubgraphsView.cpp
--- a/Engine/Plugins/PCG/Source/PCGEditor/Private/Widgets/SPCGEditorGraphEmbeddedSubgraphsView.cpp
+++ b/Engine/Plugins/PCG/Source/PCGEditor/Private/Widgets/SPCGEditorGraphEmbeddedSubgraphsView.cpp
@@ -71,6 +71,7 @@ namespace PCGEditorGraphEmbeddedSubgraphsView
 		}
 		
 		virtual FReply DroppedOnPanel(const TSharedRef<SWidget>&, const FVector2f&, const FVector2f&, UEdGraph&) override;
+		virtual FReply DroppedOnNode(const FVector2f& ScreenPosition, const FVector2f& GraphPosition) override;
 		virtual void HoverTargetChanged() override;
 		FText GetHoverText() const;
 	};
@@ -145,6 +146,19 @@ FReply PCGEditorGraphEmbeddedSubgraphsView::FEmbeddedSubgraphDragDropAction::Dro
 	return FReply::Handled();
 }
 
+FReply PCGEditorGraphEmbeddedSubgraphsView::FEmbeddedSubgraphDragDropAction::DroppedOnNode(const FVector2f& ScreenPosition, const FVector2f& GraphPosition)
+{
+	UEdGraphNode* CurrentHoveredNode = GetHoveredNode();
+	if (!CurrentHoveredNode || !CurrentHoveredNode->GetGraph() || !CurrentHoveredNode->GetGraph()->GetSchema())
+	{
+		return FReply::Unhandled();
+	}
+
+	FAssetData Asset(Item->WeakEmbeddedSubgraph.Get());
+	CurrentHoveredNode->GetGraph()->GetSchema()->DroppedAssetsOnNode({ Asset }, GraphPosition, CurrentHoveredNode);
+	return FReply::Handled();
+}
+
 FText PCGEditorGraphEmbeddedSubgraphsView::FEmbeddedSubgraphDragDropAction::GetHoverText() const
 {
 	return Item.IsValid() && Item->WeakEmbeddedSubgraph.IsValid() ? Item->WeakEmbeddedSubgraph->Title : LOCTEXT("InvalidHoverText", "Invalid");
diff --git a/Engine/Plugins/PCG/Source/PCGEditor/Public/Schema/PCGEditorGraphSchema.h b/Engine/Plugins/PCG/Source/PCGEditor/Public/Schema/PCGEditorGraphSchema.h
--- a/Engine/Plugins/PCG/Source/PCGEditor/Public/Schema/PCGEditorGraphSchema.h
+++ b/Engine/Plugins/PCG/Source/PCGEditor/Public/Schema/PCGEditorGraphSchema.h
@@ -51,7 +51,9 @@ public:
 	PCGEDITOR_API virtual void BreakSinglePinLink(UEdGraphPin* SourcePin, UEdGraphPin* TargetPin) const override;
 	PCGEDITOR_API virtual FConnectionDrawingPolicy* CreateConnectionDrawingPolicy(int32 InBackLayerID, int32 InFrontLayerID, float InZoomFactor, const FSlateRect& InClippingRect, class FSlateWindowElementList& InDrawElements, class UEdGraph* InGraphObj) const override;
 	PCGEDITOR_API virtual void DroppedAssetsOnGraph(const TArray<FAssetData>& Assets, const FVector2f& GraphPosition, UEdGraph* Graph) const override;
+	PCGEDITOR_API virtual void DroppedAssetsOnNode(const TArray<struct FAssetData>& Assets, const FVector2f& GraphPosition, UEdGraphNode* Node) const override;
 	PCGEDITOR_API virtual void GetAssetsGraphHoverMessage(const TArray<FAssetData>& Assets, const UEdGraph* HoverGraph, FString& OutTooltipText, bool& OutOkIcon) const override;
+	PCGEDITOR_API virtual void GetAssetsNodeHoverMessage(const TArray<struct FAssetData>& Assets, const UEdGraphNode* HoverNode, FString& OutTooltipText, bool& OutOkIcon) const override;
 	PCGEDITOR_API virtual void OnPinConnectionDoubleCicked(UEdGraphPin* PinA, UEdGraphPin* PinB, const FVector2f& GraphPosition) const override;
 	PCGEDITOR_API virtual void GetGraphDisplayInformation(const UEdGraph& Graph, /*out*/ FGraphDisplayInfo& DisplayInfo) const override;
 	//~ End EdGraphSchema Interface
@@ -78,6 +80,8 @@ protected:
 
 	PCGEDITOR_API virtual bool TryCreateConnectionInternal(UEdGraphPin* A, UEdGraphPin* B, bool bAddConversionNodeIfNeeded) const;
 
+	PCGEDITOR_API virtual bool DroppedAssetsOnNodeInternal(const TArray<struct FAssetData>& Assets, const FVector2f& GraphPosition, const UEdGraphNode* Node, bool bValidate, FString& OutTooltipText) const;
+
 private:
 	static TSharedPtr<IAssetReferenceFilter> MakeAssetReferenceFilter(const UEdGraph* Graph);
 };
