Fixed User Parameter value changes causing a cosmetic-only change. A previous CL changed the inspected property from the actually edited property to the head property, changing execution flow.
Kept the head property inspection, but in the User Parameters case checking for the actually edited property. If it's different, it's likely a User Parameter, so we verify value changes

#Synced-CL 50932743

diff --git a/Engine/Plugins/PCG/Source/PCG/Private/PCGGraph.cpp b/Engine/Plugins/PCG/Source/PCG/Private/PCGGraph.cpp
--- a/Engine/Plugins/PCG/Source/PCG/Private/PCGGraph.cpp
+++ b/Engine/Plugins/PCG/Source/PCG/Private/PCGGraph.cpp
@@ -2173,115 +2173,128 @@ void UPCGGraph::PostEditChangeChainProperty(FPropertyChangedChainEvent& Property
 	{
 		return;
 	}
+	
+	const FProperty* EditedProperty = PropertyChangedEvent.Property;
+	const FName EditedPropertyName = EditedProperty->GetFName();
+	
+	const FEditPropertyChain::TDoubleLinkedListNode* HeadPropertyNode = PropertyChangedEvent.PropertyChain.GetHead();
+	const FProperty* HeadProperty = HeadPropertyNode ? HeadPropertyNode->GetValue() : nullptr;
 
-	const FEditPropertyChain::TDoubleLinkedListNode* MemberPropertyNode = PropertyChangedEvent.PropertyChain.GetHead();
-	const FProperty* MemberProperty = MemberPropertyNode ? MemberPropertyNode->GetValue() : nullptr;
-
-	if (!MemberProperty)
+	if (!HeadProperty)
 	{
 		return;
 	}
 
-	const FName MemberPropertyName = MemberProperty->GetFName();
+	const FName HeadPropertyName = HeadProperty->GetFName();
+	
 	bool bHiGenSettingsChanged = false;
 
-	if (MemberPropertyName == GET_MEMBER_NAME_CHECKED(UPCGGraph, bLandscapeUsesMetadata) ||
-		MemberPropertyName == GET_MEMBER_NAME_CHECKED(UPCGGraph, bHasDefaultConstructedInputs))
+	if (HeadPropertyName == GET_MEMBER_NAME_CHECKED(UPCGGraph, bLandscapeUsesMetadata) ||
+		HeadPropertyName == GET_MEMBER_NAME_CHECKED(UPCGGraph, bHasDefaultConstructedInputs))
 	{
 		NotifyGraphChanged(EPCGChangeType::Input);
 	}
-	else if (MemberPropertyName == GET_MEMBER_NAME_CHECKED(UPCGGraph, UserParameters))
+	else if (HeadPropertyName == GET_MEMBER_NAME_CHECKED(UPCGGraph, UserParameters))
 	{
-		EPCGGraphParameterEvent ChangeType = EPCGGraphParameterEvent::None;
-		const int32 NumberOfUserParametersPreEdit = PreviousPropertyBag ? PreviousPropertyBag->GetPropertyDescs().Num() : 0;
-		const int32 NumberOfUserParametersPostEdit = UserParameters.GetNumPropertiesInBag();
-		FName ChangedPropertyName = NAME_None;
-
-		if (NumberOfUserParametersPostEdit > NumberOfUserParametersPreEdit)
+		// If the edited property is not the User Parameter itself, it's a parameter inside
+		if (EditedProperty != HeadProperty)
 		{
-			ChangeType = EPCGGraphParameterEvent::Added;
+			VerifyAndUpdateIfGraphParameterValueChanged(PropertyChangedEvent);
 		}
-		else if (NumberOfUserParametersPostEdit < NumberOfUserParametersPreEdit)
-		{
-			// Removed, but not knowing if it is used or not yet.
-			ChangeType = EPCGGraphParameterEvent::RemovedUnused;
-		}
-		else if (PreviousPropertyBag) // && NumberOfUserParametersPostEdit == NumberOfUserParametersPreEdit
+		// If it's the same, it was likely triggered by ApplyChangesToPropertyDescs, which means a structural change to the parameters
+		else
 		{
-			for (int32 i = 0; i < NumberOfUserParametersPostEdit; ++i)
-			{
-				const FPropertyBagPropertyDesc& PreDesc = PreviousPropertyBag->GetPropertyDescs()[i];
-				const FPropertyBagPropertyDesc& PostDesc = UserParameters.GetPropertyBagStruct()->GetPropertyDescs()[i];
+			EPCGGraphParameterEvent ChangeType = EPCGGraphParameterEvent::None;
+			const int32 NumberOfUserParametersPreEdit = PreviousPropertyBag ? PreviousPropertyBag->GetPropertyDescs().Num() : 0;
+			const int32 NumberOfUserParametersPostEdit = UserParameters.GetNumPropertiesInBag();
+			FName ChangedPropertyName = NAME_None;
 
-				// Not Same ID -> Moved
-				if (PreDesc.ID != PostDesc.ID)
-				{
-					ChangeType = EPCGGraphParameterEvent::PropertyMoved;
-					break;
-				}
-				// Same ID but different name -> Renamed
-				else if (PreDesc.Name != PostDesc.Name)
-				{
-					ChangeType = EPCGGraphParameterEvent::PropertyRenamed;
-					ChangedPropertyName = PostDesc.Name;
-					break;
-				}
-				// Same name but different type -> Type modified
-				else if (!PostDesc.CompatibleType(PreDesc))
-				{
-					ChangeType = EPCGGraphParameterEvent::PropertyTypeModified;
-					ChangedPropertyName = PostDesc.Name;
-					break;
-				}
-				// Category changed for the property
-				else if (PreDesc.CachedProperty && PostDesc.CachedProperty)
+			if (NumberOfUserParametersPostEdit > NumberOfUserParametersPreEdit)
+			{
+				ChangeType = EPCGGraphParameterEvent::Added;
+			}
+			else if (NumberOfUserParametersPostEdit < NumberOfUserParametersPreEdit)
+			{
+				// Removed, but not knowing if it is used or not yet.
+				ChangeType = EPCGGraphParameterEvent::RemovedUnused;
+			}
+			else if (PreviousPropertyBag) // && NumberOfUserParametersPostEdit == NumberOfUserParametersPreEdit
+			{
+				for (int32 i = 0; i < NumberOfUserParametersPostEdit; ++i)
 				{
-					using namespace PCGPropertyHelpers::Constants;
-					const bool bPreCategoriesEnabled = PreDesc.CachedProperty->HasMetaData(EnableCategoriesMetadataName);
-					const bool bPostCategoriesEnabled = PostDesc.CachedProperty->HasMetaData(EnableCategoriesMetadataName);
-					// Check first if categories or enabled and then compare if they have the same category.
-					if (bPreCategoriesEnabled != bPostCategoriesEnabled
-						|| (bPreCategoriesEnabled
-							&& bPostCategoriesEnabled
-							&& !PreDesc.CachedProperty->GetMetaData(CategoryMetadataName).Equals(PostDesc.CachedProperty->GetMetaData(CategoryMetadataName))))
+					const FPropertyBagPropertyDesc& PreDesc = PreviousPropertyBag->GetPropertyDescs()[i];
+					const FPropertyBagPropertyDesc& PostDesc = UserParameters.GetPropertyBagStruct()->GetPropertyDescs()[i];
+
+					// Not Same ID -> Moved
+					if (PreDesc.ID != PostDesc.ID)
+					{
+						ChangeType = EPCGGraphParameterEvent::PropertyMoved;
+						break;
+					}
+					// Same ID but different name -> Renamed
+					else if (PreDesc.Name != PostDesc.Name)
+					{
+						ChangeType = EPCGGraphParameterEvent::PropertyRenamed;
+						ChangedPropertyName = PostDesc.Name;
+						break;
+					}
+					// Same name but different type -> Type modified
+					else if (!PostDesc.CompatibleType(PreDesc))
 					{
-						ChangeType = EPCGGraphParameterEvent::CategoryChanged;
+						ChangeType = EPCGGraphParameterEvent::PropertyTypeModified;
 						ChangedPropertyName = PostDesc.Name;
 						break;
 					}
+					// Category changed for the property
+					else if (PreDesc.CachedProperty && PostDesc.CachedProperty)
+					{
+						using namespace PCGPropertyHelpers::Constants;
+						const bool bPreCategoriesEnabled = PreDesc.CachedProperty->HasMetaData(EnableCategoriesMetadataName);
+						const bool bPostCategoriesEnabled = PostDesc.CachedProperty->HasMetaData(EnableCategoriesMetadataName);
+						// Check first if categories or enabled and then compare if they have the same category.
+						if (bPreCategoriesEnabled != bPostCategoriesEnabled
+							|| (bPreCategoriesEnabled
+								&& bPostCategoriesEnabled
+								&& !PreDesc.CachedProperty->GetMetaData(CategoryMetadataName).Equals(PostDesc.CachedProperty->GetMetaData(CategoryMetadataName))))
+						{
+							ChangeType = EPCGGraphParameterEvent::CategoryChanged;
+							ChangedPropertyName = PostDesc.Name;
+							break;
+						}
+					}
 				}
 			}
-		}
 
-		if (ChangeType != EPCGGraphParameterEvent::None)
-		{
-			OnGraphParametersChanged(ChangeType, ChangedPropertyName);
+			if (ChangeType != EPCGGraphParameterEvent::None)
+			{
+				OnGraphParametersChanged(ChangeType, ChangedPropertyName);
+			}
 		}
 	}
-	else if (MemberPropertyName == GET_MEMBER_NAME_CHECKED(UPCGGraph, HiGenGridSize)
-		|| MemberPropertyName == GET_MEMBER_NAME_CHECKED(UPCGGraph, bUseHierarchicalGeneration)
-		|| MemberPropertyName == GET_MEMBER_NAME_CHECKED(UPCGGraph, bUse2DGrid)
-		|| MemberPropertyName == GET_MEMBER_NAME_CHECKED(UPCGGraph, bUseActorComponentlessGeneration))
+	else if (HeadPropertyName == GET_MEMBER_NAME_CHECKED(UPCGGraph, HiGenGridSize)
+		|| HeadPropertyName == GET_MEMBER_NAME_CHECKED(UPCGGraph, bUseHierarchicalGeneration)
+		|| HeadPropertyName == GET_MEMBER_NAME_CHECKED(UPCGGraph, bUse2DGrid)
+		|| HeadPropertyName == GET_MEMBER_NAME_CHECKED(UPCGGraph, bUseActorComponentlessGeneration))
 	{
 		bHiGenSettingsChanged = true;
 	}
-	else if (MemberPropertyName == GET_MEMBER_NAME_CHECKED(UPCGGraph, HiGenGridSizeMultiplier))
+	else if (HeadPropertyName == GET_MEMBER_NAME_CHECKED(UPCGGraph, HiGenGridSizeMultiplier))
 	{
 		HiGenGridSizeMultiplier = PCGHiGenGrid::SanitizeGridSizeMultiplier(HiGenGridSizeMultiplier);
 
 		bHiGenSettingsChanged = true;
 	}
-	else if (MemberPropertyName == GET_MEMBER_NAME_CHECKED(UPCGGraphInterface, Title)
-		|| MemberPropertyName == GET_MEMBER_NAME_CHECKED(UPCGGraphInterface, bOverrideTitle)
-		|| MemberPropertyName == GET_MEMBER_NAME_CHECKED(UPCGGraphInterface, Color)
-		|| MemberPropertyName == GET_MEMBER_NAME_CHECKED(UPCGGraphInterface, bOverrideColor)
-		|| MemberPropertyName == GET_MEMBER_NAME_CHECKED(UPCGGraph, bIsStandaloneGraph)
-		|| MemberPropertyName == GET_MEMBER_NAME_CHECKED(UPCGGraph, Category)
-		|| MemberPropertyName == GET_MEMBER_NAME_CHECKED(UPCGGraph, Description))
+	else if (HeadPropertyName == GET_MEMBER_NAME_CHECKED(UPCGGraphInterface, Title)
+		|| HeadPropertyName == GET_MEMBER_NAME_CHECKED(UPCGGraphInterface, bOverrideTitle)
+		|| HeadPropertyName == GET_MEMBER_NAME_CHECKED(UPCGGraphInterface, Color)
+		|| HeadPropertyName == GET_MEMBER_NAME_CHECKED(UPCGGraphInterface, bOverrideColor)
+		|| HeadPropertyName == GET_MEMBER_NAME_CHECKED(UPCGGraph, bIsStandaloneGraph)
+		|| HeadPropertyName == GET_MEMBER_NAME_CHECKED(UPCGGraph, Category)
+		|| HeadPropertyName == GET_MEMBER_NAME_CHECKED(UPCGGraph, Description))
 	{
 		NotifyGraphChanged(EPCGChangeType::Cosmetic | EPCGChangeType::GraphCustomization);
 	}
-	else if (MemberPropertyName == GET_MEMBER_NAME_CHECKED(UPCGGraph, GenerationRadii))
+	else if (HeadPropertyName == GET_MEMBER_NAME_CHECKED(UPCGGraph, GenerationRadii))
 	{
 		GenerationRadii.ComputeHash();
 	}
@@ -2289,10 +2302,6 @@ void UPCGGraph::PostEditChangeChainProperty(FPropertyChangedChainEvent& Property
 	{
 		NotifyGraphChanged(EPCGChangeType::GraphCustomization);
 	}
-	else if (VerifyAndUpdateIfGraphParameterValueChanged(PropertyChangedEvent))
-	{
-		// Handled in the function
-	}
 
 	if (bHiGenSettingsChanged)
 	{
