[PCG] Small adjustments to generic attribute visualizer

* Add helpers on the generic attribute to know which type of container is the attribute
* Setup the generic attribute to the column info if it is a container
* Make sure to keep the default attribute list view when we can, otherwise propose the temporary viewer.

#rnx
#Synced-CL 50862016

diff --git a/Engine/Plugins/PCG/Source/PCG/Private/Metadata/PCGMetadataAttributeGeneric.cpp b/Engine/Plugins/PCG/Source/PCG/Private/Metadata/PCGMetadataAttributeGeneric.cpp
--- a/Engine/Plugins/PCG/Source/PCG/Private/Metadata/PCGMetadataAttributeGeneric.cpp
+++ b/Engine/Plugins/PCG/Source/PCG/Private/Metadata/PCGMetadataAttributeGeneric.cpp
@@ -445,6 +445,26 @@ void FPCGMetadataAttributeGeneric::AddReferencedObjects(FReferenceCollector& Col
 	// We'll have to see the cost of this, and if we don't hold object references for too long.
 }
 
+bool FPCGMetadataAttributeGeneric::IsSingleValue() const
+{
+	return CachedDesc.ContainerTypes.IsEmpty() || CachedDesc.ContainerTypes[0] == EPCGMetadataAttributeContainerTypes::None;
+}
+
+bool FPCGMetadataAttributeGeneric::IsArray() const
+{
+	return !CachedDesc.ContainerTypes.IsEmpty() && CachedDesc.ContainerTypes[0] == EPCGMetadataAttributeContainerTypes::Array;
+}
+
+bool FPCGMetadataAttributeGeneric::IsSet() const
+{
+	return !CachedDesc.ContainerTypes.IsEmpty() && CachedDesc.ContainerTypes[0] == EPCGMetadataAttributeContainerTypes::Set;
+}
+
+bool FPCGMetadataAttributeGeneric::IsMap() const
+{
+	return !CachedDesc.ContainerTypes.IsEmpty() && CachedDesc.ContainerTypes[0] == EPCGMetadataAttributeContainerTypes::Map;
+}
+
 TArray<PCGMetadataValueKey> FPCGMetadataAttributeGeneric::AddValues_Internal(FInValues InValues, int32 Count, int32 ExistingIndex, bool bLockless)
 {
 	TArray<PCGMetadataValueKey> ValueKeys;
diff --git a/Engine/Plugins/PCG/Source/PCG/Public/Metadata/PCGMetadataAttributeGeneric.h b/Engine/Plugins/PCG/Source/PCG/Public/Metadata/PCGMetadataAttributeGeneric.h
--- a/Engine/Plugins/PCG/Source/PCG/Public/Metadata/PCGMetadataAttributeGeneric.h
+++ b/Engine/Plugins/PCG/Source/PCG/Public/Metadata/PCGMetadataAttributeGeneric.h
@@ -68,6 +68,18 @@ public:
 
 	virtual bool IsGeneric() const override { return true; }
 	
+	/** Returns true if the attribute is not a container (Array/Set/Map). */
+	PCG_API bool IsSingleValue() const;
+
+	/** Returns true if the attribute is an array. */
+	PCG_API bool IsArray() const;
+
+	/** Returns true if the attribute is a set. */
+	PCG_API bool IsSet() const;
+
+	/** Returns true if the attribute is a map. */
+	PCG_API bool IsMap() const;
+
 	const FPCGMetadataAttributeDesc& GetAttributeDesc() const { return CachedDesc; }
 
 	PCG_API virtual PCGMetadataValueKey GetValueKeyOffsetForChild() const override;
diff --git a/Engine/Plugins/PCG/Source/PCGEditor/Private/DataVisualizations/PCGDataVisualizationHelpers.cpp b/Engine/Plugins/PCG/Source/PCGEditor/Private/DataVisualizations/PCGDataVisualizationHelpers.cpp
--- a/Engine/Plugins/PCG/Source/PCGEditor/Private/DataVisualizations/PCGDataVisualizationHelpers.cpp
+++ b/Engine/Plugins/PCG/Source/PCGEditor/Private/DataVisualizations/PCGDataVisualizationHelpers.cpp
@@ -166,19 +166,26 @@ namespace PCGDataVisualizationHelpers
 
 		for (int32 I = 0; I < AttributeNames.Num(); I++)
 		{
+			const FPCGMetadataAttributeGeneric* GenericAttr = MetadataDomain->GetConstGenericAttribute(AttributeNames[I]);
 			const int32 ColumnInfosNum = OutInfo.ColumnInfos.Num();
-			PCGMetadataAttribute::CallbackWithRightType(static_cast<uint16>(AttributeTypes[I]), AddInfo, AttributeNames[I]);
-			if (ColumnInfosNum == OutInfo.ColumnInfos.Num())
+
+			// If we don't have a generic attribute or the generic attribute is a single value (not in a container), try to add the columns for that type. 
+			if (!GenericAttr || GenericAttr->IsSingleValue())
 			{
-				if (const FPCGMetadataAttributeGeneric* GenericAttr = MetadataDomain->GetConstGenericAttribute(AttributeNames[I]))
-				{
-					FPCGTableVisualizerColumnInfo& ColumnInfo = OutInfo.ColumnInfos.Emplace_GetRef();
-					ColumnInfo.Label = FText::FromName(AttributeNames[I]);
-					ColumnInfo.Id = AttributeNames[I] == NAME_None ? TEXT("AttrNone") : AttributeNames[I];
-					ColumnInfo.GenericAttribute = GenericAttr;
-					Selector.SetAttributeName(AttributeNames[I]);
-					ColumnInfo.AccessorKeys = TSharedPtr<const IPCGAttributeAccessorKeys>(PCGAttributeAccessorHelpers::CreateConstKeys(Data, Selector).Release());
-				}
+				PCGMetadataAttribute::CallbackWithRightType(static_cast<uint16>(AttributeTypes[I]), AddInfo, AttributeNames[I]);
+			}
+
+			// If it failed (there was no column added), and we have a generic attribute, then configure the column to get the generic attribute.
+			// We don't want to do it for all attributes as we still want the basic types to be displayed as before, even if they are generic attributes
+			// and not typed attributes.
+			if (ColumnInfosNum == OutInfo.ColumnInfos.Num() && GenericAttr)
+			{
+				FPCGTableVisualizerColumnInfo& ColumnInfo = OutInfo.ColumnInfos.Emplace_GetRef();
+				ColumnInfo.Label = FText::FromName(AttributeNames[I]);
+				ColumnInfo.Id = AttributeNames[I] == NAME_None ? TEXT("AttrNone") : AttributeNames[I];
+				ColumnInfo.GenericAttribute = GenericAttr;
+				Selector.SetAttributeName(AttributeNames[I]);
+				ColumnInfo.AccessorKeys = TSharedPtr<const IPCGAttributeAccessorKeys>(PCGAttributeAccessorHelpers::CreateConstKeys(Data, Selector).Release());
 			}
 		}
 	}
diff --git a/Engine/Plugins/PCG/Source/PCGEditor/Private/Widgets/SPCGEditorGraphAttributeListView.cpp b/Engine/Plugins/PCG/Source/PCGEditor/Private/Widgets/SPCGEditorGraphAttributeListView.cpp
--- a/Engine/Plugins/PCG/Source/PCGEditor/Private/Widgets/SPCGEditorGraphAttributeListView.cpp
+++ b/Engine/Plugins/PCG/Source/PCGEditor/Private/Widgets/SPCGEditorGraphAttributeListView.cpp
@@ -151,8 +151,9 @@ TSharedRef<SWidget> SPCGListViewItemRow::GenerateWidgetForColumn(const FName& Co
 
 	if (FPCGColumnData* PCGColumnData = SharedAttributeListView->PCGColumnData.Find(ColumnId))
 	{
+		const bool bShouldUseGenericAttributeViewer = PCGColumnData->GenericAttribute != nullptr;
 		// We have to make sure the data is still valid otherwise this will crash.
-		if (ensure(SharedAttributeListView->DataPtr) && PCGColumnData->DataAccessor.IsValid() && PCGColumnData->DataKeys.IsValid())
+		if (ensure(SharedAttributeListView->DataPtr) && PCGColumnData->DataAccessor.IsValid() && PCGColumnData->DataKeys.IsValid() && !bShouldUseGenericAttributeViewer)
 		{
 			bDataRetrievedWithAccessors = true;
 			const int32 Index = InternalItem->Index;
