[PCG] Fix crash where LS component used before it is registered.

#rnx
#Synced-CL 51091775

diff --git a/Engine/Plugins/PCG/Source/PCG/Private/Elements/PCGGenerateLandscapeTextures.cpp b/Engine/Plugins/PCG/Source/PCG/Private/Elements/PCGGenerateLandscapeTextures.cpp
--- a/Engine/Plugins/PCG/Source/PCG/Private/Elements/PCGGenerateLandscapeTextures.cpp
+++ b/Engine/Plugins/PCG/Source/PCG/Private/Elements/PCGGenerateLandscapeTextures.cpp
@@ -923,6 +923,7 @@ bool FPCGGenerateLandscapeTexturesElement::ExecuteInternal(FPCGContext* InContex
 				return true;
 			}
 
+			bAllReady &= LandscapeComponent->IsRegistered();
 			bAllReady &= LandscapeComponent->CanRenderGrassMap();
 		}
 
@@ -988,7 +989,6 @@ bool FPCGGenerateLandscapeTexturesElement::ExecuteInternal(FPCGContext* InContex
 		LandscapeComponentKeys.Reserve(FMath::Min(Context->LandscapeComponents.Num(), Settings->NumLandscapeComponentsGeneratedPerFrame));
 
 		FBox LandscapeComponentBounds = FBox(EForceInit::ForceInit);
-		LandscapeComponentKeys.Reserve(Context->LandscapeComponents.Num());
 
 		int32 NumComponentsToGenerate = 0;
 
@@ -1019,7 +1019,7 @@ bool FPCGGenerateLandscapeTexturesElement::ExecuteInternal(FPCGContext* InContex
 		}
 
 		TArray<FIntVector2> LandscapeTileCoords;
-		LandscapeComponents.Reserve(LandscapeComponents.Num());
+		LandscapeTileCoords.Reserve(LandscapeComponents.Num());
 
 		for (ULandscapeComponent* LandscapeComponent : LandscapeComponents)
 		{
