[PCG] Fix misaligned textures in Generate Landscape Textures node when landscape bounds and generation bounds would lie on different texel grids
* Fixes error in height and grass map texture sampling, which causes floating instances and incorrect grass variety selection

#rnx

#Synced-CL 51068116

diff --git a/Engine/Plugins/PCG/Source/PCG/Private/Elements/PCGGenerateLandscapeTextures.cpp b/Engine/Plugins/PCG/Source/PCG/Private/Elements/PCGGenerateLandscapeTextures.cpp
--- a/Engine/Plugins/PCG/Source/PCG/Private/Elements/PCGGenerateLandscapeTextures.cpp
+++ b/Engine/Plugins/PCG/Source/PCG/Private/Elements/PCGGenerateLandscapeTextures.cpp
@@ -141,19 +141,19 @@ namespace PCGGenerateLandscapeTextures
 		const uint32 LandscapeComponentResolution = ComponentSizeQuads + 1;
 
 		const FBox& TotalBounds = InData.LandscapeComponentBounds;
-		const FBox& GenerationBounds = InData.Context->GenerationBounds;
+		const FBox& OutputBounds = InData.Context->OutputTextureBounds;
 
-		UE_LOG(LogPCG, Verbose, TEXT("PCGUnpackGrassMap: %s, Cell bounds (%.2f, %.2f -> %.2f, %.2f), Total bounds (%.2f, %.2f -> %.2f, %.2f), efficiency %.2f"),
+		UE_LOG(LogPCG, Verbose, TEXT("PCGUnpackGrassMap: %s, Output texture bounds (%.2f, %.2f -> %.2f, %.2f), Total bounds (%.2f, %.2f -> %.2f, %.2f), efficiency %.2f"),
 			*InData.Context->GetExecutionSourceName(),
-			GenerationBounds.Min.X, GenerationBounds.Min.Y, GenerationBounds.Max.X, GenerationBounds.Max.Y,
+			OutputBounds.Min.X, OutputBounds.Min.Y, OutputBounds.Max.X, OutputBounds.Max.Y,
 			TotalBounds.Min.X, TotalBounds.Min.Y, TotalBounds.Max.X, TotalBounds.Max.Y,
-			(GenerationBounds.Max.X - GenerationBounds.Min.X) * (GenerationBounds.Max.Y - GenerationBounds.Min.Y) / ((TotalBounds.Max.X - TotalBounds.Min.X) * (TotalBounds.Max.Y - TotalBounds.Min.Y)));
+			(OutputBounds.Max.X - OutputBounds.Min.X) * (OutputBounds.Max.Y - OutputBounds.Min.Y) / ((TotalBounds.Max.X - TotalBounds.Min.X) * (TotalBounds.Max.Y - TotalBounds.Min.Y)));
 
 		FIntPoint TileGridStart, TileGridEnd;
-		TileGridStart.X = FMath::Floor((TotalBounds.Min.X - GenerationBounds.Min.X) / QuadSizeWorld);
-		TileGridStart.Y = FMath::Floor((TotalBounds.Min.Y - GenerationBounds.Min.Y) / QuadSizeWorld);
-		TileGridEnd.X = FMath::CeilToInt((TotalBounds.Max.X - GenerationBounds.Min.X) / QuadSizeWorld);
-		TileGridEnd.Y = FMath::CeilToInt((TotalBounds.Max.Y - GenerationBounds.Min.Y) / QuadSizeWorld);
+		TileGridStart.X = FMath::Floor((TotalBounds.Min.X - OutputBounds.Min.X) / QuadSizeWorld);
+		TileGridStart.Y = FMath::Floor((TotalBounds.Min.Y - OutputBounds.Min.Y) / QuadSizeWorld);
+		TileGridEnd.X = FMath::CeilToInt((TotalBounds.Max.X - OutputBounds.Min.X) / QuadSizeWorld);
+		TileGridEnd.Y = FMath::CeilToInt((TotalBounds.Max.Y - OutputBounds.Min.Y) / QuadSizeWorld);
 
 		const FVector TotalLandscapeComponentExtent = InData.LandscapeComponentBounds.GetExtent() * 2;
 		const uint32 NumTilesX = FMath::RoundToInt(TotalLandscapeComponentExtent.X / InData.Context->LandscapeComponentExtent);
@@ -852,21 +852,55 @@ bool FPCGGenerateLandscapeTexturesElement::ExecuteInternal(FPCGContext* InContex
 			Context->LandscapeScaleZ = Landscape->GetRootComponent()->GetRelativeScale3D().Z;
 
 			const double QuadSizeWorld = Context->LandscapeComponentExtent / Landscape->ComponentSizeQuads;
+			check(!FMath::IsNearlyZero(QuadSizeWorld));
+
+			// Texel-space offset of the generation bounds from the landscape component bounds.
+			FVector2D OutputOffsetTexelsF;
+			OutputOffsetTexelsF.X = (Context->GenerationBounds.Min.X - Context->TotalLandscapeComponentBounds.Min.X) / QuadSizeWorld;
+			OutputOffsetTexelsF.Y = (Context->GenerationBounds.Min.Y - Context->TotalLandscapeComponentBounds.Min.Y) / QuadSizeWorld;
+ 
+			// Round down so that we aren't in between texels.
+			FVector2D OutputOffsetTexels;
+			OutputOffsetTexels.X = FMath::FloorToInt(OutputOffsetTexelsF.X);
+			OutputOffsetTexels.Y = FMath::FloorToInt(OutputOffsetTexelsF.Y);
+
+			// A non-zero fraction means the generation bounds and landscape bounds don't lie on the same texel grid, so we'll have to grow the
+			// output texture extents by one texel in order to cover the full generation bounds.
+			const FVector2D OutputOffsetTexelsFrac = OutputOffsetTexelsF - OutputOffsetTexels;
+
+			// Extents of the output should be the generation bounds, rounded up so that we aren't in between texels.
+			// Need to grow the extents by one texel if the generation bounds don't line up with the landscape component bounds in texel space.
+			FIntPoint OutputExtentTexels;
+			OutputExtentTexels.X = FMath::CeilToInt((Context->GenerationBounds.Max.X - Context->GenerationBounds.Min.X) / QuadSizeWorld) + FMath::CeilToInt(OutputOffsetTexelsFrac.X);
+			OutputExtentTexels.Y = FMath::CeilToInt((Context->GenerationBounds.Max.Y - Context->GenerationBounds.Min.Y) / QuadSizeWorld) + FMath::CeilToInt(OutputOffsetTexelsFrac.Y);
 
 			// Dilate the horizontal extents by half a texel in both X and Y. This will place the texel centers on the vert grid.
+			OutputOffsetTexels.X -= 0.5;
+			OutputOffsetTexels.Y -= 0.5;
+			OutputExtentTexels.X += 1;
+			OutputExtentTexels.Y += 1;
+
+			// Apply offset (in world space) to the landscape component bounds to find the output bounds.
+			FVector2D OutputBoundsMin;
+			OutputBoundsMin.X = (OutputOffsetTexels.X * QuadSizeWorld) + Context->TotalLandscapeComponentBounds.Min.X;
+			OutputBoundsMin.Y = (OutputOffsetTexels.Y * QuadSizeWorld) + Context->TotalLandscapeComponentBounds.Min.Y;
+
+			// Convert texel space back to world space.
+			const FVector2D OutputExtents = OutputExtentTexels * QuadSizeWorld;
+
 			FBox OutputBounds;
-			OutputBounds.Min.X = (FMath::Floor(Context->GenerationBounds.Min.X / QuadSizeWorld) - 0.5) * QuadSizeWorld;
-			OutputBounds.Min.Y = (FMath::Floor(Context->GenerationBounds.Min.Y / QuadSizeWorld) - 0.5) * QuadSizeWorld;
-			OutputBounds.Max.X = (FMath::CeilToInt(Context->GenerationBounds.Max.X / QuadSizeWorld) + 0.5) * QuadSizeWorld;
-			OutputBounds.Max.Y = (FMath::CeilToInt(Context->GenerationBounds.Max.Y / QuadSizeWorld) + 0.5) * QuadSizeWorld;
+			OutputBounds.Min.X = OutputBoundsMin.X;
+			OutputBounds.Min.Y = OutputBoundsMin.Y;
+			OutputBounds.Max.X = OutputBoundsMin.X + OutputExtents.X;
+			OutputBounds.Max.Y = OutputBoundsMin.Y + OutputExtents.Y;
 			OutputBounds.Min.Z = Context->TotalLandscapeComponentBounds.Min.Z;
 			OutputBounds.Max.Z = Context->TotalLandscapeComponentBounds.Max.Z;
 
 			Context->OutputTextureTransform = FTransform(FQuat::Identity, OutputBounds.GetCenter(), OutputBounds.GetExtent());
+			Context->OutputTextureBounds = OutputBounds;
 
-			Context->OutputResolution = FIntPoint(
-				FMath::CeilToInt((Context->GenerationBounds.Max.X - Context->GenerationBounds.Min.X) / QuadSizeWorld) + 1,
-				FMath::CeilToInt((Context->GenerationBounds.Max.Y - Context->GenerationBounds.Min.Y) / QuadSizeWorld) + 1);
+			// Output resolution should be the dimensions (in texels) of the output bounds.
+			Context->OutputResolution = OutputExtentTexels;
 		}
 
 		Context->bLandscapeComponentsFiltered = true;
diff --git a/Engine/Plugins/PCG/Source/PCG/Public/Elements/PCGGenerateLandscapeTextures.h b/Engine/Plugins/PCG/Source/PCG/Public/Elements/PCGGenerateLandscapeTextures.h
--- a/Engine/Plugins/PCG/Source/PCG/Public/Elements/PCGGenerateLandscapeTextures.h
+++ b/Engine/Plugins/PCG/Source/PCG/Public/Elements/PCGGenerateLandscapeTextures.h
@@ -161,6 +161,8 @@ public:
 
 	FTransform OutputTextureTransform;
 
+	FBox OutputTextureBounds = FBox(EForceInit::ForceInit);
+
 	FIntPoint OutputResolution = FIntPoint::ZeroValue;
 
 	TObjectPtr<UTexture> BlackTexture;
