[PCG] Level to Asset: add a project setting to exclude some actor classes from the export

#Synced-CL 51093404

diff --git a/Engine/Plugins/PCG/Source/PCGEditor/Private/PCGLevelToAsset.cpp b/Engine/Plugins/PCG/Source/PCGEditor/Private/PCGLevelToAsset.cpp
--- a/Engine/Plugins/PCG/Source/PCGEditor/Private/PCGLevelToAsset.cpp
+++ b/Engine/Plugins/PCG/Source/PCGEditor/Private/PCGLevelToAsset.cpp
@@ -3,6 +3,7 @@
 #include "PCGLevelToAsset.h"
 
 #include "PCGEditorModule.h"
+#include "PCGEditorSettings.h"
 #include "PCGEditorUtils.h"
 #include "PCGAssetExporterUtils.h"
 
@@ -317,43 +318,63 @@ bool UPCGLevelToAsset::BP_ExportWorld_Implementation(UWorld* World, const FStrin
 		}
 	};
 
-	auto IsActorExcluded = [LevelInstanceSubsystem, &GetAttachParentActor](AActor* InActor)
+	const UPCGEditorProjectSettings* EditorProjectSettings = UPCGEditorProjectSettings::StaticClass()->GetDefaultObject<UPCGEditorProjectSettings>();
+	const TArray<TSubclassOf<AActor>>& RootExcludedClasses = EditorProjectSettings->PCGLevelToAssetExcludedActorClasses;
+
+	// Cache of already processed classes
+	TSet<TSubclassOf<AActor>> IncludedActorClasses;
+	TSet<TSubclassOf<AActor>> ExcludedActorClasses;
+	ExcludedActorClasses.Append(RootExcludedClasses);
+
+	TFunction<bool(AActor*)> IsActorExcluded = [LevelInstanceSubsystem, &GetAttachParentActor, &RootExcludedClasses, &ExcludedActorClasses, &IncludedActorClasses, &IsActorExcluded](AActor* InActor)
 	{
 		if (InActor->ActorHasTag(PCGLevelToAssetConstants::ExcludedActorTag))
 		{
 			return true;
 		}
 
-		// Check Level Instance Hierarchy first
-		if (InActor->IsInLevelInstance())
+		// Check for project specific class exclusions
+		if (!RootExcludedClasses.IsEmpty())
 		{
-			bool bExcluded = false;
-			LevelInstanceSubsystem->ForEachLevelInstanceAncestors(InActor, [&bExcluded](ILevelInstanceInterface* Ancestor)
+			TSubclassOf<AActor> ActorClass(InActor->GetClass());
+			if (ExcludedActorClasses.Contains(ActorClass))
 			{
-				AActor* AncestorActor = CastChecked<AActor>(Ancestor);
-				if (AncestorActor->ActorHasTag(PCGLevelToAssetConstants::ExcludedActorTag))
+				return true;
+			}
+			// If not already processed inclusion then process actor class
+			else if (!IncludedActorClasses.Contains(ActorClass))
+			{
+				for (const TSubclassOf<AActor>& ExcludedActorClass : RootExcludedClasses)
 				{
-					bExcluded = true;
-					return false;
+					if (ActorClass->IsChildOf(ExcludedActorClass))
+					{
+						ExcludedActorClasses.Add(ActorClass);
+						return true;
+					}
 				}
-				return true;
-			});
 
-			if (bExcluded)
-			{
-				return true;
+				IncludedActorClasses.Add(ActorClass);
 			}
 		}
 
-		// Check Attachment next (which will go passed Level Instance Hierarchy)
-		AActor* AttachParent = GetAttachParentActor(InActor);
-		while (AttachParent)
+		// Check Level Instance Hierarchy first (calling IsActorExcluded will go up hierarchy)
+		if (InActor->IsInLevelInstance())
 		{
-			if (AttachParent->ActorHasTag(PCGLevelToAssetConstants::ExcludedActorTag))
+			if (ILevelInstanceInterface* ParentLevelInstance = LevelInstanceSubsystem->GetParentLevelInstance(InActor))
 			{
-				return true;
+				AActor* AncestorActor = CastChecked<AActor>(ParentLevelInstance);
+				if (IsActorExcluded(AncestorActor))
+				{
+					return true;
+				}
 			}
-			AttachParent = GetAttachParentActor(AttachParent);
+		}
+
+		// Check Attachment hierarchy next  (calling IsActorExcluded will go up hierarchy)
+		AActor* AttachParent = GetAttachParentActor(InActor);
+		if (AttachParent && IsActorExcluded(AttachParent))
+		{
+			return true;
 		}
 
 		return false;
diff --git a/Engine/Plugins/PCG/Source/PCGEditor/Public/PCGEditorSettings.h b/Engine/Plugins/PCG/Source/PCGEditor/Public/PCGEditorSettings.h
--- a/Engine/Plugins/PCG/Source/PCGEditor/Public/PCGEditorSettings.h
+++ b/Engine/Plugins/PCG/Source/PCGEditor/Public/PCGEditorSettings.h
@@ -42,6 +42,10 @@ public:
 	UPROPERTY(EditAnywhere, config, Category = Builder)
 	TSoftObjectPtr<UPCGBuilderSettings> DefaultBuilderSetting;
 
+	/** Actor classes to ignore for the PCG Level to Asset export */
+	UPROPERTY(EditAnywhere, config, Category = "PCG Level to Asset", meta=(DisplayName = "Excluded Actor Classes"))
+	TArray<TSubclassOf<AActor>> PCGLevelToAssetExcludedActorClasses;
+
 	/** [EXPERIMENTAL] Whether to automatically refresh components that use GPU Static Mesh Spawners when materials are modified. Temporary workaround for
 	* issues where instances can be lost.
 	* Note: This setting is subject to change/removal without deprecation.
