[PCG] GPU compiler fixes

* Fix bug where all output buffers from a kernel where marked for export if any of them required export
* Fix cases where we still used the Settings pins when setting up kernels, which are not the same as kernel pins

#rnx

#Synced-CL 51091775

diff --git a/Engine/Plugins/PCG/Source/PCG/Private/Graph/PCGGraphCompilerGPU.cpp b/Engine/Plugins/PCG/Source/PCG/Private/Graph/PCGGraphCompilerGPU.cpp
--- a/Engine/Plugins/PCG/Source/PCG/Private/Graph/PCGGraphCompilerGPU.cpp
+++ b/Engine/Plugins/PCG/Source/PCG/Private/Graph/PCGGraphCompilerGPU.cpp
@@ -1046,7 +1046,7 @@ void FPCGGraphCompilerGPU::CreateDataInterfaces(
 				{
 					for (const FPCGGraphTaskInput& Input : InCompiledTasks[Successor].Inputs)
 					{
-						if (Input.TaskId == TaskId && Input.UpstreamPin && !InCollapsedTasks.Contains(Successor))
+						if (Input.TaskId == TaskId && Input.UpstreamPin && Input.UpstreamPin->Label == OutputPinProperties.Label && !InCollapsedTasks.Contains(Successor))
 						{
 							bRequiresExport = true;
 							break;
@@ -1258,6 +1258,9 @@ void FPCGGraphCompilerGPU::CompileComputeGraph(
 
 		InOutComputeGraph->KernelInvocations.Add(KernelWithBindings.Kernel);
 
+		TArray<FPCGPinPropertiesGPU> KernelOutputPinProperties;
+		KernelWithBindings.Kernel->GetOutputPins(KernelOutputPinProperties);
+
 		if (const UPCGStaticMeshSpawnerKernel* SMSpawnerKernel = Cast<UPCGStaticMeshSpawnerKernel>(KernelWithBindings.Kernel))
 		{
 			InOutComputeGraph->StaticMeshSpawners.Add(SMSpawnerKernel);
@@ -1265,7 +1268,7 @@ void FPCGGraphCompilerGPU::CompileComputeGraph(
 
 		// Populate static data labels. Cannot happen during BuildComputeGraphStaticData() because this relies on kernel index, which is not initialized at that time.
 		{
-			auto GetDataLabels = [Kernel = KernelWithBindings.Kernel, InOutComputeGraph](const TArray<FPCGPinProperties>& PinProperties)
+			auto GetDataLabels = [Kernel=KernelWithBindings.Kernel, InOutComputeGraph](const auto& PinProperties)
 			{
 				for (const FPCGPinProperties& PinProps : PinProperties)
 				{
@@ -1281,8 +1284,11 @@ void FPCGGraphCompilerGPU::CompileComputeGraph(
 				}
 			};
 
-			GetDataLabels(Settings->AllInputPinProperties());
-			GetDataLabels(Settings->AllOutputPinProperties());
+			TArray<FPCGPinProperties> InputPinProperties;
+			KernelWithBindings.Kernel->GetInputPins(InputPinProperties);
+			GetDataLabels(InputPinProperties);
+
+			GetDataLabels(KernelOutputPinProperties);
 		}
 
 		struct FDataInterfaceInfo
@@ -1370,7 +1376,7 @@ void FPCGGraphCompilerGPU::CompileComputeGraph(
 		}
 
 		// Setup output pin DIs. Always bind a DI to every output pin, so kernel always has something to write to.
-		for (const FPCGPinProperties& OutputPinProperties : Settings->AllOutputPinProperties())
+		for (const FPCGPinPropertiesGPU& OutputPinProperties : KernelOutputPinProperties)
 		{
 			if (OutputPinProperties.Label == NAME_None)
 			{
