[PCG] Added actor name/label option.

#Synced-CL 51118381

diff --git a/Engine/Plugins/PCG/Source/PCG/Private/Elements/PCGSpawnActor.cpp b/Engine/Plugins/PCG/Source/PCG/Private/Elements/PCGSpawnActor.cpp
--- a/Engine/Plugins/PCG/Source/PCG/Private/Elements/PCGSpawnActor.cpp
+++ b/Engine/Plugins/PCG/Source/PCG/Private/Elements/PCGSpawnActor.cpp
@@ -1003,6 +1003,11 @@ void FPCGSpawnActorElement::SpawnActors(FPCGSubgraphContext* Context, AActor* Ta
 
 		bool bAllActorOverridesSucceeded = true;
 
+#if WITH_EDITOR
+		// Since the actor params take a string view as a parameter, we need a temporary string on the stack.
+		FString TempLabelString;
+#endif
+
 		const FConstPCGPointValueRanges ValueRanges(PointData);
 
 		TPCGValueRange<int64> OutMetadataEntryRange = OutPointData ? OutPointData->GetMetadataEntryValueRange() : TPCGValueRange<int64>();
@@ -1012,6 +1017,15 @@ void FPCGSpawnActorElement::SpawnActors(FPCGSubgraphContext* Context, AActor* Ta
 			bAllActorOverridesSucceeded &= ActorOverrides.Apply(i);
 
 			SpawnDefaultActorParams.Transform = ValueRanges.TransformRange[i];
+
+#if WITH_EDITOR
+			if (!Settings->ActorLabel.IsEmpty())
+			{
+				TempLabelString = FString::Format(TEXT("{0}{1}"), { Settings->ActorLabel, i });
+				SpawnDefaultActorParams.SpawnParams.InitialActorLabel = TempLabelString;
+			}
+#endif
+
 			AActor* GeneratedActor = UPCGActorHelpers::SpawnDefaultActor(SpawnDefaultActorParams);
 
 			if (!GeneratedActor)
diff --git a/Engine/Plugins/PCG/Source/PCG/Private/Helpers/PCGActorHelpers.cpp b/Engine/Plugins/PCG/Source/PCG/Private/Helpers/PCGActorHelpers.cpp
--- a/Engine/Plugins/PCG/Source/PCG/Private/Helpers/PCGActorHelpers.cpp
+++ b/Engine/Plugins/PCG/Source/PCG/Private/Helpers/PCGActorHelpers.cpp
@@ -718,7 +718,7 @@ AActor* UPCGActorHelpers::SpawnDefaultActor(const FSpawnDefaultActorParams& Para
 
 #if WITH_EDITOR
 	// Capture label before potentially changing it for previewing
-	FString Label = SpawnParams.Name.IsNone() ? FString() : SpawnParams.Name.ToString();
+	FString DefaultLabel = SpawnParams.Name.IsNone() ? FString() : SpawnParams.Name.ToString();
 
 	// If we are spawning a non-runtime preview actor in a world partition world, we need to assign it an external package as Data Layers are not supported on non OFPA actors.
 	// Note: this could at somepoint become an option in the FActorSpawnParameters to allow RF_Transient actors to create their own package in UWorld::SpawnActor.
@@ -734,9 +734,9 @@ AActor* UPCGActorHelpers::SpawnDefaultActor(const FSpawnDefaultActorParams& Para
 		SpawnParams.Name = FActorSpawnUtils::MakeUniqueActorName(SpawnParams.OverrideLevel, Params.ActorClass, *ActorClassName, /*bGloballyUniqueName*/true);
 		
 		// If Label wasn't empty add suffix to it, if not use basename as label
-		if (Label.IsEmpty())
+		if (DefaultLabel.IsEmpty())
 		{
-			Label = ActorClassName;
+			DefaultLabel = ActorClassName;
 		}
 		
 		SpawnParams.OverridePackage = UPCGActorHelpers::CreatePreviewPackage(SpawnParams.OverrideLevel, SpawnParams.Name.ToString());
@@ -746,7 +746,10 @@ AActor* UPCGActorHelpers::SpawnDefaultActor(const FSpawnDefaultActorParams& Para
 		}
 	}
 
-	SpawnParams.InitialActorLabel = Label;
+	if (SpawnParams.InitialActorLabel.IsEmpty())
+	{
+		SpawnParams.InitialActorLabel = DefaultLabel;
+	}
 
 	// Avoid creating an actor package for runtime actors
 	SpawnParams.bCreateActorPackage &= !Params.bIsRuntime;
diff --git a/Engine/Plugins/PCG/Source/PCG/Public/Elements/PCGSpawnActor.h b/Engine/Plugins/PCG/Source/PCG/Public/Elements/PCGSpawnActor.h
--- a/Engine/Plugins/PCG/Source/PCG/Public/Elements/PCGSpawnActor.h
+++ b/Engine/Plugins/PCG/Source/PCG/Public/Elements/PCGSpawnActor.h
@@ -72,6 +72,10 @@ public:
 	UPROPERTY(BlueprintReadWrite, EditAnywhere, Category = Settings, meta = (EditCondition = "Option!=EPCGSpawnActorOption::CollapseActors", EditConditionHides))
 	TArray<FName> TagsToAddOnActors;
 
+	/** Overrides the default label that would be set on the spawned actors (class based) by a user-specified one. */
+	UPROPERTY(BlueprintReadWrite, EditAnywhere, Category = Settings, meta = (EditCondition = "Option!=EPCGSpawnActorOption::CollapseActors", EditConditionHides, PCG_Overridable))
+	FString ActorLabel;
+
 	UPROPERTY(EditAnywhere, BlueprintReadWrite, Instanced, Category = Settings, meta = (ShowInnerProperties, EditCondition = "GetAllowTemplateActorEditing()", EditConditionHides))
 	TObjectPtr<AActor> TemplateActor;
 
