[PCG] Generalized Component Data Storage

+ FPCGSourceDataContainer: for execution source data
* UPCGComponent exposes SourceDataContainer through the execution source interface

#Synced-CL 51160565

diff --git a/Engine/Plugins/PCG/Source/PCG/Private/Graph/PCGSourceDataContainer.cpp b/Engine/Plugins/PCG/Source/PCG/Private/Graph/PCGSourceDataContainer.cpp
new file mode 100644
--- /dev/null
+++ b/Engine/Plugins/PCG/Source/PCG/Private/Graph/PCGSourceDataContainer.cpp
@@ -0,0 +1,110 @@
+// Copyright Epic Games, Inc. All Rights Reserved.
+
+#include "Graph/PCGSourceDataContainer.h"
+
+#include UE_INLINE_GENERATED_CPP_BY_NAME(PCGSourceDataContainer)
+
+#define LOCTEXT_NAMESPACE "PCGSourceDataContainer"
+
+FPCGSourceDataStorageKey::FPCGSourceDataStorageKey(const FName InLabel, const uint64 InHash) : Label(InLabel), Hash(InHash) {}
+
+FPCGSourceDataContainer::FPCGSourceDataContainer(const FPCGSourceDataContainer& Other)
+{
+	bAutoDirty.Store(Other.bAutoDirty.Load());
+	DirtyGeneration = Other.DirtyGeneration;
+	DataStorage = Other.DataStorage;
+}
+
+FPCGSourceDataContainer::FPCGSourceDataContainer(FPCGSourceDataContainer&& Other)
+{
+	DataStorage = MoveTemp(Other.DataStorage);
+	DirtyGeneration = Other.DirtyGeneration;
+	bAutoDirty.Store(Other.bAutoDirty.Load());
+	Other.DirtyGeneration = 0;
+	Other.bAutoDirty.Store(false);
+}
+
+FPCGSourceDataContainer& FPCGSourceDataContainer::operator=(FPCGSourceDataContainer&& Other)
+{
+	if (this != &Other)
+	{
+		DataStorage = MoveTemp(Other.DataStorage);
+		DirtyGeneration = Other.DirtyGeneration;
+		bAutoDirty.Store(Other.bAutoDirty.Load());
+		Other.DirtyGeneration = 0;
+		Other.bAutoDirty.Store(false);
+	}
+	return *this;
+}
+
+FPCGSourceDataContainer& FPCGSourceDataContainer::operator=(const FPCGSourceDataContainer& Other)
+{
+	if (this != &Other)
+	{
+		bAutoDirty.Store(Other.bAutoDirty.Load());
+		DirtyGeneration = Other.DirtyGeneration;
+		DataStorage = Other.DataStorage;
+	}
+
+	return *this;
+}
+
+#if WITH_EDITOR
+bool FPCGSourceDataContainer::Remove(const FPCGSourceDataStorageKey& DataKey)
+{
+	PCG::TUniqueScopeLock Lock(DataStorageLock);
+
+	const int32 NumRemovedItems = DataStorage.Remove(DataKey);
+	const bool bWasRemoved = NumRemovedItems > 0;
+
+	if (bWasRemoved && bAutoDirty)
+	{
+		++DirtyGeneration;
+	}
+
+	return bWasRemoved;
+}
+
+void FPCGSourceDataContainer::MarkDirty()
+{
+	PCG::TUniqueScopeLock Lock(DataStorageLock);
+	++DirtyGeneration;
+}
+
+void FPCGSourceDataContainer::SetShouldAutoDirty(const bool bShouldAutoDirty)
+{
+	PCG::TUniqueScopeLock Lock(DataStorageLock);
+	bAutoDirty = bShouldAutoDirty;
+}
+
+void FPCGSourceDataContainer::Empty()
+{
+	DataStorage.Empty();
+	DirtyGeneration = 0;
+}
+#endif // WITH_EDITOR
+
+uint32 FPCGSourceDataContainer::GetDirtyGeneration() const
+{
+#if WITH_EDITOR
+	PCG::TSharedScopeLock Lock(DataStorageLock);
+#endif // WITH_EDITOR
+	return DirtyGeneration;
+}
+
+int32 FPCGSourceDataContainer::Num() const
+{
+	return DataStorage.Num();
+}
+
+bool FPCGSourceDataContainer::IsEmpty() const
+{
+	return DataStorage.IsEmpty();
+}
+
+uint32 GetTypeHash(const FPCGSourceDataStorageKey& Key)
+{
+	return HashCombine(GetTypeHash(Key.Label), GetTypeHash(Key.Hash));
+}
+
+#undef LOCTEXT_NAMESPACE
\ No newline at end of file
diff --git a/Engine/Plugins/PCG/Source/PCG/Private/Helpers/Parsing/PCGIndexing.cpp b/Engine/Plugins/PCG/Source/PCG/Private/Helpers/Parsing/PCGIndexing.cpp
--- a/Engine/Plugins/PCG/Source/PCG/Private/Helpers/Parsing/PCGIndexing.cpp
+++ b/Engine/Plugins/PCG/Source/PCG/Private/Helpers/Parsing/PCGIndexing.cpp
@@ -152,11 +152,14 @@ namespace PCGIndexing
 		return IndexRanges.IsEmpty();
 	}
 
-	void FPCGIndexCollection::ForEach(const TFunctionRef<void(FPCGIndexRange)> Action)
+	void FPCGIndexCollection::ForEach(const TFunctionRef<bool(FPCGIndexRange)> Action) const
 	{
 		for (const FPCGIndexRange& Range : IndexRanges)
 		{
-			Action(Range);
+			if (Action(Range))
+			{
+				break;
+			}
 		}
 	}
 
diff --git a/Engine/Plugins/PCG/Source/PCG/Private/PCGComponent.cpp b/Engine/Plugins/PCG/Source/PCG/Private/PCGComponent.cpp
--- a/Engine/Plugins/PCG/Source/PCG/Private/PCGComponent.cpp
+++ b/Engine/Plugins/PCG/Source/PCG/Private/PCGComponent.cpp
@@ -2013,6 +2013,8 @@ PRAGMA_ENABLE_DEPRECATION_WARNINGS
 		LoadedPreviewGeneratedGraphOutput = GeneratedGraphOutput;
 		bDirtyGenerated = PCGSystemSwitches::CVarDirtyLoadAsPreviewOnLoad.GetValueOnAnyThread();
 	}
+
+	// @todo_pcg: Migrate FPCGInteractiveToolDataContainer and runtime generation scheduling data entries into SourceDataContainer her
 #endif
 
 	if (!IsValid(SchedulingPolicy))
@@ -4049,6 +4051,12 @@ void FPCGComponentInstanceData::ApplyToComponent(UActorComponent* Component, con
 			PCGComponent->Refresh();
 		}
 #endif // WITH_EDITOR
+
+		// Copy over the source data container
+		if (ConstructionSourceComponent)
+		{
+			PCGComponent->SourceDataContainer = ConstructionSourceComponent->SourceDataContainer;
+		}
 	}
 }
 
diff --git a/Engine/Plugins/PCG/Source/PCG/Private/PCGComponentExecutionState.cpp b/Engine/Plugins/PCG/Source/PCG/Private/PCGComponentExecutionState.cpp
--- a/Engine/Plugins/PCG/Source/PCG/Private/PCGComponentExecutionState.cpp
+++ b/Engine/Plugins/PCG/Source/PCG/Private/PCGComponentExecutionState.cpp
@@ -288,6 +288,18 @@ FPCGTaskId FPCGComponentExecutionState::Generate(const FGenerateParams& InGenera
 	return Component->GenerateLocalGetTaskId(GenTrigger, InGenerateParams.bEvenIfAlreadyGenerated, Grid, InGenerateParams.Dependencies);
 }
 
+FPCGSourceDataContainer* FPCGComponentExecutionState::GetSourceDataContainer()
+{
+	check(Component);
+	return &Component->SourceDataContainer;
+}
+
+const FPCGSourceDataContainer* FPCGComponentExecutionState::GetSourceDataContainer() const
+{
+	check(Component);
+	return &Component->SourceDataContainer;
+}
+
 // UE_DEPRECATED(5.8)
 FPCGTaskId FPCGComponentExecutionState::GenerateLocalGetTaskId(EPCGHiGenGrid Grid)
 {
diff --git a/Engine/Plugins/PCG/Source/PCG/Private/PCGContext.cpp b/Engine/Plugins/PCG/Source/PCG/Private/PCGContext.cpp
--- a/Engine/Plugins/PCG/Source/PCG/Private/PCGContext.cpp
+++ b/Engine/Plugins/PCG/Source/PCG/Private/PCGContext.cpp
@@ -773,6 +773,18 @@ FPCGContext* FPCGContext::GetContextForThread()
 	return PCGContext::ThisThreadContext;
 }
 
+FPCGSourceDataContainer* FPCGContext::GetSourceDataContainer()
+{
+	IPCGGraphExecutionSource* Source = ExecutionSource.IsValid() ? ExecutionSource.Get() : nullptr;
+	return Source ? Source->GetExecutionState().GetSourceDataContainer() : nullptr;
+}
+
+const FPCGSourceDataContainer* FPCGContext::GetSourceDataContainer() const
+{
+	const IPCGGraphExecutionSource* Source = ExecutionSource.IsValid() ? ExecutionSource.Get() : nullptr;
+	return Source ? Source->GetExecutionState().GetSourceDataContainer() : nullptr;
+}
+
 FPCGContextScope::FPCGContextScope(FPCGContext* InContext)
 {
 	PreviousContext = PCGContext::ThisThreadContext;
diff --git a/Engine/Plugins/PCG/Source/PCG/Public/Graph/PCGSourceDataContainer.h b/Engine/Plugins/PCG/Source/PCG/Public/Graph/PCGSourceDataContainer.h
new file mode 100644
--- /dev/null
+++ b/Engine/Plugins/PCG/Source/PCG/Public/Graph/PCGSourceDataContainer.h
@@ -0,0 +1,197 @@
+// Copyright Epic Games, Inc. All Rights Reserved.
+
+#pragma once
+
+#include "PCGModule.h"
+#include "StructUtils/SharedStruct.h"
+
+#include "PCGSourceDataContainer.generated.h"
+
+// Concept for guaranteeing compatible script struct types for the source data types.
+template <typename T>
+concept CPCGSourceDataCompatibleType = requires { T::StaticStruct(); };
+
+/**
+ * This data storage key merges an associated label with a predefined hash to create a unique key.
+ * Note: The Label can be used independently if desired, with the Hash left as 0.
+ */
+USTRUCT()
+struct PCG_API FPCGSourceDataStorageKey
+{
+	GENERATED_BODY()
+
+	FPCGSourceDataStorageKey() = default;
+	FPCGSourceDataStorageKey(FName InLabel, uint64 InHash = 0u);
+
+	bool operator==(const FPCGSourceDataStorageKey& Other) const = default;
+
+	PCG_API friend uint32 GetTypeHash(const FPCGSourceDataStorageKey& Key);
+
+	/** A label associated with this storage key. (Ex. "MyCustomToolData") */
+	UPROPERTY()
+	FName Label;
+
+	/** An additional hash to associate with a unique identifier. */
+	UPROPERTY()
+	uint64 Hash = 0;
+};
+
+/** Type-erased storage value. Wraps a FSharedStruct payload for runtime verification and key collision avoidance. */
+USTRUCT()
+struct PCG_API FPCGSourceDataStorageValue
+{
+	GENERATED_BODY()
+
+	/** A FSharedStruct container of data to store within the Execution Source. */
+	UPROPERTY()
+	FSharedStruct Data;
+
+	template <CPCGSourceDataCompatibleType T>
+	FConstSharedStruct GetAs() const;
+
+	template <CPCGSourceDataCompatibleType T>
+	FSharedStruct GetAsMutable();
+};
+
+/**
+ * A container for data storage on an Execution Source, supporting unique keying and Crc compatibility.
+ * NOTE: This container is READ-ONLY at runtime.
+ */
+USTRUCT()
+struct PCG_API FPCGSourceDataContainer
+{
+	GENERATED_BODY()
+
+	FPCGSourceDataContainer() = default;
+	~FPCGSourceDataContainer() = default;
+	// Required by USTRUCT
+	// NOT THREAD SAFE: These must only be called from a single thread with no concurrent readers/writers.
+	FPCGSourceDataContainer(const FPCGSourceDataContainer& Other);
+	FPCGSourceDataContainer(FPCGSourceDataContainer&& Other);
+	FPCGSourceDataContainer& operator=(const FPCGSourceDataContainer& Other);
+	FPCGSourceDataContainer& operator=(FPCGSourceDataContainer&& Other);
+
+#if WITH_EDITOR
+	/** Store data of type T within the storage container. */
+	template <CPCGSourceDataCompatibleType T>
+	void Store(const FPCGSourceDataStorageKey& DataKey, const T& Data);
+
+	/** Remove the data attached to a specific Data Key. */
+	bool Remove(const FPCGSourceDataStorageKey& DataKey);
+
+	/** Get a mutable shared reference to stored data by type. Returns an empty FSharedStruct if key or type doesn't exist. */
+	template <CPCGSourceDataCompatibleType T>
+	FSharedStruct GetMutable(const FPCGSourceDataStorageKey& DataKey);
+
+	/** Mark the container dirty to invalidate the Crc cache. */
+	void MarkDirty();
+
+	/** Set the container to automatically dirty to invalidate the Crc cache. */
+	void SetShouldAutoDirty(bool bShouldAutoDirty);
+
+	/** Empty the contents of the container. */
+	void Empty();
+#endif // WITH_EDITOR
+
+	/** Get the dirty generation value for Crc combination. */
+	uint32 GetDirtyGeneration() const;
+
+	/** Get the number of storage values. */
+	int32 Num() const;
+
+	/** True if the storage is empty. */
+	bool IsEmpty() const;
+
+	/** Get a const shared reference to stored data by type. Returns an empty FConstSharedStruct if key or type doesn't exist. */
+	template <CPCGSourceDataCompatibleType T>
+	FConstSharedStruct Get(const FPCGSourceDataStorageKey& DataKey) const;
+
+private:
+	UPROPERTY()
+	TMap<FPCGSourceDataStorageKey, FPCGSourceDataStorageValue> DataStorage;
+
+	/** Automatically increments the dirty generation to invalidate the Crc node cache. */
+	TAtomic<bool> bAutoDirty = false;
+
+	/** Mark the source data container dirty by incrementing this, which will invalidate the Crc cache. */
+	UPROPERTY()
+	uint32 DirtyGeneration = 0;
+
+#if WITH_EDITOR
+	mutable PCG::FSharedLock DataStorageLock;
+#endif // WITH_EDITOR
+};
+
+template <CPCGSourceDataCompatibleType T>
+FConstSharedStruct FPCGSourceDataStorageValue::GetAs() const
+{
+#if WITH_EDITOR
+	if (!ensure(Data.GetPtr<T>() || !Data.IsValid()))
+	{
+		const FString RequestedType = T::StaticStruct()->GetName();
+		const FString StoredType = Data.GetScriptStruct()->GetName();
+		UE_LOG(LogPCG, Error, TEXT("FPCGSourceDataStorageValue type mismatch: requested '%s', stored '%s'"), *RequestedType, *StoredType);
+		return {};
+	}
+#endif // WITH_EDITOR
+	return FConstSharedStruct(Data);
+}
+
+template <CPCGSourceDataCompatibleType T>
+FSharedStruct FPCGSourceDataStorageValue::GetAsMutable()
+{
+#if WITH_EDITOR
+	if (!ensure(Data.GetPtr<T>() || !Data.IsValid()))
+	{
+		const FString RequestedType = T::StaticStruct()->GetName();
+		const FString StoredType = Data.GetScriptStruct()->GetName();
+		UE_LOG(LogPCG, Error, TEXT("FPCGSourceDataStorageValue type mismatch: requested '%s', stored '%s'"), *RequestedType, *StoredType);
+		return {};
+	}
+#endif // WITH_EDITOR
+	return FSharedStruct(Data);
+}
+
+#if WITH_EDITOR
+template <CPCGSourceDataCompatibleType T>
+void FPCGSourceDataContainer::Store(const FPCGSourceDataStorageKey& DataKey, const T& Data)
+{
+	PCG::TUniqueScopeLock Lock(DataStorageLock);
+
+	FPCGSourceDataStorageValue& Value = DataStorage.FindOrAdd(DataKey);
+	Value.Data = FSharedStruct::Make(Data);
+
+	if (bAutoDirty)
+	{
+		++DirtyGeneration;
+	}
+}
+
+template <CPCGSourceDataCompatibleType T>
+FSharedStruct FPCGSourceDataContainer::GetMutable(const FPCGSourceDataStorageKey& DataKey)
+{
+	PCG::TSharedScopeLock Lock(DataStorageLock);
+
+	if (FPCGSourceDataStorageValue* StoredValue = DataStorage.Find(DataKey))
+	{
+		return StoredValue->GetAsMutable<T>();
+	}
+
+	return {};
+}
+#endif // WITH_EDITOR
+
+template <CPCGSourceDataCompatibleType T>
+FConstSharedStruct FPCGSourceDataContainer::Get(const FPCGSourceDataStorageKey& DataKey) const
+{
+#if WITH_EDITOR
+	PCG::TSharedScopeLock Lock(DataStorageLock);
+#endif // WITH_EDITOR
+
+	if (const FPCGSourceDataStorageValue* Value = DataStorage.Find(DataKey))
+	{
+		return Value->GetAs<T>();
+	}
+
+	return {};
+}
diff --git a/Engine/Plugins/PCG/Source/PCG/Public/Helpers/Parsing/PCGIndexing.h b/Engine/Plugins/PCG/Source/PCG/Public/Helpers/Parsing/PCGIndexing.h
--- a/Engine/Plugins/PCG/Source/PCG/Public/Helpers/Parsing/PCGIndexing.h
+++ b/Engine/Plugins/PCG/Source/PCG/Public/Helpers/Parsing/PCGIndexing.h
@@ -68,8 +68,8 @@ namespace PCGIndexing
 		/** The index collection has no ranges. */
 		bool IsEmpty() const;
 
-		/** Iterate over the entire range of indices. */
-		void ForEach(const TFunctionRef<void(FPCGIndexRange)> Action);
+		/** Iterate over the entire range of indices. Action should return true if the loop should terminate */
+		void ForEach(const TFunctionRef<bool(FPCGIndexRange)> Action) const;
 
 		bool operator==(const FPCGIndexCollection& Other) const;
 		FPCGIndexCollection& operator+=(const FPCGIndexCollection& Other);
diff --git a/Engine/Plugins/PCG/Source/PCG/Public/PCGComponent.h b/Engine/Plugins/PCG/Source/PCG/Public/PCGComponent.h
--- a/Engine/Plugins/PCG/Source/PCG/Public/PCGComponent.h
+++ b/Engine/Plugins/PCG/Source/PCG/Public/PCGComponent.h
@@ -11,6 +11,7 @@
 #include "PCGSettings.h"
 #include "ChangeTracking/PCGChangeTracker.h"
 #include "Data/PCGToolData.h"
+#include "Graph/PCGSourceDataContainer.h"
 #include "Graph/PCGStackContext.h"
 #include "Grid/PCGGridDescriptor.h"
 #include "Utils/PCGExtraCapture.h"
@@ -845,6 +846,10 @@ public:
 #endif
 
 private:
+	/** Generalized storage for the execution source (ex. tool data, runtime gen data, manual overrides, etc.) */
+	UPROPERTY()
+	FPCGSourceDataContainer SourceDataContainer;
+
 	FPCGComponentExecutionState ExecutionState;
 
 	// Deprecated section
diff --git a/Engine/Plugins/PCG/Source/PCG/Public/PCGComponentExecutionState.h b/Engine/Plugins/PCG/Source/PCG/Public/PCGComponentExecutionState.h
--- a/Engine/Plugins/PCG/Source/PCG/Public/PCGComponentExecutionState.h
+++ b/Engine/Plugins/PCG/Source/PCG/Public/PCGComponentExecutionState.h
@@ -54,6 +54,8 @@ public:
 	virtual FTransactionallySafeCriticalSection* GetManagedResourceContainerLock() override;
 	virtual void OnManagedResourceAdded(UPCGManagedResource* InResource) override;
 	virtual FPCGTaskId Generate(const FGenerateParams& InGenerateParams = {}) override;
+	virtual FPCGSourceDataContainer* GetSourceDataContainer() override;
+	virtual const FPCGSourceDataContainer* GetSourceDataContainer() const override;
 
 	UE_DEPRECATED(5.8, "Call/Implement Generate instead.")
 	virtual FPCGTaskId GenerateLocalGetTaskId(EPCGHiGenGrid Grid = EPCGHiGenGrid::Uninitialized) override;
diff --git a/Engine/Plugins/PCG/Source/PCG/Public/PCGContext.h b/Engine/Plugins/PCG/Source/PCG/Public/PCGContext.h
--- a/Engine/Plugins/PCG/Source/PCG/Public/PCGContext.h
+++ b/Engine/Plugins/PCG/Source/PCG/Public/PCGContext.h
@@ -26,6 +26,7 @@ class UPCGSpatialData;
 struct FPCGGetFromCacheParams;
 struct FPCGInitializeElementParams;
 struct FPCGSettingsOverridableParam;
+struct FPCGSourceDataContainer;
 struct FPCGStack;
 struct FPCGStoreInCacheParams;
 
@@ -309,6 +310,9 @@ PRAGMA_ENABLE_DEPRECATION_WARNINGS
 	 */
 	UE_API static FPCGContext* GetContextForThread();
 
+	UE_API FPCGSourceDataContainer* GetSourceDataContainer();
+	UE_API const FPCGSourceDataContainer* GetSourceDataContainer() const;
+
 protected:
 	virtual UObject* GetExternalContainerForOverridableParam(const FPCGSettingsOverridableParam& InParam) { return nullptr; }
 	virtual FStructView GetExternalStructContainerForOverridableParam(const FPCGSettingsOverridableParam& InParam) { return {}; }
diff --git a/Engine/Plugins/PCG/Source/PCG/Public/PCGGraphExecutionStateInterface.h b/Engine/Plugins/PCG/Source/PCG/Public/PCGGraphExecutionStateInterface.h
--- a/Engine/Plugins/PCG/Source/PCG/Public/PCGGraphExecutionStateInterface.h
+++ b/Engine/Plugins/PCG/Source/PCG/Public/PCGGraphExecutionStateInterface.h
@@ -11,6 +11,7 @@
 
 #include "PCGGraphExecutionStateInterface.generated.h"
 
+struct FPCGSourceDataContainer;
 class IPCGBaseSubsystem;
 class FPCGGraphExecutionInspection;
 class IPCGGraphExecutionSource;
@@ -168,6 +169,12 @@ public:
 	/** Cache any data that is not thread safe or that needs to be computed only once for the whole execution */
 	virtual void ExecutePreGraph(FPCGContext* InContext) {}
 
+	/** Get the data container from the execution source. */
+	virtual FPCGSourceDataContainer* GetSourceDataContainer() { return nullptr; }
+
+	/** Get the immutable data container from the execution source. */
+	virtual const FPCGSourceDataContainer* GetSourceDataContainer() const { return nullptr; }
+
 #if WITH_EDITOR
 
 	/** Returns the FExtraCapture object for this ExecutionState */
