[PCG] IPCGGraphExecutionState default GetOriginalLocalSpaceBounds to call GetLocalSpaceBounds so that subclasses that do not support partioning only have to implement GetLocalSpaceBounds

#Synced-CL 51208877

diff --git a/Engine/Plugins/PCG/Source/PCG/Public/PCGGraphExecutionStateInterface.h b/Engine/Plugins/PCG/Source/PCG/Public/PCGGraphExecutionStateInterface.h
--- a/Engine/Plugins/PCG/Source/PCG/Public/PCGGraphExecutionStateInterface.h
+++ b/Engine/Plugins/PCG/Source/PCG/Public/PCGGraphExecutionStateInterface.h
@@ -109,7 +109,7 @@ public:
 	virtual FBox GetLocalSpaceBounds() const { return GetBounds(); }
 
 	/** Returns the ExecutionState's original source local space bounds if this ExecutionState is a local source. */
-	virtual FBox GetOriginalLocalSpaceBounds() const { return GetBounds(); }
+	virtual FBox GetOriginalLocalSpaceBounds() const { return GetLocalSpaceBounds(); }
 
 	/** Returns the ExecutionState's generated bounds which includes the artifact bounds. */
 	virtual FBox GetTotalBounds() const { return GetBounds(); }
