[PCG] Move the default value for generic attributes out of the values array

* Since we are type erased, it made sense to have everything in the same array, the default value at 0 and the rest afterwards
* But with multiple inheritance,the ValueKeyOffset would be different than the actual number of values (because it would be incremented by one everytime we inherit, to take into account the default value).
* This forces us to copy the default value multiple times in an initialize by copy or flatten without compression, which it wasteful
* So create a new struct that will hold the memory for the default value, and handle copy/initialization
* Adapted UnitTests to relfect what we expect in terms of value key offsets.

#rnx

#Synced-CL 51188099

diff --git a/Engine/Plugins/PCG/Source/PCG/Private/Metadata/PCGMetadataAttributeGeneric.cpp b/Engine/Plugins/PCG/Source/PCG/Private/Metadata/PCGMetadataAttributeGeneric.cpp
--- a/Engine/Plugins/PCG/Source/PCG/Private/Metadata/PCGMetadataAttributeGeneric.cpp
+++ b/Engine/Plugins/PCG/Source/PCG/Private/Metadata/PCGMetadataAttributeGeneric.cpp
@@ -394,12 +394,66 @@ void FPCGAttributeProperty::Copy(void* DestPtr, const void* SrcPtr, int32 Count)
 	PCG::Private::CopyArray(Property.Get(), DestPtr, SrcPtr, Count);
 }
 
+bool FPCGAttributeProperty::IsValid() const
+{
+	return Property.IsValid() && Property->Inner != nullptr;
+}
+
 const void* FPCGAttributeProperty::GetPtrInArray(const void* InArrayStart, int32 Index) const
 {
 	check(InnerElementCachedSize > 0);
 	return static_cast<const uint8*>(InArrayStart) + Index * InnerElementCachedSize;
 }
 
+FPCGAttributeDefaultValue::~FPCGAttributeDefaultValue()
+{
+	Destroy();
+}
+
+void FPCGAttributeDefaultValue::Init(TSharedPtr<FPCGAttributeProperty> InUnderlyingProperty)
+{
+	if (!InUnderlyingProperty || !InUnderlyingProperty->IsValid())
+	{
+		return;
+	}
+	
+	// Check if we already have a property set. If the property is the same type, we can keep the allocated memory.
+	// Otherwise we have to destroy it.
+	if (UnderlyingProperty && UnderlyingProperty->IsValid() && Memory)
+	{
+		if (UnderlyingProperty->Property->SameType(InUnderlyingProperty->Property.Get()))
+		{
+			UnderlyingProperty = MoveTemp(InUnderlyingProperty);
+			return;
+		}
+	}
+	
+	Destroy();
+	
+	UnderlyingProperty = MoveTemp(InUnderlyingProperty);
+	Memory = UnderlyingProperty->Property->Inner->AllocateAndInitializeValue();
+}
+
+void FPCGAttributeDefaultValue::CopyFrom(const void* SrcData)
+{
+	check(Memory);
+	UnderlyingProperty->Copy(GetRawPtr(), SrcData, /*Count=*/1);
+}
+
+void FPCGAttributeDefaultValue::Destroy()
+{
+	if (Memory)
+	{
+		if (ensureMsgf(UnderlyingProperty && UnderlyingProperty->IsValid(), 
+			TEXT("Memory is allocated but we do not have an underlying property. This is ill-formed and will leak memory")))
+		{
+			UnderlyingProperty->Property->Inner->DestroyValue(Memory);
+		}
+	}
+
+	Memory = nullptr;
+}
+
 FPCGMetadataAttributeGeneric::FPCGMetadataAttributeGeneric(const FPCGMetadataAttributeDesc& InAttributeDesc, FPCGMetadataDomain* InMetadata, FName InName, bool bInAllowsInterpolation)
 	: FPCGMetadataAttributeBase(InMetadata, InName, nullptr, bInAllowsInterpolation)
 	, CachedDesc(InAttributeDesc)
@@ -426,10 +480,10 @@ FPCGMetadataAttributeGeneric::FPCGMetadataAttributeGeneric(TNotNull<const FPCGMe
 
 	ValueKeyOffset = InParent->GetValueKeyOffsetForChild();
 	
-	if (UnderlyingProperty && ensure(!Values.IsEmpty()) && ensure(!InParent->Values.IsEmpty()))
+	if (ensure(DefaultValue.IsValid()) && ensure(InParent->DefaultValue.IsValid()))
 	{
 		// Copy the default value from the parent.
-		UnderlyingProperty->Copy(Values.GetData(), InParent->Values.GetData(), 1);
+		DefaultValue.CopyFrom(InParent->DefaultValue.GetRawPtr());
 	}
 }
 
@@ -472,11 +526,17 @@ TArray<PCGMetadataValueKey> FPCGMetadataAttributeGeneric::AddValues_Internal(FIn
 	int32 FirstValueIndex = INDEX_NONE;
 
 	TOptional<FScriptArrayHelper> Helper = MakeArrayHelper();
-	if (!Helper.IsSet())
+	if (!Helper.IsSet() || Count <= 0)
 	{
 		return ValueKeys;
 	}
 	
+	// Hard-requirement that when we write to the default value, we can only write a single value
+	if (!ensure(ExistingIndex != PCG::Private::ExistingIndexForDefaultValue || Count == 1))
+	{
+		Count = 1;
+	}
+
 	{
 		PCG::TUniqueScopeLock WriteLock(ValueLock, /*bShouldLock=*/!bLockless);
 		if (ExistingIndex == INDEX_NONE)
@@ -484,22 +544,33 @@ TArray<PCGMetadataValueKey> FPCGMetadataAttributeGeneric::AddValues_Internal(FIn
 			check(!bLockless);
 			FirstValueIndex = AllocateValues_Unsafe(Count);
 		}
-		else
+		else if (ExistingIndex != PCG::Private::ExistingIndexForDefaultValue)
 		{
 			// Expected to be already allocated.
 			check(Helper->Num() >= ExistingIndex + Count);
 			FirstValueIndex = ExistingIndex;
 		}
-		
+		else
+		{
+			// Expected the default value to always be allocated.
+			check(DefaultValue.GetRawPtr() != nullptr)
+			FirstValueIndex = ExistingIndex;
+		}
+
 		check(FirstValueIndex != INDEX_NONE);
-		
-		Visit([this, &Helper, FirstValueIndex, Count](auto&& Value)
+
+		auto GetRawPtr = [this, &Helper, FirstValueIndex](const int32 Index)
+		{
+			return FirstValueIndex == PCG::Private::ExistingIndexForDefaultValue ? DefaultValue.GetRawPtr() : Helper->GetRawPtr(Index);
+		};
+
+		Visit([this, &GetRawPtr, FirstValueIndex, Count](auto&& Value)
 			{
 				using T = std::decay_t<decltype(Value)>;
 				if constexpr(std::is_same_v<T, FInValuesByValue>)
 				{
 					const void* SrcPtr = Value.InValues;
-					void* DestPtr = Helper->GetRawPtr(FirstValueIndex);
+					void* DestPtr = GetRawPtr(FirstValueIndex);
 					UnderlyingProperty->Copy(DestPtr, SrcPtr, Count);
 				}
 				else if constexpr(std::is_same_v<T, FInValuesByValueSubset>)
@@ -507,7 +578,7 @@ TArray<PCGMetadataValueKey> FPCGMetadataAttributeGeneric::AddValues_Internal(FIn
 					for (int32 i = 0; i < Value.Indices.Num(); ++i)
 					{
 						const void* SrcPtr = UnderlyingProperty->GetPtrInArray(Value.InValues,  Value.Indices[i]);
-						void* DestPtr = Helper->GetRawPtr(FirstValueIndex + i);
+						void* DestPtr = GetRawPtr(FirstValueIndex + i);
 						UnderlyingProperty->Copy(DestPtr, SrcPtr, 1);
 					}
 				}
@@ -516,7 +587,7 @@ TArray<PCGMetadataValueKey> FPCGMetadataAttributeGeneric::AddValues_Internal(FIn
 					const FSetProperty* InnerProperty = CastFieldChecked<FSetProperty>(UnderlyingProperty->Property->Inner);
 					for (int32 i = 0; i < Count; ++i)
 					{
-						FScriptSetHelper OtherHelper(InnerProperty, Helper->GetRawPtr(FirstValueIndex + i));
+						FScriptSetHelper OtherHelper(InnerProperty, GetRawPtr(FirstValueIndex + i));
 						const int32 Num = Value.InValues[i].Num();
 				
 						for (int j = 0; j < Num; ++j)
@@ -530,7 +601,7 @@ TArray<PCGMetadataValueKey> FPCGMetadataAttributeGeneric::AddValues_Internal(FIn
 					const FMapProperty* InnerProperty = CastFieldChecked<FMapProperty>(UnderlyingProperty->Property->Inner);
 					for (int32 i = 0; i < Count; ++i)
 					{
-						FScriptMapHelper OtherHelper(InnerProperty, Helper->GetRawPtr(FirstValueIndex + i));
+						FScriptMapHelper OtherHelper(InnerProperty, GetRawPtr(FirstValueIndex + i));
 						const int32 Num = Value.InValues[i].Num();
 				
 						for (int j = 0; j < Num; ++j)
@@ -544,7 +615,7 @@ TArray<PCGMetadataValueKey> FPCGMetadataAttributeGeneric::AddValues_Internal(FIn
 					const FArrayProperty* InnerProperty = CastFieldChecked<FArrayProperty>(UnderlyingProperty->Property->Inner);
 					for (int32 i = 0; i < Count; ++i)
 					{
-						FScriptArrayHelper OtherHelper(InnerProperty, Helper->GetRawPtr(FirstValueIndex + i));
+						FScriptArrayHelper OtherHelper(InnerProperty, GetRawPtr(FirstValueIndex + i));
 						const int32 Num = Value.InValues[i].template Get<1>();
 						if (InnerProperty->HasAnyPropertyFlags(CPF_IsPlainOldData))
 						{
@@ -687,7 +758,7 @@ void FPCGMetadataAttributeGeneric::FindValuesInternal(const FInValuesByValue& In
 				continue;
 			}
 
-			for (PCGMetadataValueKey ValueKey = 1; ValueKey < Values.Num(); ++ValueKey)
+			for (PCGMetadataValueKey ValueKey = 0; ValueKey < Values.Num(); ++ValueKey)
 			{
 				const PCGMetadataValueKey OffsetKey = ValueKey + ValueKeyOffset;
 				if (Compare(ValueIndex, OffsetKey))
@@ -703,15 +774,14 @@ void FPCGMetadataAttributeGeneric::FindValuesInternal(const FInValuesByValue& In
 
 void FPCGMetadataAttributeGeneric::Init()
 {
-	if (!ensure(UnderlyingProperty) || !ensure(Values.IsEmpty()) || !UnderlyingProperty->Property)
+	if (!ensure(UnderlyingProperty) || !ensure(!DefaultValue.IsValid()) || !UnderlyingProperty->IsValid())
 	{
 		return;
 	}
 
-	// Initialize the array with a single element containing the default value
-	FScriptArrayHelper Helper(UnderlyingProperty->Property.Get(), &Values);
-	Helper.AddValue();
-	
+	// Initialize the default value
+	DefaultValue.Init(UnderlyingProperty);
+
 	auto GetType = [](const FProperty* InProperty) -> int16
 	{
 		int16 InType = static_cast<int16>(PCGAttributeAccessorHelpers::GetMetadataTypeForProperty(InProperty));
@@ -760,20 +830,17 @@ void FPCGMetadataAttributeGeneric::Flatten()
 	if (Parent && UnderlyingProperty)
 	{
 		PCG::TUniqueScopeLock WriteLock(ValueLock);
-		
+
 		TArray<const FScriptArray*> OriginalValues;
-		
+
 		int32 ValueCount = 0;
 
 		const FPCGMetadataAttributeGeneric* Current = this;
 		while (Current)
 		{
-			if (ensure(Current->Values.Num() > 0))
-			{
-				ValueCount += Current->Values.Num();
-				OriginalValues.Add(&Current->Values);
-			}
-			
+			ValueCount += Current->Values.Num();
+			OriginalValues.Add(&Current->Values);
+
 			if (Current->Parent && Current->Parent->IsGeneric())
 			{
 				Current = static_cast<const FPCGMetadataAttributeGeneric*>(Current->Parent);
@@ -783,14 +850,11 @@ void FPCGMetadataAttributeGeneric::Flatten()
 				Current = nullptr;
 			}
 		}
-		
+
 		// Swap the current values to an empty one
 		FScriptArray NewValues{};
 		FScriptArrayHelper NewValuesHelper{UnderlyingProperty->Property.Get(), &NewValues};
-		
-		// Since each attribute in the hierarchy chain has a default value on index 0, all the values are offset by one
-		// To keep the mapping intact, we have to also copy the default value of each parent. But we'll make sure the default
-		// value (the one at index 0 at the end) is the one from `this`
+
 		if (UnderlyingProperty->bIsPlainOldData)
 		{
 			NewValuesHelper.AddUninitializedValues(ValueCount);
@@ -799,7 +863,7 @@ void FPCGMetadataAttributeGeneric::Flatten()
 		{
 			NewValuesHelper.AddValues(ValueCount);
 		}
-		
+
 		// Then copy all the values in reverse order (oldest first)
 		int32 Count = 0;
 		for (int32 ValuesIndex = OriginalValues.Num() - 1; ValuesIndex >= 0; --ValuesIndex)
@@ -809,17 +873,14 @@ void FPCGMetadataAttributeGeneric::Flatten()
 			UnderlyingProperty->Copy(NewValuesHelper.GetRawPtr(Count), OtherHelper.GetRawPtr(0), /*Count=*/OtherCount);
 			Count += OtherCount;
 		}
-		
-		// Then override the default value (index 0) by this default value
-		UnderlyingProperty->Copy(NewValues.GetData(), Values.GetData(), /*Count=*/1);
-		
+
 		// Finally move the new values in place of the old values.
 		if (TOptional<FScriptArrayHelper> Helper = MakeArrayHelper(); ensure(Helper.IsSet()))
 		{
 			Helper->MoveAssign(&NewValues);
 		}
 	}
-		
+
 	// Reset value offset, and lose parent
 	ValueKeyOffset = 0;
 	Parent = nullptr;
@@ -866,22 +927,18 @@ void FPCGMetadataAttributeGeneric::FlattenAndCompress(const TArrayView<const PCG
 	FScriptArrayHelper NewValuesHelper{UnderlyingProperty->Property.Get(), &NewValues};
 	
 	const TArray<PCGMetadataValueKey>& AllValueKeysRef = !bUseValueKeys ? AllValueKeys : AllUniqueValueKeys;
-	// Take into account the default value in the size.
 	if (UnderlyingProperty->bIsPlainOldData)
 	{
-		NewValuesHelper.AddUninitializedValues(AllValueKeysRef.Num() + 1);
+		NewValuesHelper.AddUninitializedValues(AllValueKeysRef.Num());
 	}
 	else
 	{
-		NewValuesHelper.AddValues(AllValueKeysRef.Num() + 1);
+		NewValuesHelper.AddValues(AllValueKeysRef.Num());
 	}
 	
 	ValueKeyMapping.Reserve(AllValueKeysRef.Num());
 	
-	// Copy the default value
-	UnderlyingProperty->Copy(NewValues.GetData(), Values.GetData(), /*Count=*/1);
-	
-	int32 Count = 1;
+	int32 Count = 0;
 	for (PCGMetadataValueKey ValueKey : AllValueKeysRef)
 	{
 		if (ValueKey != PCGDefaultValueKey)
@@ -940,8 +997,7 @@ void FPCGMetadataAttributeGeneric::Reset()
 	if (TOptional<FScriptArrayHelper> Helper = MakeArrayHelper(); Helper.IsSet())
 	{
 		PCG::TUniqueScopeLock WriteLock(ValueLock);
-		// Keep the default value
-		Helper->Resize(1);
+		Helper->EmptyValues();
 	}
 }
 
@@ -969,11 +1025,11 @@ FPCGMetadataAttributeBase* FPCGMetadataAttributeGeneric::Copy(FName NewName, FPC
 	else
 	{
 		AttributeCopy = new FPCGMetadataAttributeGeneric(CachedDesc, InMetadata, NewName, bAllowsInterpolation);
-		
-		if (ensure(AttributeCopy->UnderlyingProperty) && ensure(!Values.IsEmpty()) && ensure(!AttributeCopy->Values.IsEmpty()))
+
+		if (ensure(DefaultValue.IsValid()) && ensure(AttributeCopy->DefaultValue.IsValid()))
 		{
 			// Copy this default value to the attribute copy.
-			AttributeCopy->UnderlyingProperty->Copy(AttributeCopy->Values.GetData(), Values.GetData(), 1);
+			AttributeCopy->DefaultValue.CopyFrom(DefaultValue.GetRawPtr());
 		}
 	}
 
@@ -1023,20 +1079,19 @@ FPCGMetadataAttributeBase* FPCGMetadataAttributeGeneric::Copy(FName NewName, FPC
 				// Only need to lock if Current == this, since the parents are assumed to be const.
 				PCG::TSharedScopeLock ReadLock(Current->ValueLock, /*bShouldLock=*/Current == this);
 
-				// Do not copy the first value, as it is the default value, already copied/initialized by the ctor
 				TOptional<FScriptArrayHelper> ParentHelper = Current->MakeArrayHelper();
 				if (!ParentHelper.IsSet())
 				{
 					continue;
 				}
 				
-				const int32 NumValuesToCopy = ParentHelper->Num() - 1;
+				const int32 NumValuesToCopy = ParentHelper->Num();
 				if (NumValuesToCopy <= 0 || AttributeCopy->AllocateValues_Unsafe(NumValuesToCopy) == INDEX_NONE)
 				{
 					continue;	
 				}
 
-				UnderlyingProperty->Copy(Helper->GetRawPtr(1), ParentHelper->GetRawPtr(1), NumValuesToCopy);
+				UnderlyingProperty->Copy(Helper->GetRawPtr(0), ParentHelper->GetRawPtr(0), NumValuesToCopy);
 
 				// The expected value key offset is the one for this attribute (i == 0), and only if we
 				// keep the parent. Otherwise we don't have any parent, so offset should be kept at 0.
@@ -1159,7 +1214,7 @@ bool FPCGMetadataAttributeGeneric::IsEqualToDefaultValue(PCGMetadataValueKey Val
 
 void FPCGMetadataAttributeGeneric::SetDefaultValueToFirstEntry()
 {
-	if (Values.Num() == 0 || Values.Num() > 2 || !ensure(UnderlyingProperty))
+	if (Values.Num() == 0 || Values.Num() > 1 || !ensure(UnderlyingProperty) || !DefaultValue.IsValid())
 	{
 		return;
 	}
@@ -1168,14 +1223,13 @@ void FPCGMetadataAttributeGeneric::SetDefaultValueToFirstEntry()
 		PCG::TUniqueScopeLock WriteLock(ValueLock);
 
 		// First value key of this attribute.
-		const void* SrcPtr = GetReadAddressForValueKey_Unsafe(ValueKeyOffset + 1);
+		const void* SrcPtr = GetReadAddressForValueKey_Unsafe(ValueKeyOffset);
 		if (!SrcPtr)
 		{
 			return;	
 		}
 		
-		// Default value is at index 0
-		void* DestPtr = Values.GetData();
+		void* DestPtr = DefaultValue.GetRawPtr();
 		
 		UnderlyingProperty->Copy(DestPtr, SrcPtr);
 	}
@@ -1229,9 +1283,9 @@ bool FPCGMetadataAttributeGeneric::AreValuesEqual(PCGMetadataValueKey ValueKey1,
  * - Output as Set
  * - Output as Map
  * 
- * On top of that, in order to reduce as much as we can virtual function calls, we also track for the different unretrived values if they are contiguous in memory, for input AND output.
+ * On top of that, in order to reduce as much as we can virtual function calls, we also track for the different un-retrieved values if they are contiguous in memory, for input AND output.
  * Values are contiguous in memory if their index in the output array are contiguous AND the values keys associated with them are also contiguous.
- * For example, if we have indices {0, 1, 2, 5} (because we still haven't retrived values for those indices), and we have value keys {1, 2, 8, 9}, we have a contiguous
+ * For example, if we have indices {0, 1, 2, 5} (because we still haven't retrieved values for those indices), and we have value keys {1, 2, 8, 9}, we have a contiguous
  * range for the first two (their indices AND values keys are contiguous), but not for the rest.
  * So we can do 3 extractions (in this example [Index:ValueKey]): [0:1 ,1:2], [2:8] and [5:9] and optimize the first one to reduce the number of virtual calls.
  */
@@ -1279,7 +1333,18 @@ void FPCGMetadataAttributeGeneric::GetValues_Internal(const TArrayView<const PCG
 		
 		auto RetrieveValues = [this, &OutValues, &Helper](FContinuousRange& Range)
 		{
-			const uint8* SrcPtr = Helper->GetRawPtr(Range.StartValueKey);
+			const uint8* SrcPtr = nullptr;
+			if (Range.StartValueKey == PCG::Private::ExistingIndexForDefaultValue)
+			{
+				check(Range.Count == 1);
+				SrcPtr = static_cast<const uint8*>(DefaultValue.GetRawPtr());
+			}
+			else
+			{
+				SrcPtr = Helper->GetRawPtr(Range.StartValueKey);
+			}
+
+			check(SrcPtr);
 
 			// Since the values are stored in an array, they are all contiguous in memory so we can just do pointer arithmetic using the
 			// element size.
@@ -1345,23 +1410,23 @@ void FPCGMetadataAttributeGeneric::GetValues_Internal(const TArrayView<const PCG
 			
 			if (ValueKey == PCGDefaultValueKey)
 			{
-				ValueKeyIndex = 0;
+				ValueKeyIndex = PCG::Private::ExistingIndexForDefaultValue;
 			}
 			else if (ValueKey >= ValueKeyOffset)
 			{
 				int32 ValueIndex = ValueKey - ValueKeyOffset;
-				ValueKeyIndex = ValueIndex < Values.Num() ? ValueIndex : 0;
+				ValueKeyIndex = ValueIndex < Values.Num() ? ValueIndex : PCG::Private::ExistingIndexForDefaultValue;
 			}
 			else if (!ThisParent)
 			{
-				ValueKeyIndex = 0;
+				ValueKeyIndex = PCG::Private::ExistingIndexForDefaultValue;
 			}
 			else
 			{
 				bFoundAllKeys = false;
 			}
 			
-			if (ValueKeyIndex <= 0)
+			if (ValueKeyIndex < 0)
 			{
 				// Reached the end
 				if (ContinuousRange.StartIndex != INDEX_NONE)
@@ -1370,10 +1435,10 @@ void FPCGMetadataAttributeGeneric::GetValues_Internal(const TArrayView<const PCG
 					ContinuousRange.Reset();
 				}
 				
-				// Then copy the default value if the value key is 0
-				if (ValueKeyIndex == 0)
+				// Then copy the default value if the value key is PCG::Private::ExistingIndexForDefaultValue
+				if (ValueKeyIndex == PCG::Private::ExistingIndexForDefaultValue)
 				{
-					ContinuousRange.Init(Index, 0);
+					ContinuousRange.Init(Index, PCG::Private::ExistingIndexForDefaultValue);
 					RetrieveValues(ContinuousRange);
 					ContinuousRange.Reset();
 				}
@@ -1461,11 +1526,10 @@ TOptional<FScriptArrayHelper> FPCGMetadataAttributeGeneric::MakeArrayHelper() co
 
 const void* FPCGMetadataAttributeGeneric::GetReadAddressForValueKey_Unsafe(const PCGMetadataValueKey InValueKey) const
 {
-	check(!Values.IsEmpty());
-
 	if (InValueKey == PCGDefaultValueKey)
 	{
-		return Values.GetData();
+		check(DefaultValue.IsValid());
+		return DefaultValue.GetRawPtr();
 	}
 
 	if (InValueKey < ValueKeyOffset)
@@ -1503,7 +1567,7 @@ void FPCGMetadataAttributeGeneric::SetValueFromProperty(PCGMetadataEntryKey Item
 			return;
 		}
 		
-		const int32 ExistingIndex = ItemKey == PCGInvalidEntryKey && ensure(Values.Num() > 0) ? 0 : INDEX_NONE;
+		const int32 ExistingIndex = ItemKey == PCGInvalidEntryKey && ensure(DefaultValue.IsValid()) ? PCG::Private::ExistingIndexForDefaultValue : INDEX_NONE;
 
 		SetValuesFromValueKeys({ItemKey}, AddValues_Internal(FInValues{TInPlaceType<FInValuesByValue>{}, SrcPtr, 1}, 1, ExistingIndex), /*bResetDefaultValueOnDefaultValueKey=*/ItemKey==PCGInvalidEntryKey);
 	}
diff --git a/Engine/Plugins/PCG/Source/PCG/Public/Metadata/PCGMetadataAttributeGeneric.h b/Engine/Plugins/PCG/Source/PCG/Public/Metadata/PCGMetadataAttributeGeneric.h
--- a/Engine/Plugins/PCG/Source/PCG/Public/Metadata/PCGMetadataAttributeGeneric.h
+++ b/Engine/Plugins/PCG/Source/PCG/Public/Metadata/PCGMetadataAttributeGeneric.h
@@ -14,6 +14,7 @@
 #include "Misc/TVariant.h"
 #include "StructUtils/PropertyBag.h"
 
+struct FPCGAttributeDefaultValue;
 class FPCGMetadataAttributeGeneric;
 
 namespace PCG::Private
@@ -31,7 +32,9 @@ namespace PCG::Private
 		TIsTArray_V<Container> 
 		|| TIsTArrayView_V<Container>
 		|| TIsStaticArray<Container>::value;
-		
+
+	/** Magic number for AddValues to know if we need to write to the DefaultValue or not. */
+	constexpr int32 ExistingIndexForDefaultValue = -2;
 }
 
 struct FPCGAttributeProperty
@@ -41,10 +44,11 @@ public:
 	PCG_API ~FPCGAttributeProperty();
 
 	PCG_API void Copy(void* DestPtr, const void* SrcPtr, int32 Count = 1);
-	PCG_API void AllocateValues(int32 Count = 1);
+	PCG_API bool IsValid() const;
 
 private:
 	friend FPCGMetadataAttributeGeneric;
+	friend FPCGAttributeDefaultValue;
 
 	/** Unsafe method that will apply pointer arithmetic to get to the element in the array at the specified index. */
 	const void* GetPtrInArray(const void* InArrayStart, int32 Index) const;
@@ -55,6 +59,29 @@ private:
 	bool bIsPlainOldData = false;
 };
 
+// Simple wrapper around a raw pointer to hold the default value to mimic how it was stored in the typed attributes
+struct FPCGAttributeDefaultValue
+{
+	FPCGAttributeDefaultValue() = default;
+	PCG_API ~FPCGAttributeDefaultValue();
+	
+	// Since this is an internal struct and attributes are never copied with constructor/assignment, just delete the copy/move
+	UE_NONCOPYABLE(FPCGAttributeDefaultValue);
+
+	PCG_API void Init(TSharedPtr<FPCGAttributeProperty> InUnderlyingProperty);
+	PCG_API void CopyFrom(const void* SrcData);
+	
+	void* GetRawPtr() { return Memory; }
+	const void* GetRawPtr() const { return Memory; };
+	bool IsValid() const { return Memory != nullptr; }
+
+private:
+	void Destroy();
+	
+	TSharedPtr<FPCGAttributeProperty> UnderlyingProperty;
+	void* Memory = nullptr;
+};
+
 class FPCGMetadataAttributeGeneric : public FPCGMetadataAttributeBase
 {
 public:
@@ -83,7 +110,7 @@ public:
 	const FPCGMetadataAttributeDesc& GetAttributeDesc() const { return CachedDesc; }
 
 	PCG_API virtual PCGMetadataValueKey GetValueKeyOffsetForChild() const override;
-
+	
 	/**
 	 * GetValues section
 	 * Templated code to support all container types.
@@ -254,8 +281,8 @@ public:
 	{
 		if (TOptional<FScriptArrayHelper> Helper = MakeArrayHelper(); Helper.IsSet())
 		{
-			check(Helper->Num() > 0);
-			AddValues(MakeArrayView(&Value, 1), /*ExistingIndex=*/0);
+			check(DefaultValue.IsValid());
+			AddValues(MakeArrayView(&Value, 1), /*ExistingIndex=*/PCG::Private::ExistingIndexForDefaultValue, /*Count=*/1);
 		}
 	}
 
@@ -573,7 +600,8 @@ private:
 	FPCGMetadataAttributeDesc CachedDesc;
 	TSharedPtr<FPCGAttributeProperty> UnderlyingProperty;
 	FScriptArray Values;
-	
+	FPCGAttributeDefaultValue DefaultValue;
+
 	mutable PCG::FSharedLock ValueLock;
 	PCGMetadataValueKey ValueKeyOffset = 0;
 };
diff --git a/Engine/Plugins/PCG/Tests/Private/Metadata/PCGMetadataAttributeGenericTests.cpp b/Engine/Plugins/PCG/Tests/Private/Metadata/PCGMetadataAttributeGenericTests.cpp
--- a/Engine/Plugins/PCG/Tests/Private/Metadata/PCGMetadataAttributeGenericTests.cpp
+++ b/Engine/Plugins/PCG/Tests/Private/Metadata/PCGMetadataAttributeGenericTests.cpp
@@ -866,8 +866,8 @@ TEST_CASE_METHOD(PCGTests::FPCGBaseTestWithBasicContext, "PCG::AttributeGeneric:
 	// Set the values
 	GenericAttribute->SetValues<FString>(EntryKeys, MakeArrayView(Values.GetData(), NumPoints / 2));
 	
-	// Verify that we have only 10 Value keys (including the default value one)
-	REQUIRE_EQUAL(GenericAttribute->GetValueKeyOffsetForChild(), 10);
+	// Verify that we have only 9 Value keys (the tenth one is the default value).
+	REQUIRE_EQUAL(GenericAttribute->GetValueKeyOffsetForChild(), 9);
 	
 	// Then add a mix of existing strings and new strings, to test that we only add what does not exist
 	for (int32 i = NumPoints/2; i < NumPoints; ++i)
@@ -888,8 +888,8 @@ TEST_CASE_METHOD(PCGTests::FPCGBaseTestWithBasicContext, "PCG::AttributeGeneric:
 	// Set the values
 	GenericAttribute->SetValues<FString>(EntryKeys, MakeArrayView(Values.GetData() + (NumPoints / 2), NumPoints / 2));
 	
-	// Verify that we have only 20 Value keys (including the default value one)
-	REQUIRE_EQUAL(GenericAttribute->GetValueKeyOffsetForChild(), 20);
+	// Verify that we have only 19 Value keys (the twentieth one is the default value)
+	REQUIRE_EQUAL(GenericAttribute->GetValueKeyOffsetForChild(), 19);
 	
 	// Verify that all the values can be retrieved.
 	TStaticArray<const FString*, NumPoints> ReadValues{};
diff --git a/Engine/Plugins/PCG/Tests/Private/Metadata/PCGMetadataFlattenTests.cpp b/Engine/Plugins/PCG/Tests/Private/Metadata/PCGMetadataFlattenTests.cpp
--- a/Engine/Plugins/PCG/Tests/Private/Metadata/PCGMetadataFlattenTests.cpp
+++ b/Engine/Plugins/PCG/Tests/Private/Metadata/PCGMetadataFlattenTests.cpp
@@ -329,11 +329,10 @@ TEST_CASE_METHOD(PCGTests::FPCGBaseTestWithBasicContext, "PCG::Metadata::Flatten
 
 	SECTION("Flatten metadata has the right number of entries/values")
 	{
-		// + 1 for the default value
 		REQUIRE_EQUAL(FinalPointData->Metadata->GetItemCountForChild(), 5);
-		REQUIRE_EQUAL(Attribute1->GetValueKeyOffsetForChild(), 5 + 1);
-		REQUIRE_EQUAL(Attribute2->GetValueKeyOffsetForChild(), 1 + 1);
-		REQUIRE_EQUAL(Attribute3->GetValueKeyOffsetForChild(), 5 + 1);
+		REQUIRE_EQUAL(Attribute1->GetValueKeyOffsetForChild(), 5);
+		REQUIRE_EQUAL(Attribute2->GetValueKeyOffsetForChild(), 1);
+		REQUIRE_EQUAL(Attribute3->GetValueKeyOffsetForChild(), 5);
 		REQUIRE_EQUAL(FinalPointData->GetNumPoints(), NumPoints);
 	}
 
