[PCG] Data Override Systems - First Pass

+ PCGDataOverride mechanics, involving data storage and key creation
+ Point-specific deltas
* Application of the override phase in PCGElement.h
* Override in PrepareData for PCGStaticMeshSpawner

#Synced-CL 51160565

diff --git a/Engine/Plugins/PCG/Source/PCG/Private/Elements/PCGStaticMeshSpawner.cpp b/Engine/Plugins/PCG/Source/PCG/Private/Elements/PCGStaticMeshSpawner.cpp
--- a/Engine/Plugins/PCG/Source/PCG/Private/Elements/PCGStaticMeshSpawner.cpp
+++ b/Engine/Plugins/PCG/Source/PCG/Private/Elements/PCGStaticMeshSpawner.cpp
@@ -24,6 +24,7 @@
 #include "GrassInstancedStaticMeshComponent.h"
 #include "Components/InstancedStaticMeshComponent.h"
 #include "Engine/StaticMesh.h"
+#include "Graph/DataOverride/PCGDataOverride.h"
 #include "Materials/MaterialInterface.h"
 #include "UObject/Package.h"
 
@@ -892,6 +893,11 @@ void FPCGStaticMeshSpawnerElement::SpawnStaticMeshInstances(FPCGStaticMeshSpawne
 	}
 }
 
+EPCGDataOverridePhase FPCGStaticMeshSpawnerElement::GetDataOverridePhase() const
+{
+	return EPCGDataOverridePhase::PrepareData;
+}
+
 void FPCGStaticMeshSpawnerElement::AbortInternal(FPCGContext* InContext) const
 {
 	TRACE_CPUPROFILER_EVENT_SCOPE(FPCGStaticMeshSpawnerElement::AbortInternal);
diff --git a/Engine/Plugins/PCG/Source/PCG/Private/Graph/DataOverride/PCGDataOverride.cpp b/Engine/Plugins/PCG/Source/PCG/Private/Graph/DataOverride/PCGDataOverride.cpp
new file mode 100644
--- /dev/null
+++ b/Engine/Plugins/PCG/Source/PCG/Private/Graph/DataOverride/PCGDataOverride.cpp
@@ -0,0 +1,107 @@
+// Copyright Epic Games, Inc. All Rights Reserved.
+
+#include "Graph/DataOverride/PCGDataOverride.h"
+
+FPCGDeltaKey::FPCGDeltaKey(const FXxHash64 InHash, const FName InDeltaLabel)
+#if WITH_EDITOR
+	: DeltaLabel(InDeltaLabel)
+#endif // WITH_EDITOR
+{
+	if (InDeltaLabel != NAME_None)
+	{
+		FXxHash64Builder Builder;
+		const uint64 RawHash = InHash.Hash;
+		Builder.Update(&RawHash, sizeof(RawHash));
+		const FString LabelString = InDeltaLabel.ToString();
+		Builder.Update(*LabelString, LabelString.Len() * sizeof(TCHAR));
+		Hash = Builder.Finalize().Hash;
+	}
+	else
+	{
+		Hash = InHash.Hash;
+	}
+}
+
+TInstancedStruct<FPCGDeltaBase>* FPCGDeltaCollection::Find(const FPCGDeltaKey& InKey)
+{
+	return Deltas.Find(InKey);
+}
+
+const TInstancedStruct<FPCGDeltaBase>* FPCGDeltaCollection::Find(const FPCGDeltaKey& InKey) const
+{
+	return Deltas.Find(InKey);
+}
+
+bool FPCGDeltaCollection::Contains(const FPCGDeltaKey& InKey) const
+{
+	return Find(InKey) != nullptr;
+}
+
+#if WITH_EDITOR
+TInstancedStruct<FPCGDeltaBase>& FPCGDeltaCollection::Add(const FPCGDeltaKey& InKey, TInstancedStruct<FPCGDeltaBase>&& InDelta)
+{
+	TInstancedStruct<FPCGDeltaBase>& Struct = Deltas.Add(InKey, MoveTemp(InDelta));
+	if (FPCGDeltaBase* Delta = Struct.GetMutablePtr<FPCGDeltaBase>())
+	{
+		Delta->ComputedKey = InKey;
+	}
+
+	return Struct;
+}
+
+bool FPCGDeltaCollection::Remove(const FPCGDeltaKey& InKey)
+{
+	return Deltas.Remove(InKey) > 0;
+}
+
+void FPCGDeltaCollection::Empty()
+{
+	Deltas.Empty();
+}
+#endif // WITH_EDITOR
+
+int32 FPCGDeltaCollection::Num() const
+{
+	return Deltas.Num();
+}
+
+bool FPCGDeltaCollection::IsEmpty() const
+{
+	return Deltas.IsEmpty();
+}
+
+void FPCGDeltaCollection::ForEachDelta(TFunctionRef<void(const FPCGDeltaKey&, TInstancedStruct<FPCGDeltaBase>&)> Func)
+{
+	for (TPair<FPCGDeltaKey, TInstancedStruct<FPCGDeltaBase>>& Pair : Deltas)
+	{
+		Func(Pair.Key, Pair.Value);
+	}
+}
+
+void FPCGDeltaCollection::ForEachDelta(TFunctionRef<void(const FPCGDeltaKey&, const TInstancedStruct<FPCGDeltaBase>&)> Func) const
+{
+	for (const TPair<FPCGDeltaKey, TInstancedStruct<FPCGDeltaBase>>& Pair : Deltas)
+	{
+		Func(Pair.Key, Pair.Value);
+	}
+}
+
+FName FPCGDeltaBase::GetDeltaName() const
+{
+	return NAME_None;
+}
+
+PCGIndexing::FPCGIndexCollection FPCGDeltaBase::FilterCandidates(const UPCGData* InData) const
+{
+	return PCGIndexing::FPCGIndexCollection::Invalid();
+}
+
+int32 FPCGDeltaBase::Resolve(const UPCGData* InData, const PCGIndexing::FPCGIndexCollection& FilteredCandidates, const FPCGDeltaSettings& DeltaSettings) const
+{
+	return INDEX_NONE;
+}
+
+bool FPCGDeltaBase::Apply(UPCGData* InData, const int32 ResolvedIndex) const
+{
+	return false;
+}
diff --git a/Engine/Plugins/PCG/Source/PCG/Private/Graph/DataOverride/PCGDataOverrideHelpers.cpp b/Engine/Plugins/PCG/Source/PCG/Private/Graph/DataOverride/PCGDataOverrideHelpers.cpp
new file mode 100644
--- /dev/null
+++ b/Engine/Plugins/PCG/Source/PCG/Private/Graph/DataOverride/PCGDataOverrideHelpers.cpp
@@ -0,0 +1,235 @@
+// Copyright Epic Games, Inc. All Rights Reserved.
+
+#include "Graph/DataOverride/PCGDataOverrideHelpers.h"
+
+#include "PCGContext.h"
+#include "Data/PCGPointArrayData.h"
+#include "Graph/PCGSourceDataContainer.h"
+#include "Graph/PCGStackContext.h"
+
+namespace PCG::DataOverride
+{
+	namespace Helpers
+	{
+		FXxHash64 ComputeDataHash(const FPCGStack& InStack)
+		{
+			FXxHash64Builder HashBuilder;
+
+			for (const FPCGStackFrame& Frame : InStack.GetStackFrames())
+			{
+				if (Frame.IsLoopIndexFrame())
+				{
+					const int32 LoopIndex = Frame.LoopIndex;
+					HashBuilder.Update(&LoopIndex, sizeof(LoopIndex));
+				}
+				else
+				{
+					// Skip the execution source, since it is not relevant to the local data.
+					const UObject* FrameObject = Frame.GetObject_NoGuard();
+					if (!FrameObject || FrameObject->Implements<UPCGGraphExecutionSource>())
+					{
+						continue;
+					}
+
+					const FString Path = Frame.Object.ToString();
+					HashBuilder.Update(*Path, Path.Len() * sizeof(TCHAR));
+				}
+			}
+
+			const FXxHash64 Result = HashBuilder.Finalize();
+			return Result;
+		}
+
+		FXxHash64 ComputeDataHashWithPin(const FPCGStack& InStack, const FName PinLabel)
+		{
+			const FXxHash64 PinlessKey = ComputeDataHash(InStack);
+			return ComputeDataHashWithPin(PinlessKey, PinLabel);
+		}
+
+		FXxHash64 ComputeDataHashWithPin(const FXxHash64 NodeKey, const FName PinLabel)
+		{
+			if (PinLabel == NAME_None)
+			{
+				return NodeKey;
+			}
+
+			FXxHash64Builder HashBuilder;
+			HashBuilder.Update(&NodeKey, sizeof(NodeKey));
+
+			const FString PinString = PinLabel.ToString();
+			HashBuilder.Update(*PinString, PinString.Len() * sizeof(TCHAR));
+
+			return HashBuilder.Finalize();
+		}
+	} // namespace Helpers
+
+	namespace Keys
+	{
+		FXxHash64 ComputePositionHash(const FVector& InLocation, const double Tolerance)
+		{
+			// 0 Tolerance is an exact spatial hash
+			if (FMath::Max(Tolerance, 0.0) == 0.0)
+			{
+				return FXxHash64::HashBuffer(&InLocation, sizeof(InLocation));
+			}
+
+			// IMPLEMENTATION NOTE: A significant change in location or bounds will orphan this override.
+			const int64 GridX = FMath::RoundToInt(InLocation.X / Tolerance);
+			const int64 GridY = FMath::RoundToInt(InLocation.Y / Tolerance);
+			const int64 GridZ = FMath::RoundToInt(InLocation.Z / Tolerance);
+
+			const int64 Grid[] = {GridX, GridY, GridZ};
+			return FXxHash64::HashBuffer(&Grid, sizeof(Grid));
+		}
+
+		FXxHash64 ComputeTransformHash(const FTransform& InTransform, const double Tolerance)
+		{
+			if (FMath::Max(Tolerance, 0.0) == 0.0)
+			{
+				return FXxHash64::HashBuffer(&InTransform, sizeof(InTransform));
+			}
+
+			const FVector Location = InTransform.GetLocation();
+			const int64 GridX = FMath::RoundToInt(Location.X / Tolerance);
+			const int64 GridY = FMath::RoundToInt(Location.Y / Tolerance);
+			const int64 GridZ = FMath::RoundToInt(Location.Z / Tolerance);
+
+			FXxHash64Builder HashBuilder;
+			const int64 Grid[] = {GridX, GridY, GridZ};
+			HashBuilder.Update(Grid, sizeof(Grid));
+
+			// @todo_pcg: Scale doesn't follow the same spatial tolerance as location. Might experiment with a
+			// separate tolerance or leave it directly hashed.
+			const FQuat Rotation = InTransform.GetRotation();
+			HashBuilder.Update(&Rotation, sizeof(Rotation));
+			const FVector Scale = InTransform.GetScale3D();
+			HashBuilder.Update(&Scale, sizeof(Scale));
+
+			return HashBuilder.Finalize();
+		}
+
+		FXxHash64 ComputeAABBHash(const FBox& InBounds, const double Tolerance)
+		{
+			if (FMath::Max(Tolerance, 0.0) == 0.0)
+			{
+				return FXxHash64::HashBuffer(&InBounds, sizeof(InBounds));
+			}
+
+			const int64 Grid[] = {
+				FMath::RoundToInt(InBounds.Min.X / Tolerance),
+				FMath::RoundToInt(InBounds.Min.Y / Tolerance),
+				FMath::RoundToInt(InBounds.Min.Z / Tolerance),
+				FMath::RoundToInt(InBounds.Max.X / Tolerance),
+				FMath::RoundToInt(InBounds.Max.Y / Tolerance),
+				FMath::RoundToInt(InBounds.Max.Z / Tolerance)
+			};
+
+			return FXxHash64::HashBuffer(&Grid, sizeof(Grid));
+		}
+	} // namespace Keys
+
+	bool HasNodeLevelDataOverrides(FPCGContext* Context)
+	{
+		const FPCGSourceDataContainer* DataContainer = Context ? Context->GetSourceDataContainer() : nullptr;
+		const FPCGStack* Stack = DataContainer ? Context->GetStack() : nullptr;
+		if (!DataContainer || !Stack || DataContainer->IsEmpty())
+		{
+			return false;
+		}
+
+		// First check if the node has data override deltas.
+		const FXxHash64 NodeKey = Helpers::ComputeDataHash(*Stack);
+		const FPCGSourceDataStorageKey NodeStorageKey(Constants::DefaultOverrideLabel, NodeKey.Hash);
+		return DataContainer->Get<FPCGDeltaCollection>(NodeStorageKey).IsValid();
+	}
+
+	bool HasPinLevelDataOverrides(FPCGContext* Context, const FPCGTaggedData& InTaggedData)
+	{
+		const FPCGSourceDataContainer* DataContainer = Context ? Context->GetSourceDataContainer() : nullptr;
+		const FPCGStack* Stack = DataContainer ? Context->GetStack() : nullptr;
+		if (!DataContainer || !Stack || DataContainer->IsEmpty())
+		{
+			return false;
+		}
+
+		const FXxHash64 NodeKey = Helpers::ComputeDataHash(*Stack);
+		const FXxHash64 PinKey = Helpers::ComputeDataHashWithPin(NodeKey, InTaggedData.Pin);
+		const FPCGSourceDataStorageKey PinStorageKey(Constants::DefaultOverrideLabel, PinKey.Hash);
+		return DataContainer->Get<FPCGDeltaCollection>(PinStorageKey).IsValid();
+	}
+
+	bool HasAnyDataOverrides(FPCGContext* Context, const TArrayView<const FPCGTaggedData> InTaggedData)
+	{
+		if (HasNodeLevelDataOverrides(Context))
+		{
+			return true;
+		}
+
+		for (const FPCGTaggedData& TaggedData : InTaggedData)
+		{
+			if (HasPinLevelDataOverrides(Context, TaggedData))
+			{
+				return true;
+			}
+		}
+
+		return false;
+	}
+
+	// @todo_pcg: Still requires collision handling. Will be done in future pass.
+	void ApplyDataOverrides(FPCGContext* Context, const TArrayView<FPCGTaggedData> OutTaggedData)
+	{
+		const FPCGSourceDataContainer* DataContainer = Context ? Context->GetSourceDataContainer() : nullptr;
+		if (!DataContainer)
+		{
+			return;
+		}
+
+		const FPCGStack* Stack = Context->GetStack();
+		if (!Stack)
+		{
+			return;
+		}
+
+		const FXxHash64 NodeKey = Helpers::ComputeDataHash(*Stack);
+		for (FPCGTaggedData& TaggedData : OutTaggedData)
+		{
+			// @todo_pcg: This should be part of a registry. For now, default to points, but should be element
+			// We have data ownership at either phase and thus const_cast is okay.
+			if (UPCGPointArrayData* PointData = const_cast<UPCGPointArrayData*>(Cast<UPCGPointArrayData>(TaggedData.Data)))
+			{
+				FPCGSourceDataStorageKey PinDataKey(Constants::DefaultOverrideLabel, Helpers::ComputeDataHashWithPin(NodeKey, TaggedData.Pin).Hash);
+
+				// Check to see if the pin itself has overrides
+				FConstSharedStruct SharedStruct = DataContainer->Get<FPCGDeltaCollection>(PinDataKey);
+				const FPCGDeltaCollection* DeltaCollection = SharedStruct.IsValid() ? SharedStruct.GetPtr<const FPCGDeltaCollection>() : nullptr;
+
+				// If not, check the node level
+				if (!DeltaCollection)
+				{
+					FPCGSourceDataStorageKey NodeDataKey(Constants::DefaultOverrideLabel, NodeKey.Hash);
+					SharedStruct = DataContainer->Get<FPCGDeltaCollection>(NodeDataKey);
+					DeltaCollection = SharedStruct.IsValid() ? SharedStruct.GetPtr<const FPCGDeltaCollection>() : nullptr;
+				}
+
+				if (DeltaCollection && !DeltaCollection->IsEmpty())
+				{
+					// @todo_pcg: Sort deltas (edits->deletions->insertions) to avoid index instability.
+					// @todo_pcg: As a future optimization pass, aggregate candidates to avoid redundant key computation.
+					DeltaCollection->ForEachDelta([DeltaCollection, PointData](const FPCGDeltaKey&, const TInstancedStruct<FPCGDeltaBase>& DeltaStruct)
+					{
+						const FPCGDeltaBase* Delta = DeltaStruct.GetPtr<FPCGDeltaBase>();
+						if (!Delta)
+						{
+							return;
+						}
+
+						const PCGIndexing::FPCGIndexCollection Candidates = Delta->FilterCandidates(PointData);
+						const int32 ResolvedIndex = Delta->Resolve(PointData, Candidates, DeltaCollection->Settings);
+						Delta->Apply(PointData, ResolvedIndex);
+					});
+				}
+			}
+		}
+	}
+} // namespace PCG::DataOverride
diff --git a/Engine/Plugins/PCG/Source/PCG/Private/Graph/DataOverride/PCGDataOverridePoints.cpp b/Engine/Plugins/PCG/Source/PCG/Private/Graph/DataOverride/PCGDataOverridePoints.cpp
new file mode 100644
--- /dev/null
+++ b/Engine/Plugins/PCG/Source/PCG/Private/Graph/DataOverride/PCGDataOverridePoints.cpp
@@ -0,0 +1,192 @@
+// Copyright Epic Games, Inc. All Rights Reserved.
+
+#include "Graph/DataOverride/PCGDataOverridePoints.h"
+
+#include "Data/PCGBasePointData.h"
+#include "Data/PCGPointArrayData.h"
+#include "Graph/DataOverride/PCGDataOverrideHelpers.h"
+
+PCGIndexing::FPCGIndexCollection FPCGPointDeltaBase::FilterCandidates(const UPCGData* InData) const
+{
+	if (const UPCGBasePointData* PointData = Cast<UPCGBasePointData>(InData))
+	{
+		PCGIndexing::FPCGIndexCollection Collection(PointData->GetNumPoints());
+
+		TArray<int32> CandidateIndices;
+		const PCGPointOctree::FPointOctree& Octree = PointData->GetPointOctree();
+
+		Octree.FindElementsWithBoundsTest(
+			Bounds,
+			[&CandidateIndices](const PCGPointOctree::FPointRef& Ref)
+			{
+				CandidateIndices.Add(Ref.Index);
+			});
+
+		Collection += CandidateIndices;
+
+		return Collection;
+	}
+
+	return PCGIndexing::FPCGIndexCollection::Invalid();
+}
+
+int32 FPCGPointDeltaBase::Resolve(const UPCGData* InData, const PCGIndexing::FPCGIndexCollection& FilteredCandidates, const FPCGDeltaSettings& DeltaSettings) const
+{
+	using namespace PCG::DataOverride;
+
+	int32 ResolvedIndex = INDEX_NONE;
+	const UPCGBasePointData* PointData = Cast<UPCGBasePointData>(InData);
+	if (!PointData)
+	{
+		return ResolvedIndex;
+	}
+
+	auto ComputeKey = [&DeltaSettings](const FTransform& InTransform, const FBox& InWorldBounds) -> FPCGDeltaKey
+	{
+		switch (DeltaSettings.KeyTarget)
+		{
+			case EPCGSpatialKeyTarget::Position:
+				return Keys::FPCGLocationDeltaKey(InTransform.GetLocation(), DeltaSettings.SpatialTolerance);
+			case EPCGSpatialKeyTarget::AABB:
+				return Keys::FPCGAABBDeltaKey(InWorldBounds, DeltaSettings.SpatialTolerance);
+			case EPCGSpatialKeyTarget::Transform:
+				return Keys::FPCGTransformDeltaKey(InTransform, DeltaSettings.SpatialTolerance);
+			default:
+				return {};
+		}
+	};
+
+	FPCGDeltaKey DeltaKey = ComputeKey(OriginalTransform, Bounds);
+
+	FilteredCandidates.ForEach([&ResolvedIndex, &DeltaKey, ComputeKey, PointData](const PCGIndexing::FPCGIndexRange& IndexRange)
+	{
+		for (int32 Index = IndexRange.StartIndex; Index < IndexRange.EndIndex; ++Index)
+		{
+			FTransform Transform = PointData->GetTransform(Index);
+			FBox LocalBounds = PointData->GetLocalBounds(Index);
+			FBox WorldBounds = LocalBounds.TransformBy(Transform);
+
+			FPCGDeltaKey CandidateKey = ComputeKey(Transform, WorldBounds);
+			if (DeltaKey == CandidateKey)
+			{
+				ResolvedIndex = Index;
+				return true;
+			}
+		}
+
+		return false;
+	});
+
+	return ResolvedIndex;
+}
+
+FName FPCGPointEditDelta::GetDeltaName() const
+{
+	return GetDeltaNameStatic();
+}
+
+bool FPCGPointEditDelta::Apply(UPCGData* InData, const int32 ResolvedIndex) const
+{
+	if (UPCGPointArrayData* PointData = Cast<UPCGPointArrayData>(InData))
+	{
+		if (PointData->IsEmpty())
+		{
+			return false;
+		}
+
+		TPCGValueRange<FTransform> TransformRange = PointData->GetTransformValueRange();
+
+		if (TransformRange.IsValidIndex(ResolvedIndex))
+		{
+			if (bOverridePosition)
+			{
+				TransformRange[ResolvedIndex].SetLocation(TransformOverride.GetLocation());
+			}
+
+			if (bOverrideRotation)
+			{
+				TransformRange[ResolvedIndex].SetRotation(TransformOverride.GetRotation());
+			}
+
+			if (bOverrideScale)
+			{
+				TransformRange[ResolvedIndex].SetScale3D(TransformOverride.GetScale3D());
+			}
+
+			return true;
+		}
+	}
+	// @todo_pcg: Add the path to evaluate the other point data
+
+	return false;
+}
+
+FName FPCGPointDeletionDelta::GetDeltaName() const
+{
+	return GetDeltaNameStatic();
+}
+
+bool FPCGPointDeletionDelta::Apply(UPCGData* InData, const int32 ResolvedIndex) const
+{
+	if (UPCGPointArrayData* PointData = Cast<UPCGPointArrayData>(InData))
+	{
+		if (PointData->IsEmpty() || ResolvedIndex == INDEX_NONE)
+		{
+			return false;
+		}
+
+		TPCGValueRange<FTransform> TransformRange = PointData->GetTransformValueRange();
+		if (TransformRange.IsValidIndex(ResolvedIndex))
+		{
+			// @todo_pcg: To be discussed and fixed in second pass. For now, we'll set the scale to 0,0,0. We should
+			// delete the point, but deletion of a point introduces index instability, so it should happen last or it may thrash the octree.
+			TransformRange[ResolvedIndex].SetScale3D(FVector::ZeroVector);
+			return true;
+		}
+	}
+
+	return false;
+}
+
+FName FPCGPointInsertionDelta::GetDeltaName() const
+{
+	return GetDeltaNameStatic();
+}
+
+PCGIndexing::FPCGIndexCollection FPCGPointInsertionDelta::FilterCandidates(const UPCGData* InData) const
+{
+	// Filtering not needed. Points are simply injected.
+	return PCGIndexing::FPCGIndexCollection::Invalid();
+}
+
+bool FPCGPointInsertionDelta::Apply(UPCGData* InData, const int32 ResolvedIndex) const
+{
+	if (UPCGPointArrayData* PointData = Cast<UPCGPointArrayData>(InData))
+	{
+		if (InjectedPoints.IsEmpty())
+		{
+			return false;
+		}
+
+		const int32 OriginalPointCount = PointData->GetNumPoints();
+		if (!ensure(std::numeric_limits<int32>::max() - OriginalPointCount >= InjectedPoints.Num()))
+		{
+			// Too many point injections
+			return false;
+		}
+
+		PointData->SetNumPoints(OriginalPointCount + InjectedPoints.Num());
+
+		FPCGPointValueRanges PointRanges(PointData, /*bAllocate=*/true);
+
+		for (int32 I = 0; I < InjectedPoints.Num(); ++I)
+		{
+			PointRanges.SetFromPoint(OriginalPointCount + I, InjectedPoints[I]);
+		}
+
+		return true;
+	}
+	// @todo_pcg: Add the path to evaluate the other point data
+
+	return false;
+}
diff --git a/Engine/Plugins/PCG/Source/PCG/Private/PCGElement.cpp b/Engine/Plugins/PCG/Source/PCG/Private/PCGElement.cpp
--- a/Engine/Plugins/PCG/Source/PCG/Private/PCGElement.cpp
+++ b/Engine/Plugins/PCG/Source/PCG/Private/PCGElement.cpp
@@ -10,6 +10,8 @@
 #include "Data/PCGBasePointData.h"
 #include "Data/PCGPointArrayData.h"
 #include "Data/PCGPointData.h"
+#include "Graph/DataOverride/PCGDataOverride.h"
+#include "Graph/DataOverride/PCGDataOverrideHelpers.h"
 #include "Graph/PCGGraphCache.h"
 #include "Helpers/PCGActorHelpers.h"
 #include "Helpers/PCGHelpers.h"
@@ -547,6 +549,31 @@ bool IPCGElement::ConvertInputsIfNeeded(FPCGContext* Context) const
 		}
 	}
 
+	if (EnumHasAnyFlags(GetDataOverridePhase(), EPCGDataOverridePhase::PrepareData))
+	{
+		TRACE_CPUPROFILER_EVENT_SCOPE(IPCGElement::ConvertInputsIfNeeded::ApplyDataOverrides);
+
+		const bool bHasNodeLevelOverrides = PCG::DataOverride::HasNodeLevelDataOverrides(Context);
+		bool bHasAnyOverrides = bHasNodeLevelOverrides;
+
+		// Need to duplicate the data in pre-application phase to avoid stomping data shared by other elements.
+		// @todo_pcg: Consider data stealing if possible.
+		for (FPCGTaggedData& TaggedData : Context->InputData.TaggedData)
+		{
+			// If we have node overrides, all data will be copied. If not, check for pin level overrides.
+			if (TaggedData.Data && (bHasNodeLevelOverrides || PCG::DataOverride::HasPinLevelDataOverrides(Context, TaggedData)))
+			{
+				TaggedData.Data = TaggedData.Data->DuplicateData(Context);
+				bHasAnyOverrides = true;
+			}
+		}
+
+		if (bHasAnyOverrides)
+		{
+			PCG::DataOverride::ApplyDataOverrides(Context, Context->InputData.TaggedData);
+		}
+	}
+
 	Context->bConvertInputsIfNeededDone = true;
 	return true;
 }
@@ -581,6 +608,15 @@ void IPCGElement::PostExecute(FPCGContext* Context) const
 		}
 	}
 
+	// Data overrides must be applied before the Crc
+	if (!Context->OutputData.TaggedData.IsEmpty()
+		&& EnumHasAnyFlags(GetDataOverridePhase(), EPCGDataOverridePhase::PostExecute)
+		&& PCG::DataOverride::HasAnyDataOverrides(Context, Context->OutputData.TaggedData))
+	{
+		TRACE_CPUPROFILER_EVENT_SCOPE(IPCGElement::PostExecute::ApplyDataOverrides);
+		PCG::DataOverride::ApplyDataOverrides(Context, Context->OutputData.TaggedData);
+	}
+
 	if (Context->IsCacheEnabled())
 	{
 		{
@@ -1238,6 +1274,14 @@ PRAGMA_ENABLE_DEPRECATION_WARNINGS
 			Crc.Combine(InParams.ExecutionSource->GetExecutionState().GetSeed());
 		}
 
+		const FPCGSourceDataContainer* DataContainer = InParams.ExecutionSource ? InParams.ExecutionSource->GetExecutionState().GetSourceDataContainer() : nullptr;
+		if (DataContainer && !DataContainer->IsEmpty())
+		{
+			// @todo_pcg: This ties a non-empty data container to every element's Crc. If the data within the container is changed,
+			// every element's Crc will be invalidated. To be addressed in a secondary pass.
+			Crc.Combine(DataContainer->GetDirtyGeneration());
+		}
+
 		OutCrc = Crc;
 	}
 	else
diff --git a/Engine/Plugins/PCG/Source/PCG/Public/Elements/PCGStaticMeshSpawner.h b/Engine/Plugins/PCG/Source/PCG/Public/Elements/PCGStaticMeshSpawner.h
--- a/Engine/Plugins/PCG/Source/PCG/Public/Elements/PCGStaticMeshSpawner.h
+++ b/Engine/Plugins/PCG/Source/PCG/Public/Elements/PCGStaticMeshSpawner.h
@@ -152,6 +152,7 @@ protected:
 	UE_API virtual void AbortInternal(FPCGContext* Context) const override;
 	UE_API void SpawnStaticMeshInstances(FPCGStaticMeshSpawnerContext* Context, const FPCGStaticMeshSpawnerContext::FPackedInstanceListData& InstanceListData, const FPCGMeshInstanceList& InstanceList, const FPCGPackedCustomData* InPackedCustomData) const;
 	virtual bool SupportsBasePointDataInputs(FPCGContext* InContext) const { return true; }
+	virtual EPCGDataOverridePhase GetDataOverridePhase() const override;
 };
 
 #undef UE_API
diff --git a/Engine/Plugins/PCG/Source/PCG/Public/Graph/DataOverride/PCGDataOverride.h b/Engine/Plugins/PCG/Source/PCG/Public/Graph/DataOverride/PCGDataOverride.h
new file mode 100644
--- /dev/null
+++ b/Engine/Plugins/PCG/Source/PCG/Public/Graph/DataOverride/PCGDataOverride.h
@@ -0,0 +1,175 @@
+// Copyright Epic Games, Inc. All Rights Reserved.
+
+#pragma once
+
+#include "Helpers/Parsing/PCGIndexing.h"
+#include "StructUtils/InstancedStruct.h"
+
+#include "Hash/xxhash.h"
+
+#include "PCGDataOverride.generated.h"
+
+class UPCGData;
+struct FPCGStack;
+struct FPCGContext;
+struct FPCGPoint;
+
+namespace PCG::DataOverride
+{
+	namespace Constants
+	{
+		inline constexpr FLazyName DefaultOverrideLabel = "Override";
+	}
+} // namespace PCG::DataOverride
+
+/** How the overridden data is tied to its storage. Used for mapping. */
+UENUM()
+enum class EPCGDataOverrideKeyPolicy : uint8
+{
+	None = 0,
+	Spatial,
+	ByIndex,
+	ByAttribute
+};
+
+UENUM()
+enum class EPCGSpatialKeyTarget : uint8
+{
+	Position,
+	Transform,
+	AABB
+};
+
+/** Phase of element execution to apply overrides. Usually in PostExecute, but artifact generating elements may require the PrepareData phase. */
+UENUM(Meta = (Bitflags, UseEnumValuesAsMaskValuesInEditor = "true"))
+enum class EPCGDataOverridePhase : uint8
+{
+	None        = 0,
+	PrepareData = 1 << 0,
+	PostExecute = 1 << 1,
+};
+ENUM_CLASS_FLAGS(EPCGDataOverridePhase);
+
+/** If there are multiple collisions. Note: NoOp may have added cost, requiring the key matching to go through the entire set of candidates. */
+UENUM()
+enum class EPCGKeyCollisionResolution : uint8
+{
+	NoOp,      // Do not apply the override
+	TakeFirst, // Apply the first override and ignore any further collisions
+	Fail       // Fail and cancel the operation
+};
+
+/** In the event of a collision, this is how the user will be notified. */
+UENUM()
+enum class EPCGKeyCollisionResponse : uint8
+{
+	DoNothing = 0,
+	Log,
+	Warning,
+	Error
+};
+
+/** Configuration for how the override should be applied. */
+USTRUCT()
+struct FPCGDeltaSettings
+{
+	GENERATED_BODY()
+
+	UPROPERTY(EditAnywhere, Category = "Key Composition")
+	EPCGDataOverrideKeyPolicy KeyPolicy = EPCGDataOverrideKeyPolicy::None;
+
+	UPROPERTY(EditAnywhere, Category = "Key Composition", meta = (EditCondition = "KeyPolicy == EPCGDataOverrideKeyPolicy::ByAttribute", EditConditionHides))
+	FName AttributeName = NAME_None;
+
+	UPROPERTY(EditAnywhere, Category = "Key Composition")
+	FName PinLabel = NAME_None;
+
+	UPROPERTY(EditAnywhere, Category = "Key Composition", meta = (EditCondition = "KeyPolicy == EPCGDataOverrideKeyPolicy::Spatial", EditConditionHides))
+	EPCGSpatialKeyTarget KeyTarget = EPCGSpatialKeyTarget::Position;
+
+	UPROPERTY(EditAnywhere, Category = "Key Composition", meta = (EditCondition = "KeyPolicy == EPCGDataOverrideKeyPolicy::Spatial", EditConditionHides, ClampMin = "1.0"))
+	double SpatialTolerance = 10.0;
+};
+
+/** Wrapper helper around an element's signature. Used for storing deltas. */
+USTRUCT()
+struct FPCGDeltaKey
+{
+	GENERATED_BODY()
+
+	FPCGDeltaKey() = default;
+	PCG_API explicit FPCGDeltaKey(const FXxHash64 InHash, FName InDeltaLabel = NAME_None);
+
+	bool operator==(const FPCGDeltaKey& Other) const { return Hash == Other.Hash; }
+	bool operator!=(const FPCGDeltaKey& Other) const { return Hash != Other.Hash; }
+	friend uint32 GetTypeHash(const FPCGDeltaKey& Key) { return GetTypeHash(Key.Hash); }
+
+	UPROPERTY()
+	uint64 Hash = 0u;
+
+#if WITH_EDITORONLY_DATA
+	UPROPERTY()
+	FName DeltaLabel;
+#endif // WITH_EDITORONLY_DATA
+};
+
+/** Base class for deltas to inherit. Contains debug metadata or other useful information. */
+USTRUCT()
+struct FPCGDeltaBase
+{
+	GENERATED_BODY()
+
+	virtual ~FPCGDeltaBase() = default;
+	PCG_API virtual FName GetDeltaName() const;
+	PCG_API virtual PCGIndexing::FPCGIndexCollection FilterCandidates(const UPCGData* InData) const;
+	PCG_API virtual int32 Resolve(const UPCGData* InData, const PCGIndexing::FPCGIndexCollection& FilteredCandidates, const FPCGDeltaSettings& DeltaSettings) const;
+	PCG_API virtual bool Apply(UPCGData* InData, int32 ResolvedIndex) const;
+
+#if WITH_EDITORONLY_DATA
+	// @todo_pcg: The per-point spatial key used as the key for this specific delta. Keep for debugging, potentially permanently.
+	UPROPERTY()
+	FPCGDeltaKey ComputedKey;
+#endif // WITH_EDITORONLY_DATA
+};
+
+/** A collection of deltas, usually tied directly to a node element or more specifically, to a pin output. */
+USTRUCT()
+struct FPCGDeltaCollection
+{
+	GENERATED_BODY()
+
+	/** Find the mutable Delta with the provided Key. */
+	PCG_API TInstancedStruct<FPCGDeltaBase>* Find(const FPCGDeltaKey& InKey);
+	/** Find the Delta with the provided Key. */
+	PCG_API const TInstancedStruct<FPCGDeltaBase>* Find(const FPCGDeltaKey& InKey) const;
+	/** Returns true if the collection contains the provided Key. */
+	PCG_API bool Contains(const FPCGDeltaKey& InKey) const;
+
+	// For now, this will be editor only for storage mutation.
+#if WITH_EDITOR
+	/** Add a new delta entry to the collection, as a TInstancedStruct<FPCGDeltaBase>. */
+	PCG_API TInstancedStruct<FPCGDeltaBase>& Add(const FPCGDeltaKey& InKey, TInstancedStruct<FPCGDeltaBase>&& InDelta);
+
+	/** Remove a delta entry by key. */
+	PCG_API bool Remove(const FPCGDeltaKey& InKey);
+
+	/** Empty the collection. */
+	PCG_API void Empty();
+#endif // WITH_EDITOR
+
+	/** The number of Deltas in the collection. */
+	PCG_API int32 Num() const;
+	PCG_API bool IsEmpty() const;
+
+	PCG_API void ForEachDelta(TFunctionRef<void(const FPCGDeltaKey&, TInstancedStruct<FPCGDeltaBase>&)> Func);
+	PCG_API void ForEachDelta(TFunctionRef<void(const FPCGDeltaKey&, const TInstancedStruct<FPCGDeltaBase>&)> Func) const;
+	// @todo_pcg: ForEachDeltaAs
+
+	UPROPERTY()
+	FPCGDeltaSettings Settings;
+
+private:
+	UPROPERTY()
+	TMap<FPCGDeltaKey, TInstancedStruct<FPCGDeltaBase>> Deltas;
+};
+
diff --git a/Engine/Plugins/PCG/Source/PCG/Public/Graph/DataOverride/PCGDataOverrideHelpers.h b/Engine/Plugins/PCG/Source/PCG/Public/Graph/DataOverride/PCGDataOverrideHelpers.h
new file mode 100644
--- /dev/null
+++ b/Engine/Plugins/PCG/Source/PCG/Public/Graph/DataOverride/PCGDataOverrideHelpers.h
@@ -0,0 +1,95 @@
+// Copyright Epic Games, Inc. All Rights Reserved.
+
+#pragma once
+
+#include "PCGCommon.h"
+#include "PCGDataOverride.h"
+#include "Metadata/PCGAttributePropertySelector.h"
+
+struct FXxHash64;
+struct FPCGTaggedData;
+
+namespace PCG::DataOverride
+{
+	namespace Constants
+	{
+		constexpr double SpatialToleranceDefault = 10.0;
+	}
+
+	namespace Helpers
+	{
+		/** Compute the node-level override key from the execution stack. */
+		PCG_API FXxHash64 ComputeDataHash(const FPCGStack& InStack);
+		/** Compute a pin-level override key by combining the stack and pin label. */
+		PCG_API FXxHash64 ComputeDataHashWithPin(const FPCGStack& InStack, FName PinLabel);
+		/** Compute a pin-level override key by appending a pin label to a precomputed node key. */
+		PCG_API FXxHash64 ComputeDataHashWithPin(FXxHash64 NodeKey, FName PinLabel);
+	} // namespace Helpers
+
+	namespace Keys
+	{
+		PCG_API FXxHash64 ComputePositionHash(const FVector& InLocation, double Tolerance = Constants::SpatialToleranceDefault);
+		PCG_API FXxHash64 ComputeTransformHash(const FTransform& InTransform, double Tolerance = Constants::SpatialToleranceDefault);
+		PCG_API FXxHash64 ComputeAABBHash(const FBox& InBounds, double Tolerance = Constants::SpatialToleranceDefault);
+		// @todo_pcg: Could do a radial sphere bounds key as well.
+
+		/** Create a key directly with a custom label. */
+		struct FPCGLabelDeltaKey
+		{
+			explicit FPCGLabelDeltaKey(const FName InLabel) : Label(InLabel) {}
+
+			operator FPCGDeltaKey() const { return FPCGDeltaKey(FXxHash64{}, Label); }
+
+		private:
+			FName Label;
+		};
+
+		struct FPCGLocationDeltaKey
+		{
+			explicit FPCGLocationDeltaKey(const FVector& InLocation, const double InTolerance = Constants::SpatialToleranceDefault, const FName DeltaLabel = NAME_None)
+				: Key(ComputePositionHash(InLocation, InTolerance), DeltaLabel) {}
+
+			operator FPCGDeltaKey() const { return Key; }
+
+		private:
+			FPCGDeltaKey Key;
+		};
+
+		struct FPCGAABBDeltaKey
+		{
+			explicit FPCGAABBDeltaKey(const FBox& InBounds, const double InTolerance = Constants::SpatialToleranceDefault, const FName DeltaLabel = NAME_None)
+				: Key(ComputeAABBHash(InBounds, InTolerance), DeltaLabel) {}
+
+			operator FPCGDeltaKey() const { return Key; }
+
+		private:
+			FPCGDeltaKey Key;
+		};
+
+		struct FPCGTransformDeltaKey
+		{
+			explicit FPCGTransformDeltaKey(const FTransform& InTransform, const double InTolerance = Constants::SpatialToleranceDefault, const FName DeltaLabel = NAME_None)
+				: Key(ComputeTransformHash(InTransform, InTolerance), DeltaLabel) {}
+
+			operator FPCGDeltaKey() const { return Key; }
+
+		private:
+			FPCGDeltaKey Key;
+		};
+
+		// @todo_pcg: still to implement in future pass
+		PCG_API FPCGDeltaKey ComputeAttributeKey(const FPCGAttributePropertyInputSelector& InSelector);
+	} // namespace Keys
+
+	/** Returns true if the node has node-level data overrides that apply to all outputs. */
+	bool HasNodeLevelDataOverrides(FPCGContext* Context);
+
+	/** Returns true if the specific tagged data entry has pin-level data overrides. */
+	bool HasPinLevelDataOverrides(FPCGContext* Context, const FPCGTaggedData& InTaggedData);
+
+	/** Returns true if node-level or pin-level data overrides exist. */
+	bool HasAnyDataOverrides(FPCGContext* Context, const TArrayView<const FPCGTaggedData> InTaggedData);
+
+	/** Applies override deltas from the container to the tagged data. Data must be owned/duplicated before calling. */
+	void ApplyDataOverrides(FPCGContext* Context, const TArrayView<FPCGTaggedData> OutTaggedData);
+} // namespace PCG::DataOverride
diff --git a/Engine/Plugins/PCG/Source/PCG/Public/Graph/DataOverride/PCGDataOverridePoints.h b/Engine/Plugins/PCG/Source/PCG/Public/Graph/DataOverride/PCGDataOverridePoints.h
new file mode 100644
--- /dev/null
+++ b/Engine/Plugins/PCG/Source/PCG/Public/Graph/DataOverride/PCGDataOverridePoints.h
@@ -0,0 +1,102 @@
+// Copyright Epic Games, Inc. All Rights Reserved.
+
+#pragma once
+
+#include "PCGDataOverride.h"
+#include "PCGPoint.h"
+
+#include "PCGDataOverridePoints.generated.h"
+
+namespace PCG::DataOverride::Points::Constants
+{
+	inline constexpr FLazyName PointEditDeltaName = "PointEdit";
+	inline constexpr FLazyName PointDeleteDeltaName = "PointDelete";
+	inline constexpr FLazyName PointInsertDeltaName = "PointInsert";
+}
+
+/** Base delta for point data overrides. Uses an octree bounds query to filter candidate points spatially. */
+USTRUCT()
+struct PCG_API FPCGPointDeltaBase : public FPCGDeltaBase
+{
+	GENERATED_BODY()
+
+	/** Queries the point octree with Bounds to find candidate points. */
+	virtual PCGIndexing::FPCGIndexCollection FilterCandidates(const UPCGData* InData) const override;
+
+	/** Computes the spatial key for each candidate and compares against this delta's key to find the matching point. */
+	virtual int32 Resolve(const UPCGData* InData, const PCGIndexing::FPCGIndexCollection& FilteredCandidates, const FPCGDeltaSettings& DeltaSettings) const override;
+
+	/** The original transform of the point. For use with matching the key during override. */
+	UPROPERTY()
+	FTransform OriginalTransform;
+
+	/** World-space bounding box used for octree filtering and key computation. */
+	UPROPERTY()
+	FBox Bounds = FBox(ForceInit);
+};
+
+/** Delta that overrides the transform of a specific point. */
+USTRUCT()
+struct PCG_API FPCGPointEditDelta : public FPCGPointDeltaBase
+{
+	GENERATED_BODY()
+
+	static FName GetDeltaNameStatic() { return PCG::DataOverride::Points::Constants::PointEditDeltaName; }
+	virtual FName GetDeltaName() const override;
+	virtual bool Apply(UPCGData* InData, int32 ResolvedIndex) const override;
+
+	/** The desired transform to apply as the override. */
+	UPROPERTY()
+	FTransform TransformOverride;
+
+	/** Whether to override the point's position. */
+	UPROPERTY()
+	bool bOverridePosition = true;
+
+	/** Whether to override the point's rotation. */
+	UPROPERTY()
+	bool bOverrideRotation = true;
+
+	/** Whether to override the point's scale. */
+	UPROPERTY()
+	bool bOverrideScale = true;
+
+#if WITH_EDITORONLY_DATA
+	/** The original element index. For debugging purposes only. */
+	UPROPERTY()
+	int32 ElementIndex = INDEX_NONE;
+#endif // WITH_EDITORONLY_DATA
+};
+
+/** Delta that marks a specific point for deletion. */
+USTRUCT()
+struct PCG_API FPCGPointDeletionDelta : public FPCGPointDeltaBase
+{
+	GENERATED_BODY()
+
+	static FName GetDeltaNameStatic() { return PCG::DataOverride::Points::Constants::PointDeleteDeltaName; }
+	virtual FName GetDeltaName() const override;
+	virtual bool Apply(UPCGData* InData, int32 ResolvedIndex) const override;
+
+#if WITH_EDITORONLY_DATA
+	/** The original element index. For debugging purposes only. */
+	UPROPERTY()
+	int32 ElementIndex = INDEX_NONE;
+#endif // WITH_EDITORONLY_DATA
+};
+
+/** Delta that injects a new point. */
+USTRUCT()
+struct PCG_API FPCGPointInsertionDelta : public FPCGDeltaBase
+{
+	GENERATED_BODY()
+
+	static FName GetDeltaNameStatic() { return PCG::DataOverride::Points::Constants::PointInsertDeltaName; }
+	virtual FName GetDeltaName() const override;
+	virtual PCGIndexing::FPCGIndexCollection FilterCandidates(const UPCGData* InData) const override;
+	virtual bool Apply(UPCGData* InData, int32 ResolvedIndex) const override;
+
+	/** Points to be appended to the point data. */
+	UPROPERTY()
+	TArray<FPCGPoint> InjectedPoints;
+};
diff --git a/Engine/Plugins/PCG/Source/PCG/Public/PCGElement.h b/Engine/Plugins/PCG/Source/PCG/Public/PCGElement.h
--- a/Engine/Plugins/PCG/Source/PCG/Public/PCGElement.h
+++ b/Engine/Plugins/PCG/Source/PCG/Public/PCGElement.h
@@ -3,6 +3,7 @@
 #pragma once
 
 #include "PCGData.h"
+#include "Graph/DataOverride/PCGDataOverride.h"
 
 #include "Logging/LogVerbosity.h"
 #include "Templates/SharedPointer.h"
@@ -10,17 +11,19 @@
 
 #define UE_API PCG_API
 
+class IPCGElement;
+class IPCGGraphCache;
 class IPCGGraphExecutionSource;
+class UPCGComponent;
+class UPCGNode;
+class UPCGSettings;
 class UPCGSettingsInterface;
+
 struct FPCGContext;
 struct FPCGCrc;
 struct FPCGGetFromCacheParams;
 
-class IPCGElement;
-class IPCGGraphCache;
-class UPCGComponent;
-class UPCGSettings;
-class UPCGNode;
+enum class EPCGDataOverridePhase : uint8;
 
 typedef TSharedPtr<IPCGElement, ESPMode::ThreadSafe> FPCGElementPtr;
 
@@ -245,6 +248,9 @@ protected:
 	virtual bool ShouldLog() const { return true; }
 #endif
 
+	/** Which phase to apply the data overrides to the internal data. */
+	UE_API virtual EPCGDataOverridePhase GetDataOverridePhase() const { return EPCGDataOverridePhase::PostExecute; }
+
 private:
 	UE_API void CleanupAndValidateOutput(FPCGContext* Context) const;
 
