[PCG] Add thumbnail support to PCG Nodes (for texture related elements)

#Synced-CL 51160565

#Synced-CL 51160565

diff --git a/Engine/Plugins/PCG/Source/PCG/Private/Elements/PCGTextureSampler.cpp b/Engine/Plugins/PCG/Source/PCG/Private/Elements/PCGTextureSampler.cpp
--- a/Engine/Plugins/PCG/Source/PCG/Private/Elements/PCGTextureSampler.cpp
+++ b/Engine/Plugins/PCG/Source/PCG/Private/Elements/PCGTextureSampler.cpp
@@ -38,6 +38,18 @@ void UPCGTextureSamplerSettings::GetStaticTrackedKeys(FPCGSelectionKeyToSettings
 	OutKeysToSettings.FindOrAdd(Key).Emplace(this, /*bCulling=*/false);
 }
 
+EPCGChangeType UPCGTextureSamplerSettings::GetChangeTypeForProperty(const FName& InPropertyName) const
+{
+	EPCGChangeType ChangeType = Super::GetChangeTypeForProperty(InPropertyName);
+
+	if (InPropertyName == GET_MEMBER_NAME_CHECKED(UPCGTextureSamplerSettings, Texture))
+	{
+		ChangeType |= EPCGChangeType::Cosmetic;
+	}
+
+	return ChangeType;
+}
+
 void UPCGTextureSamplerSettings::PostEditChangeProperty(FPropertyChangedEvent& PropertyChangedEvent) 
 {
 	if (PropertyChangedEvent.Property)
@@ -77,6 +89,16 @@ FText UPCGTextureSamplerSettings::GetNodeTooltipText() const
 		"If the texture is CPU-accessible, the sampler will prefer the CPU version of the texture.\n"
 		"Otherwise, the texture will be read back from the GPU if one is present.");
 }
+
+TOptional<FPCGNodeThumbnailProxy> UPCGTextureSamplerSettings::GetNodeThumbnail() const
+{
+	if (Texture.IsNull())
+	{
+		return {};
+	}
+
+	return FPCGNodeThumbnailProxy::FromAssetPath(Texture.ToSoftObjectPath());
+}
 #endif // WITH_EDITOR
 
 TArray<FPCGPinProperties> UPCGTextureSamplerSettings::OutputPinProperties() const
diff --git a/Engine/Plugins/PCG/Source/PCG/Private/PCGSettings.cpp b/Engine/Plugins/PCG/Source/PCG/Private/PCGSettings.cpp
--- a/Engine/Plugins/PCG/Source/PCG/Private/PCGSettings.cpp
+++ b/Engine/Plugins/PCG/Source/PCG/Private/PCGSettings.cpp
@@ -29,6 +29,7 @@
 
 #if WITH_EDITOR
 #include "Editor.h"
+#include "Engine/Texture.h"
 #include "Hash/PCGSettingsHashContext.h"
 #endif
 
@@ -92,6 +93,18 @@ private:
 	const UPCGSettings* PCGSettings = nullptr;
 };
 
+#if WITH_EDITOR
+FPCGNodeThumbnailProxy FPCGNodeThumbnailProxy::FromAssetPath(const FSoftObjectPath& InAssetPath)
+{
+	return FPCGNodeThumbnailProxy{FTextureType{TInPlaceType<FSoftObjectPath>(), InAssetPath}};
+}
+
+FPCGNodeThumbnailProxy FPCGNodeThumbnailProxy::FromTexture(const UTexture* InTexture)
+{
+	return FPCGNodeThumbnailProxy{FTextureType{TInPlaceType<const UTexture*>(), InTexture}};
+}
+#endif // WITH_EDITOR
+
 FString FPCGSettingsOverridableParam::GetPropertyPath() const
 {
 	return FString::JoinBy(PropertiesNames, PCGSettings::PropertyPathSeparator, [](const FName InName) { return InName.ToString(); });
diff --git a/Engine/Plugins/PCG/Source/PCG/Public/Elements/PCGTextureSampler.h b/Engine/Plugins/PCG/Source/PCG/Public/Elements/PCGTextureSampler.h
--- a/Engine/Plugins/PCG/Source/PCG/Public/Elements/PCGTextureSampler.h
+++ b/Engine/Plugins/PCG/Source/PCG/Public/Elements/PCGTextureSampler.h
@@ -29,9 +29,11 @@ public:
 	virtual FName GetDefaultNodeName() const override { return FName(TEXT("GetTextureData")); }
 	virtual FText GetDefaultNodeTitle() const override { return NSLOCTEXT("PCGTextureSamplerSettings", "NodeTitle", "Get Texture Data"); }
 	virtual FText GetNodeTooltipText() const override;
+	virtual TOptional<FPCGNodeThumbnailProxy> GetNodeThumbnail() const override;
 	virtual EPCGSettingsType GetType() const override { return EPCGSettingsType::Spatial; }
 	virtual void GetStaticTrackedKeys(FPCGSelectionKeyToSettingsMap& OutKeysToSettings, TArray<TObjectPtr<const UPCGGraph>>& OutVisitedGraphs) const override;
 	virtual bool CanDynamicallyTrackKeys() const override { return true; }
+	virtual EPCGChangeType GetChangeTypeForProperty(const FName& InPropertyName) const override;
 #endif
 
 protected:
diff --git a/Engine/Plugins/PCG/Source/PCG/Public/PCGSettings.h b/Engine/Plugins/PCG/Source/PCG/Public/PCGSettings.h
--- a/Engine/Plugins/PCG/Source/PCG/Public/PCGSettings.h
+++ b/Engine/Plugins/PCG/Source/PCG/Public/PCGSettings.h
@@ -12,6 +12,7 @@
 #include "Tests/Determinism/PCGDeterminismSettings.h"
 
 #include "Algo/AllOf.h"
+#include "Misc/TVariant.h"
 
 #include "PCGSettings.generated.h"
 
@@ -24,6 +25,7 @@ class UPCGGraph;
 class UPCGNode;
 class UPCGSettings;
 class UPropertyBag;
+class UTexture;
 struct FPCGGPUCompilationContext;
 struct FPCGKernelEdge;
 struct FPCGPinProperties;
@@ -78,6 +80,16 @@ enum class EPCGSettingsType : uint8
 };
 
 #if WITH_EDITOR
+struct FPCGNodeThumbnailProxy
+{
+	using FTextureType = TVariant<const UTexture*, FSoftObjectPath>;
+
+	PCG_API static FPCGNodeThumbnailProxy FromAssetPath(const FSoftObjectPath& InAssetPath);
+	PCG_API static FPCGNodeThumbnailProxy FromTexture(const UTexture* InTexture);
+
+	FTextureType TextureSource;
+};
+
 	DECLARE_MULTICAST_DELEGATE_TwoParams(FOnPCGSettingsChanged, UPCGSettings*, EPCGChangeType);
 #endif
 
@@ -325,6 +337,7 @@ public:
 
 	virtual FText GetNodeTooltipText() const { return FText::GetEmpty(); }
 	virtual FLinearColor GetNodeTitleColor() const { return FLinearColor::White; }
+	virtual TOptional<FPCGNodeThumbnailProxy> GetNodeThumbnail() const { return {}; }
 	virtual EPCGSettingsType GetType() const { return EPCGSettingsType::Generic; }
 
 	/** Can override the label style for a pin. Return false if no override is available. */
diff --git a/Engine/Plugins/PCG/Source/PCGEditor/Private/Nodes/SPCGEditorGraphNode.cpp b/Engine/Plugins/PCG/Source/PCGEditor/Private/Nodes/SPCGEditorGraphNode.cpp
--- a/Engine/Plugins/PCG/Source/PCGEditor/Private/Nodes/SPCGEditorGraphNode.cpp
+++ b/Engine/Plugins/PCG/Source/PCGEditor/Private/Nodes/SPCGEditorGraphNode.cpp
@@ -18,12 +18,15 @@
 #include "Pins/SPCGEditorGraphPinVectorSlider.h"
 #include "Schema/PCGEditorGraphSchema.h"
 
+#include "AssetThumbnail.h"
 #include "GraphEditorSettings.h"
 #include "IDocumentation.h"
 #include "Nodes/PCGEditorGraphNode.h"
 #include "SCommentBubble.h"
 #include "SGraphPin.h"
 #include "SPinTypeSelector.h"
+#include "AssetRegistry/IAssetRegistry.h"
+#include "ThumbnailRendering/ThumbnailManager.h"
 
 #include "Widgets/SBoxPanel.h"
 #include "Widgets/Input/SButton.h"
@@ -199,6 +202,53 @@ FReply SPCGEditorGraphNode::OnAddPin()
 	return FReply::Handled();
 }
 
+void SPCGEditorGraphNode::CreateBelowPinControls(const TSharedPtr<SVerticalBox> MainBox)
+{
+	const UPCGSettings* NodeSettings = PCGEditorGraphNode->GetSettings();
+	TSharedPtr<FAssetThumbnailPool> ThumbnailPool = UThumbnailManager::Get().GetSharedThumbnailPool();
+	if (!NodeSettings || !ThumbnailPool.IsValid())
+	{
+		return;
+	}
+
+	if (TOptional<FPCGNodeThumbnailProxy> ThumbnailProxy = NodeSettings->GetNodeThumbnail(); ThumbnailProxy.IsSet())
+	{
+		TSharedPtr<FAssetThumbnail> NodeThumbnail;
+
+		if (ThumbnailProxy->TextureSource.IsType<FSoftObjectPath>())
+		{
+			const FSoftObjectPath& Path = ThumbnailProxy->TextureSource.Get<FSoftObjectPath>();
+			if (const FAssetData Asset = IAssetRegistry::Get()->GetAssetByObjectPath(Path); Asset.IsValid())
+			{
+				NodeThumbnail = MakeShared<FAssetThumbnail>(Asset, 64, 64, ThumbnailPool);
+			}
+		}
+		else if (ensure(ThumbnailProxy->TextureSource.IsType<const UTexture*>()))
+		{
+			// @todo_pcg: handle UTexture here in the future
+		}
+
+		if (NodeThumbnail.IsValid())
+		{
+			FAssetThumbnailConfig ThumbnailConfig;
+			ThumbnailConfig.ThumbnailLabel = EThumbnailLabel::AssetName;
+
+			MainBox->AddSlot()
+			   .AutoHeight()
+			   .HAlign(HAlign_Center)
+			   .Padding(2)
+				[
+					SNew(SBox)
+				   .MaxDesiredWidth(64)
+				   .MaxDesiredHeight(64)
+					[
+						NodeThumbnail->MakeThumbnailWidget(ThumbnailConfig)
+					]
+				];
+		}
+	}
+}
+
 void SPCGEditorGraphNode::AddPin(const TSharedRef<SGraphPin>& PinToAdd)
 {
 	check(PCGEditorGraphNode);
diff --git a/Engine/Plugins/PCG/Source/PCGEditor/Private/Nodes/SPCGEditorGraphNode.h b/Engine/Plugins/PCG/Source/PCGEditor/Private/Nodes/SPCGEditorGraphNode.h
--- a/Engine/Plugins/PCG/Source/PCGEditor/Private/Nodes/SPCGEditorGraphNode.h
+++ b/Engine/Plugins/PCG/Source/PCGEditor/Private/Nodes/SPCGEditorGraphNode.h
@@ -28,6 +28,7 @@ public:
 	virtual TSharedPtr<SGraphPin> CreatePinWidget(UEdGraphPin* Pin) const override;
 	virtual EVisibility IsAddPinButtonVisible() const override;
 	virtual FReply OnAddPin() override;
+	virtual void CreateBelowPinControls(TSharedPtr<SVerticalBox> MainBox) override;
 	//~ End SGraphNode Interface
 
 	//~ Begin SNodePanel::SNode Interface
