[PCG] Intersection improvements
- Better segment subdivision
- Preserve segment types

#Synced-CL 50984763

diff --git a/Engine/Plugins/PCG/Source/PCG/Private/Elements/PCGSplineIntersection.cpp b/Engine/Plugins/PCG/Source/PCG/Private/Elements/PCGSplineIntersection.cpp
--- a/Engine/Plugins/PCG/Source/PCG/Private/Elements/PCGSplineIntersection.cpp
+++ b/Engine/Plugins/PCG/Source/PCG/Private/Elements/PCGSplineIntersection.cpp
@@ -9,6 +9,8 @@
 #include "Helpers/PCGHelpers.h"
 #include "Metadata/Accessors/PCGAttributeAccessorHelpers.h"
 
+#include "PCA3.h"
+
 #include UE_INLINE_GENERATED_CPP_BY_NAME(PCGSplineIntersection)
 
 #define LOCTEXT_NAMESPACE "PCGSplineIntersectionElement"
@@ -19,9 +21,29 @@ namespace PCGSplineIntersection
 	{
 		const FName OriginatingSplineIndexAttributeName = TEXT("OriginatingSplineIndex");
 		const FName IntersectingSplineIndexAttributeName = TEXT("IntersectingSplineIndex");
+		const FName IntersectionTypeAttributeName = TEXT("IntersectionType");
+		constexpr double KeyDifferenceThreshold = UE_DOUBLE_KINDA_SMALL_NUMBER;
 	}
 
-	struct InputData
+	namespace Helpers
+	{
+		double GetSplineKeyLength(const UPCGSplineData* InSpline)
+		{
+			check(InSpline);
+			if (InSpline->SplineStruct.GetSplinePointsPosition().Points.IsEmpty())
+			{
+				return 0;
+			}
+			else 
+			{
+				const double LastKey = InSpline->SplineStruct.GetSplinePointsPosition().Points.Last().InVal;
+				return LastKey + (InSpline->IsClosed() ? InSpline->GetLoopKeyOffset() : 0.0);
+			}
+		}
+	}
+
+	// Data tracking structure to encapsulate discretization of the spline
+	struct FInputData
 	{
 		void InitializeBounds(const FVector& HalfRadiusTolerance)
 		{
@@ -38,30 +60,16 @@ namespace PCGSplineIntersection
 			}
 		}
 
-		const TArray<FVector>& GetRoughSamples(int SegmentIndex, double MinDistanceThreshold, double& OutRoughSampleStep) const
+		const TArray<FVector>& GetSamples(int SegmentIndex, double DiscretizationError) const
 		{
 			check(Spline && SegmentIndex >= 0 && SegmentIndex < SegmentsNum);
-			TArray<FVector>& Samples = RoughSampling.FindOrAdd(SegmentIndex);
-
-			const FVector::FReal SegmentLength = Spline->GetSegmentLength(SegmentIndex);
-			FVector::FReal RoughSampleSize = FMath::Max(SegmentLength / 100.0, MinDistanceThreshold);
+			TArray<FVector>& Samples = Sampling.FindOrAdd(SegmentIndex);
 
 			if (Samples.IsEmpty())
 			{
-				// Rough sampling should be in the 10-100 range.
-				int NumSamples = FMath::Floor(SegmentLength / RoughSampleSize) + 1;
-
-				Samples.SetNumUninitialized(NumSamples+1);
-
-				const FVector::FReal DistanceAtSegmentStart = Spline->SplineStruct.GetDistanceAlongSplineAtSplinePoint(SegmentIndex);
-				for (int Sample = 0; Sample <= NumSamples; ++Sample)
-				{
-					Samples[Sample] = Spline->SplineStruct.GetLocationAtDistanceAlongSpline(DistanceAtSegmentStart + SegmentLength * (Sample / (double)NumSamples), ESplineCoordinateSpace::World);
-				}
+				Spline->SplineStruct.ConvertSplineSegmentToPolyLine(SegmentIndex, ESplineCoordinateSpace::World, DiscretizationError, Samples);
 			}
 
-			// Note: to address scaling, we'll take the distance between two points
-			OutRoughSampleStep = FMath::Max(RoughSampleSize, (Samples[1] - Samples[0]).Length());
 			return Samples;
 		}
 
@@ -69,70 +77,407 @@ namespace PCGSplineIntersection
 		int InputIndex = INDEX_NONE;
 		FBox Bounds = FBox(EForceInit::ForceInit);
 		TArray<FBox> SegmentBounds;
-		mutable TMap<int, TArray<FVector>> RoughSampling;
+		mutable TMap<int, TArray<FVector>> Sampling;
 		int SegmentsNum;
 	};
 
-	struct IntersectionPoint
+	// Represents an intersection on a spline, with a start & stop.
+	struct FIntersectionRange
+	{
+		const UPCGSplineData* Spline = nullptr;
+		const FInputData* Input = nullptr;
+		double InputKeyStart = 0.0;
+		double InputKeyEnd = 0.0;
+
+		TArray<const UPCGSplineData*, TInlineAllocator<1>> OtherSplines;
+		TArray<FVector, TInlineAllocator<2>> Points;
+
+		bool Overlaps(const FIntersectionRange& Other) const
+		{
+			if (Spline != Other.Spline)
+			{
+				return false;
+			}
+
+			// Special treatment for looping ranges
+			if (Spline->IsClosed() && !Spline->SplineStruct.GetSplinePointsPosition().Points.IsEmpty())
+			{
+				const float SplineKeyLength = Helpers::GetSplineKeyLength(Spline);
+
+				// Assumption here that the key end will be at most the spline key length, otherwise we've made an error
+				ensure(InputKeyEnd <= SplineKeyLength && Other.InputKeyStart <= SplineKeyLength);
+
+				if (InputKeyStart + SplineKeyLength <= Other.InputKeyEnd + Constants::KeyDifferenceThreshold ||
+					Other.InputKeyStart + SplineKeyLength <= InputKeyEnd + Constants::KeyDifferenceThreshold)
+				{
+					return true;
+				}
+			}
+
+			// Normal overlapping case
+			return Spline == Other.Spline && !(InputKeyStart > Other.InputKeyEnd + Constants::KeyDifferenceThreshold || Other.InputKeyStart > InputKeyEnd + Constants::KeyDifferenceThreshold);
+		}
+
+		void Merge(const FIntersectionRange& Other)
+		{
+			check(Spline == Other.Spline);
+
+			// Special loop management:
+			// In this case, we'll support cases where key start > key end to represent the loop part.
+			if (Spline->IsClosed())
+			{
+				const float SplineKeyLength = Helpers::GetSplineKeyLength(Spline);
+				const bool bInvertedA = (InputKeyEnd < InputKeyStart);
+				const bool bInvertedB = (Other.InputKeyEnd < Other.InputKeyStart);
+				const bool bShiftA = (InputKeyStart + SplineKeyLength <= Other.InputKeyEnd + Constants::KeyDifferenceThreshold);
+				const bool bShiftB = !bShiftA && Other.InputKeyStart + SplineKeyLength <= InputKeyEnd + Constants::KeyDifferenceThreshold;
+
+				double StartA = InputKeyStart + (bShiftA ? SplineKeyLength : 0);
+				double EndA = InputKeyEnd + (bShiftA ? SplineKeyLength : 0) + (bInvertedA ? SplineKeyLength : 0);
+				double StartB = Other.InputKeyStart + (bShiftB ? SplineKeyLength : 0);
+				double EndB = Other.InputKeyEnd + (bShiftB ? SplineKeyLength : 0) + (bInvertedB ? SplineKeyLength : 0);
+
+				InputKeyStart = FMath::Min(StartA, StartB);
+				InputKeyEnd = FMath::Max(EndA, EndB);
+
+				// Finally, special case for a complete loop
+				if (InputKeyEnd - InputKeyStart + Constants::KeyDifferenceThreshold >= SplineKeyLength)
+				{
+					InputKeyStart = 0;
+					InputKeyEnd = SplineKeyLength;
+				}
+				// Finally, correct end to keep in the real range of keys
+				else if (InputKeyEnd >= SplineKeyLength)
+				{
+					InputKeyEnd -= SplineKeyLength;
+				}
+			}
+			else
+			{
+				InputKeyStart = FMath::Min(InputKeyStart, Other.InputKeyStart);
+				InputKeyEnd = FMath::Max(InputKeyEnd, Other.InputKeyEnd);
+			}
+
+			for (const UPCGSplineData* OtherSpline : Other.OtherSplines)
+			{
+				OtherSplines.AddUnique(OtherSpline);
+			}
+
+			Points.Append(Other.Points);
+		}
+
+		FVector GetVector() const
+		{
+			check(Points.Num() >= 2);
+			FVector A = Points[0];
+			FVector B = Points[1];
+
+			// Find limits on a kind of triangle inequality.
+			for (int PointIndex = 2; PointIndex < Points.Num(); ++PointIndex)
+			{
+				// Basically, if the 3rd point is on the "same" side of both A & B, then it is a more extreme point.
+				// Note: this does somewhat assume linearity.
+				const FVector& C = Points[PointIndex];
+
+				const FVector CA = C - A;
+				const FVector CB = C - B;
+				if ((CA | CB) > 0)
+				{
+					if (CA.SquaredLength() < CB.SquaredLength())
+					{
+						A = C;
+					}
+					else
+					{
+						B = C;
+					}
+				}
+			}
+
+			return B - A;
+		}
+	};
+
+	// Intersection observed between two splines
+	struct FIntersectionPoint
 	{
-		const UPCGSplineData* A = nullptr;
-		const InputData* InputA = nullptr;
-		const UPCGSplineData* B = nullptr;
-		const InputData* InputB = nullptr;
-
-		FVector AveragePosition;
-		FVector PositionA;
-		FVector PositionB;
-		float InputKeyA = 0.0f;
-		float InputKeyB = 0.0f;
-
-		bool IsOnSpline(const UPCGSplineData* In) const { return In == A || In == B; }
-		float GetKey(const UPCGSplineData* In) const { check(In == A || In == B); return In == A ? InputKeyA : InputKeyB; }
-		const UPCGSplineData* GetOtherSpline(const UPCGSplineData* In) const { check(In == A || In == B); return In == A ? B : A; }
-		const InputData* GetOtherInput(const InputData* In) const { check(In == InputA || In == InputB); return In == InputA ? InputB : InputA; }
-		const FVector& GetPosition(const UPCGSplineData* In) const { check(In == A || In == B); return In == A ? PositionA : PositionB; }
-		FVector& GetPosition(const UPCGSplineData* In) { check(In == A || In == B); return In == A ? PositionA : PositionB; }
+		FIntersectionRange A;
+		FIntersectionRange B;
 	};
 
-	TArray<IntersectionPoint> CollapseColocatedIntersections(const TArray<IntersectionPoint>& InPoints, const double SqrDistanceThreshold, const UPCGSplineData* SplineFilter)
+	// Retrieves & merges the intersection ranges for a specific spline.
+	TArray<FIntersectionRange> FilterAndMergeIntersections(const TArray<FIntersectionPoint>& Intersections, const UPCGSplineData* Spline)
 	{
-		TArray<IntersectionPoint> OutPoints;
-		OutPoints.Reserve(InPoints.Num());
+		TArray<FIntersectionRange> Output;
 
-		TArray<int> PointsWeight;
-		PointsWeight.Reserve(InPoints.Num());
+		// Select intersection ranges only for the spline we're interested in
+		for (const FIntersectionPoint& Intersection : Intersections)
+		{
+			if (Intersection.A.Spline == Spline)
+			{
+				Output.Add(Intersection.A);
+			}
 
-		for (const IntersectionPoint& Point : InPoints)
+			if (Intersection.B.Spline == Spline)
+			{
+				Output.Add(Intersection.B);
+			}
+		}
+
+		// Merge ranges based on their key values
+		for (int I = 0; I < Output.Num(); ++I)
 		{
-			if (SplineFilter && !Point.IsOnSpline(SplineFilter))
+			FIntersectionRange& Current = Output[I];
+			for (int J = I + 1; J < Output.Num(); ++J)
 			{
-				continue;
+				FIntersectionRange& Other = Output[J];
+
+				if (Current.Overlaps(Other))
+				{
+					Current.Merge(Other);
+					Output.RemoveAt(J);
+					I--; // restart the loop
+					break;
+				}
 			}
+		}
 
-			const FVector& Position = SplineFilter ? Point.GetPosition(SplineFilter) : Point.AveragePosition;
+		return Output;
+	}
+
+	struct FCollapsedIntersectionPoint
+	{
+		TArray<FIntersectionRange> Ranges;
+		FTransform Transform = FTransform::Identity;
+		FBox Bounds = FBox(ForceInit);
 
-			int MatchingPointIndex = OutPoints.IndexOfByPredicate([&Position, &SqrDistanceThreshold](const IntersectionPoint& OutPt)
+		void ComputeSpatialData(bool bUsePCA)
+		{
+			// Perform PCA on all points to get bounds for the intersection to add.
+			TArray<FVector> Points;
+			for (const FIntersectionRange& Range : Ranges)
 			{
-				return (OutPt.AveragePosition - Position).SquaredLength() <= SqrDistanceThreshold;
-			});
+				Points.Append(Range.Points);
+			}
 
-			if (MatchingPointIndex == INDEX_NONE)
+			// note: might need to perturb points a bit to account for things happening on a plane
+			UE::Geometry::FPCA3d PCA;
+			if (bUsePCA && PCA.Compute(Points))
 			{
-				IntersectionPoint& NewPoint = OutPoints.Emplace_GetRef(Point);
-				NewPoint.AveragePosition = Position; // meaning of average position will change when we are filtering on a specific spline
-				PointsWeight.Add(1);
+				Transform = FTransform(PCA.Eigenvectors[0], PCA.Eigenvectors[1], PCA.Eigenvectors[2], PCA.Mean);
 			}
 			else
 			{
-				const FVector NewAveragePosition = ((PointsWeight[MatchingPointIndex] * OutPoints[MatchingPointIndex].AveragePosition) + Point.AveragePosition) / static_cast<double>(PointsWeight[MatchingPointIndex] + 1);
-				OutPoints[MatchingPointIndex].AveragePosition = NewAveragePosition;
-				++PointsWeight[MatchingPointIndex];
+				// Default & Degenerate case (collocated or collinear).
+				// Find out per-range the extreme points - this will be the primary axis (X).
+				// After that, use Z-up, and compute resulting Y axis.
+				FVector PrimaryAxis = FVector::ZeroVector;
+				for (const FIntersectionRange& Range : Ranges)
+				{
+					const FVector RangeVector = Range.GetVector();
+					if (RangeVector.SquaredLength() > PrimaryAxis.SquaredLength())
+					{
+						PrimaryAxis = RangeVector;
+					}
+				}
 
-				// @todo_pcg: Add additional information for more intersection info
+				if (PrimaryAxis.SquaredLength() > UE_DOUBLE_SMALL_NUMBER)
+				{
+					FVector MeanPosition = FVector::ZeroVector;
+
+					if (Points.Num() > 0)
+					{
+						for (const FVector& Point : Points)
+						{
+							MeanPosition += Point;
+						}
+
+						MeanPosition /= static_cast<double>(Points.Num());
+					}
+
+					// Normalize primary axis and compute a proper 'z' vector.
+					PrimaryAxis = PrimaryAxis.GetSafeNormal(UE_DOUBLE_SMALL_NUMBER, FVector::XAxisVector);
+
+					FVector TentativeZAxis = FVector::ZAxisVector;
+					if (FMath::IsNearlyEqual(FMath::Abs(PrimaryAxis | TentativeZAxis), 1.0, UE_KINDA_SMALL_NUMBER))
+					{
+						TentativeZAxis = FVector::YAxisVector;
+					}
+
+					// todo_pcg: z axis here is wrong. we need another way to get the axis
+					Transform = FTransform(FRotationMatrix::MakeFromXZ(PrimaryAxis, TentativeZAxis).ToQuat(), MeanPosition, FVector::OneVector);
+				}
+			}
+
+			// compute local box from points.
+			for (const FVector& Point : Points)
+			{
+				Bounds += Transform.InverseTransformPosition(Point);
 			}
+
+			// Finally, make sure that the bounds extents are non-zero.
+			constexpr double MinBoundsSize = 1.0;
+			const FVector OriginalExtent = Bounds.GetExtent();
+			FVector Expansion(FMath::Max(0.0, MinBoundsSize - OriginalExtent[0]),
+				FMath::Max(0.0, MinBoundsSize - OriginalExtent[1]),
+				FMath::Max(0.0, MinBoundsSize - OriginalExtent[2]));
+
+			Bounds = Bounds.ExpandBy(Expansion);
 		}
+	};
+
+	// Collapsed the intersections until all intersections are not mergeable (based on the standard intersection merge ratios).
+	TArray<FCollapsedIntersectionPoint> CollapseIntersections(const TArray<FIntersectionPoint>& Intersections, bool bUsePCA)
+	{
+		TArray<FCollapsedIntersectionPoint> Output;
+		TArray<bool> ConsumedIntersections;
+		ConsumedIntersections.SetNumZeroed(Intersections.Num());
+
+		for (int IntersectionIndex = 0; IntersectionIndex < Intersections.Num(); ++IntersectionIndex)
+		{
+			if (ConsumedIntersections[IntersectionIndex])
+			{
+				continue;
+			}
+
+			ConsumedIntersections[IntersectionIndex] = true;
+
+			const FIntersectionPoint& CurrentIntersection = Intersections[IntersectionIndex];
+
+			// @todo_pcg: this will not work for self-intersections so we might need to have a slightly different approach here.
+			// Break down intersection in its ranges
+			TMap<const UPCGSplineData*, FIntersectionRange> RangesPerSpline;
+			RangesPerSpline.Add(CurrentIntersection.A.Spline, CurrentIntersection.A);
+			RangesPerSpline.Add(CurrentIntersection.B.Spline, CurrentIntersection.B);
 
-		return OutPoints;
+			bool bDone = false;
+			while (!bDone)
+			{
+				bDone = true;
+
+				for (int OtherIntersectionIndex = IntersectionIndex + 1; OtherIntersectionIndex < Intersections.Num(); ++OtherIntersectionIndex)
+				{
+					if (ConsumedIntersections[OtherIntersectionIndex])
+					{
+						continue;
+					}
+
+					const FIntersectionPoint& OtherIntersection = Intersections[OtherIntersectionIndex];
+
+					FIntersectionRange* MatchingRangeA = RangesPerSpline.Find(OtherIntersection.A.Spline);
+					FIntersectionRange* MatchingRangeB = RangesPerSpline.Find(OtherIntersection.B.Spline);
+
+					// The condition here is a bit convoluted:
+					// - All ranges we've taken so far (for specific splines) must overlap
+					// *but* we're fine if we hadn't seen that spline before.
+					if ((!MatchingRangeA && !MatchingRangeB) ||
+						(MatchingRangeA && !MatchingRangeA->Overlaps(OtherIntersection.A)) ||
+						(MatchingRangeB && !MatchingRangeB->Overlaps(OtherIntersection.B)))
+					{
+						continue;
+					}
+
+					ConsumedIntersections[OtherIntersectionIndex] = true;
+
+					if (MatchingRangeA)
+					{
+						MatchingRangeA->Merge(OtherIntersection.A);
+					}
+					else
+					{
+						RangesPerSpline.Add(OtherIntersection.A.Spline, OtherIntersection.A);
+					}
+
+					if (MatchingRangeB)
+					{
+						MatchingRangeB->Merge(OtherIntersection.B);
+					}
+					else
+					{
+						RangesPerSpline.Add(OtherIntersection.B.Spline, OtherIntersection.B);
+					}
+
+					bDone = false;
+					break;
+				}
+			}
+
+			// Finally push to output
+			FCollapsedIntersectionPoint& Intersection = Output.Emplace_GetRef();
+			for (auto& Range : RangesPerSpline)
+			{
+				Intersection.Ranges.Add(MoveTemp(Range.Value));
+			}
+
+			Intersection.ComputeSpatialData(bUsePCA);
+		}
+
+		return Output;
+	}
+
+	// Considering two segments (AB and CD), a tolerance (SqrDistanceThreshold), compute:
+	// - if the segments overlap in any way (return value - true if overlaps)
+	// - at what ratio the overlaps start/end (ABStart - ABEnd; CDStart - CDEnd).
+	bool FindOverlapRange(const FVector& A, const FVector& B, const FVector& C, const FVector& D, const double SqrDistanceThreshold, double& ABStart, double& ABEnd, double& CDStart, double& CDEnd)
+	{
+		ABStart = 0.0;
+		ABEnd = 1.0;
+		CDStart = 0.0;
+		CDEnd = 1.0;
+
+		// Test segment to segment - if that isn't successful, there's nothing to do.
+		FVector X, Y;
+		FMath::SegmentDistToSegment(A, B, C, D, X, Y);
+
+		if ((X - Y).SquaredLength() > SqrDistanceThreshold)
+		{
+			return false;
+		}
+
+		const double ABHint = (X - A).Length() / (B - A).Length();
+		const double CDHint = (Y - C).Length() / (D - C).Length();
+
+		// TODO : we should early test out trivial cases
+		auto FindRange = [&SqrDistanceThreshold](const FVector& A, const FVector& B, const FVector& C, const FVector& D, const double& Hint, double& MinRange, double& MaxRange)
+		{
+			const FVector& Start = A;
+			const double Len = (B - A).Length();
+			const FVector Dir = (B - A) / Len;
+
+			auto FindEnd = [&SqrDistanceThreshold, &Start, &Len, &Dir, &C, &D, &Hint](bool bFindMin)
+			{
+				constexpr double RangeThreshold = UE_DOUBLE_KINDA_SMALL_NUMBER;
+				double Low = bFindMin ? 0.0 : Hint;
+				double High = bFindMin ? Hint : 1.0;
+
+				while (High - Low > RangeThreshold)
+				{
+					const double Mid = 0.5 * (Low + High);
+					const FVector Pt = Start + Dir * Mid * Len;
+
+					if ((FMath::PointDistToSegmentSquared(Pt, C, D) < SqrDistanceThreshold) == bFindMin)
+					{
+						High = Mid;
+					}
+					else
+					{
+						Low = Mid;
+					}
+				}
+
+				return 0.5 * (Low + High);
+			};
+
+			MinRange = FindEnd(/*bFindMin=*/true);
+			MaxRange = FindEnd(/*bFindMin=*/false);
+		};
+		
+		FindRange(A, B, C, D, ABHint, ABStart, ABEnd);
+		FindRange(C, D, A, B, CDHint, CDStart, CDEnd);
+
+		check(ABStart <= ABEnd);
+		check(CDStart <= CDEnd);
+		return true;
 	}
 }
 
@@ -140,6 +485,7 @@ UPCGSplineIntersectionSettings::UPCGSplineIntersectionSettings()
 {
 	OriginatingSplineIndexAttribute.SetAttributeName(PCGSplineIntersection::Constants::OriginatingSplineIndexAttributeName);
 	IntersectingSplineIndexAttribute.SetAttributeName(PCGSplineIntersection::Constants::IntersectingSplineIndexAttributeName);
+	IntersectionTypeAttribute.SetAttributeName(PCGSplineIntersection::Constants::IntersectionTypeAttributeName);
 }
 
 #if WITH_EDITOR
@@ -204,7 +550,7 @@ bool FPCGSplineIntersectionElement::ExecuteInternal(FPCGContext* Context) const
 
 	TArray<FPCGTaggedData> Inputs = Context->InputData.GetInputsByPin(PCGPinConstants::DefaultInputLabel);
 
-	TArray<PCGSplineIntersection::InputData> InputSplines;
+	TArray<PCGSplineIntersection::FInputData> InputSplines;
 	const FVector HalfDistanceRadius(0.5 * Settings->DistanceThreshold);
 	const double SqrDistanceThreshold = FMath::Square(Settings->DistanceThreshold);
 
@@ -213,7 +559,7 @@ bool FPCGSplineIntersectionElement::ExecuteInternal(FPCGContext* Context) const
 		const FPCGTaggedData& Input = Inputs[InputIndex];
 		if (const UPCGSplineData* SplineData = Cast<UPCGSplineData>(Input.Data))
 		{
-			PCGSplineIntersection::InputData& SplineInputData = InputSplines.Add_GetRef(PCGSplineIntersection::InputData{ .Spline = SplineData, .InputIndex = InputIndex });
+			PCGSplineIntersection::FInputData& SplineInputData = InputSplines.Add_GetRef(PCGSplineIntersection::FInputData{ .Spline = SplineData, .InputIndex = InputIndex });
 			SplineInputData.InitializeBounds(HalfDistanceRadius);
 		}
 	}
@@ -225,7 +571,7 @@ bool FPCGSplineIntersectionElement::ExecuteInternal(FPCGContext* Context) const
 	}
 
 	// Compute intersections...
-	TArray<PCGSplineIntersection::IntersectionPoint> Intersections;
+	TArray<PCGSplineIntersection::FIntersectionPoint> Intersections;
 
 	TArray<int> SplinesToTest;
 	for (int SplineIndex = 0; SplineIndex < InputSplines.Num(); ++SplineIndex)
@@ -244,12 +590,12 @@ bool FPCGSplineIntersectionElement::ExecuteInternal(FPCGContext* Context) const
 			}
 		}
 
-		const PCGSplineIntersection::InputData& CurrentSplineData = InputSplines[SplineIndex];
+		const PCGSplineIntersection::FInputData& CurrentSplineData = InputSplines[SplineIndex];
 		const FBox& CurrentBounds = CurrentSplineData.Bounds;
 		
 		for (int SplineToTest : SplinesToTest)
 		{
-			const PCGSplineIntersection::InputData& OtherSplineData = InputSplines[SplineToTest];
+			const PCGSplineIntersection::FInputData& OtherSplineData = InputSplines[SplineToTest];
 			const FBox& OtherBounds = OtherSplineData.Bounds;
 			const bool bIsSameData = (OtherSplineData.Spline == CurrentSplineData.Spline);
 
@@ -270,8 +616,7 @@ bool FPCGSplineIntersectionElement::ExecuteInternal(FPCGContext* Context) const
 					continue;
 				}
 
-				double CurrentRoughStep = 0.0;
-				const TArray<FVector>& CurrentSamples = CurrentSplineData.GetRoughSamples(SegmentIndex, Settings->DistanceThreshold, CurrentRoughStep);
+				const TArray<FVector>& CurrentSamples = CurrentSplineData.GetSamples(SegmentIndex, Settings->DistanceThreshold);
 
 				// Then iterate through potential segments
 				for (int32 OtherSegmentIndex = 0; OtherSegmentIndex < OtherSplineData.SegmentBounds.Num(); ++OtherSegmentIndex)
@@ -282,11 +627,8 @@ bool FPCGSplineIntersectionElement::ExecuteInternal(FPCGContext* Context) const
 						continue;
 					}
 
-					double OtherRoughStep = 0.0;
-					const TArray<FVector>& OtherSamples = OtherSplineData.GetRoughSamples(OtherSegmentIndex, Settings->DistanceThreshold, OtherRoughStep);
+					const TArray<FVector>& OtherSamples = OtherSplineData.GetSamples(OtherSegmentIndex, Settings->DistanceThreshold);
 					
-					//const double RoughToleranceSqr = 0.25 * (FMath::Square(CurrentRoughStep) + FMath::Square(OtherRoughStep));
-
 					// At this point, we are in intersection-possible territory.
 					// We don't have an analytical solution here, but there are a few things to consider:
 					// It's possible that two segments intersect in more than one place (up to 4 assumedly),
@@ -302,38 +644,31 @@ bool FPCGSplineIntersectionElement::ExecuteInternal(FPCGContext* Context) const
 							const FVector& C = OtherSamples[OtherSampleIndex];
 							const FVector& D = OtherSamples[OtherSampleIndex + 1];
 
-							// This is an approximation because we linearized the spline segment.
-							// If this isn't enough, there are a few solutions here:
-							// - Make the discretization size configurable
-							// - Change the threshold here so we incorporate the approximate difference
-							// .. and eventually resplit said segment
-							FVector X, Y;
-							// @todo_pcg: might want to use the safe version if we have colocated control points
-							FMath::SegmentDistToSegment(A, B, C, D, X, Y);
-
-							if ((X - Y).SquaredLength() > SqrDistanceThreshold)
+							double ABStart, ABEnd, CDStart, CDEnd;
+							if (PCGSplineIntersection::FindOverlapRange(A, B, C, D, SqrDistanceThreshold, ABStart, ABEnd, CDStart, CDEnd))
 							{
-								continue;
+								const FVector StartAB = A + (B - A) * ABStart;
+								const FVector EndAB = A + (B - A) * ABEnd;
+								const FVector StartCD = C + (D - C) * CDStart;
+								const FVector EndCD = C + (D - C) * CDEnd;
+
+								PCGSplineIntersection::FIntersectionPoint& Intersection = Intersections.Emplace_GetRef();
+								Intersection.A.Spline = CurrentSplineData.Spline;
+								Intersection.A.Input = &CurrentSplineData;
+								Intersection.A.InputKeyStart = CurrentSplineData.Spline->SplineStruct.FindInputKeyOnSegmentClosestToWorldLocation(StartAB, SegmentIndex);
+								Intersection.A.InputKeyEnd = CurrentSplineData.Spline->SplineStruct.FindInputKeyOnSegmentClosestToWorldLocation(EndAB, SegmentIndex);
+								Intersection.A.OtherSplines.Add(OtherSplineData.Spline);
+								Intersection.A.Points.Add(StartAB);
+								Intersection.A.Points.Add(EndAB);
+
+								Intersection.B.Spline = OtherSplineData.Spline;
+								Intersection.B.Input = &OtherSplineData;
+								Intersection.B.InputKeyStart = OtherSplineData.Spline->SplineStruct.FindInputKeyOnSegmentClosestToWorldLocation(StartCD, OtherSegmentIndex);
+								Intersection.B.InputKeyEnd = OtherSplineData.Spline->SplineStruct.FindInputKeyOnSegmentClosestToWorldLocation(EndCD, OtherSegmentIndex);
+								Intersection.B.OtherSplines.Add(CurrentSplineData.Spline);
+								Intersection.B.Points.Add(StartCD);
+								Intersection.B.Points.Add(EndCD);
 							}
-
-							const float CurrentKey = CurrentSplineData.Spline->SplineStruct.FindInputKeyOnSegmentClosestToWorldLocation(X, SegmentIndex);
-							const FVector ClosestPointOnCurrentSpline = CurrentSplineData.Spline->SplineStruct.GetLocationAtSplineInputKey(CurrentKey, ESplineCoordinateSpace::World);
-
-							const float OtherKey = OtherSplineData.Spline->SplineStruct.FindInputKeyOnSegmentClosestToWorldLocation(Y, OtherSegmentIndex);
-							const FVector ClosestPointOnOtherSpline = OtherSplineData.Spline->SplineStruct.GetLocationAtSplineInputKey(OtherKey, ESplineCoordinateSpace::World);
-
-							// @todo_pcg: Validate distance after?
-							const FVector MidPoint = (ClosestPointOnCurrentSpline + ClosestPointOnOtherSpline) * 0.5;
-							PCGSplineIntersection::IntersectionPoint& Intersection = Intersections.Emplace_GetRef();
-							Intersection.A = CurrentSplineData.Spline;
-							Intersection.InputA = &CurrentSplineData;
-							Intersection.B = OtherSplineData.Spline;
-							Intersection.InputB = &OtherSplineData;
-							Intersection.PositionA = ClosestPointOnCurrentSpline;
-							Intersection.PositionB = ClosestPointOnOtherSpline;
-							Intersection.AveragePosition = MidPoint;
-							Intersection.InputKeyA = CurrentKey;
-							Intersection.InputKeyB = OtherKey;
 						}
 					}
 				}
@@ -343,38 +678,43 @@ bool FPCGSplineIntersectionElement::ExecuteInternal(FPCGContext* Context) const
 
 	TArray<FPCGTaggedData>& Outputs = Context->OutputData.TaggedData;
 
-	auto LogErrorOnIndexAttributeCreation = [Context](const FName IndexName)
+	auto LogErrorOnAttributeCreation = [Context](const FName IndexName)
 	{
-		PCGLog::LogErrorOnGraph(FText::Format(LOCTEXT("FailedCreateIndexAttribute", "Failed to create the index attribute '{0}'."), FText::FromName(IndexName)), Context);
+		PCGLog::LogErrorOnGraph(FText::Format(LOCTEXT("FailedCreateIndexAttribute", "Failed to create the attribute '{0}'."), FText::FromName(IndexName)), Context);
 	};
 
 	// Finally, build output data.
 	if (Settings->Output == EPCGSplineIntersectionOutput::IntersectionPointsOnly)
 	{
-		// Start by filtering out points that are too close to each other.
-		TArray<PCGSplineIntersection::IntersectionPoint> FilteredIntersections = PCGSplineIntersection::CollapseColocatedIntersections(Intersections, SqrDistanceThreshold, nullptr);
+		// Merge intersection points on a global basis, regardless of source spline
+		TArray<PCGSplineIntersection::FCollapsedIntersectionPoint> CollapsedIntersections = PCGSplineIntersection::CollapseIntersections(Intersections, Settings->bUseBestFitBox);
 
 		FPCGTaggedData& Output = Outputs.Emplace_GetRef();
 		
 		UPCGBasePointData* PointData = FPCGContext::NewPointData_AnyThread(Context);
 		Output.Data = PointData;
 
-		PointData->SetNumPoints(FilteredIntersections.Num());
+		PointData->SetNumPoints(CollapsedIntersections.Num());
 
-		const EPCGPointNativeProperties PropertiesToAllocate = EPCGPointNativeProperties::Transform | (Settings->bOutputSplineIndices ? EPCGPointNativeProperties::MetadataEntry : EPCGPointNativeProperties::None);
+		const EPCGPointNativeProperties PropertiesToAllocate = EPCGPointNativeProperties::Transform | 
+			EPCGPointNativeProperties::BoundsMin | 
+			EPCGPointNativeProperties::BoundsMax | 
+			((Settings->bOutputSplineIndices || Settings->bOutputIntersectionType) ? EPCGPointNativeProperties::MetadataEntry : EPCGPointNativeProperties::None);
 		PointData->AllocateProperties(PropertiesToAllocate);
 
+		// Implementation note: there's no point in outputting the intersection type, since these are always intersections with their appropriate sizes.
+		// @todo_pcg: add array support for index attributes maybe
 		FPCGMetadataAttribute<int32>* FirstIndexAttribute = nullptr;
 		FPCGMetadataAttribute<int32>* SecondIndexAttribute = nullptr;
 
 		if (Settings->bOutputSplineIndices)
 		{
-			auto CreateIndexAttribute = [PointData, &LogErrorOnIndexAttributeCreation](const FName IndexName) -> FPCGMetadataAttribute<int32>*
+			auto CreateIndexAttribute = [PointData, &LogErrorOnAttributeCreation](const FName IndexName) -> FPCGMetadataAttribute<int32>*
 			{
 				FPCGMetadataAttribute<int32>* IndexAttribute = PointData->MutableMetadata()->FindOrCreateAttribute<int32>(IndexName, -1, /*bAllowInterpolation=*/false);
 				if (!IndexAttribute)
 				{
-					LogErrorOnIndexAttributeCreation(IndexName);
+					LogErrorOnAttributeCreation(IndexName);
 				}
 
 				return IndexAttribute;
@@ -386,42 +726,45 @@ bool FPCGSplineIntersectionElement::ExecuteInternal(FPCGContext* Context) const
 
 		FPCGPointValueRanges OutRanges(PointData, /*bAllocate=*/false);
 
-		for (int Index = 0; Index < FilteredIntersections.Num(); ++Index)
+		for (int Index = 0; Index < CollapsedIntersections.Num(); ++Index)
 		{
-			const PCGSplineIntersection::IntersectionPoint& Point = FilteredIntersections[Index];
-			OutRanges.TransformRange[Index] = FTransform(Point.AveragePosition);
+			const PCGSplineIntersection::FCollapsedIntersectionPoint& Point = CollapsedIntersections[Index];
+			check(Point.Ranges.Num() >= 1);
+
+			OutRanges.TransformRange[Index] = Point.Transform;
+			OutRanges.BoundsMinRange[Index] = Point.Bounds.Min;
+			OutRanges.BoundsMaxRange[Index] = Point.Bounds.Max;
+
+			PCGMetadataEntryKey* CurrentEntryKey = nullptr;
 
 			if (FirstIndexAttribute || SecondIndexAttribute)
 			{
-				PCGMetadataEntryKey& CurrentEntryKey = OutRanges.MetadataEntryRange[Index];
-				PointData->Metadata->InitializeOnSet(CurrentEntryKey);
+				CurrentEntryKey = &OutRanges.MetadataEntryRange[Index];
+				PointData->Metadata->InitializeOnSet(*CurrentEntryKey);
+			}
 
-				if (FirstIndexAttribute)
-				{
-					FirstIndexAttribute->SetValue(CurrentEntryKey, Point.InputA->InputIndex);
-				}
+			if (FirstIndexAttribute)
+			{
+				FirstIndexAttribute->SetValue(*CurrentEntryKey, Point.Ranges[0].Input->InputIndex);
+			}
 
-				if (SecondIndexAttribute)
-				{
-					SecondIndexAttribute->SetValue(CurrentEntryKey, Point.InputB->InputIndex);
-				}
+			if (Point.Ranges.Num() > 1 && SecondIndexAttribute)
+			{
+				SecondIndexAttribute->SetValue(*CurrentEntryKey, Point.Ranges[1].Input->InputIndex);
 			}
 		}
 	}
 	else if (Settings->Output == EPCGSplineIntersectionOutput::OriginalSplinesWithIntersections)
 	{
-		for (const PCGSplineIntersection::InputData& InputData : InputSplines)
+		for (const PCGSplineIntersection::FInputData& InputData : InputSplines)
 		{
 			const UPCGSplineData* InputSpline = InputData.Spline;
 
 			// Filter points for the current spline and collapse points that would be too close.
 			// Note: the average position in this case will be when points are collapsed on the spline itself, not against the other splines
-			TArray<PCGSplineIntersection::IntersectionPoint> CollapsedPoints = PCGSplineIntersection::CollapseColocatedIntersections(Intersections, SqrDistanceThreshold, InputSpline);
+			TArray<PCGSplineIntersection::FIntersectionRange> CollapsedPoints = PCGSplineIntersection::FilterAndMergeIntersections(Intersections, InputSpline);
 			// Sort by key
-			CollapsedPoints.Sort([InputSpline](const auto& A, const auto& B) { return A.GetKey(InputSpline) < B.GetKey(InputSpline); });
-
-			// @todo_pcg : we should not add points that are already on existing control points,
-			// but we should mark them as intersections
+			CollapsedPoints.Sort([](const auto& A, const auto& B) { return A.InputKeyStart < B.InputKeyStart; });
 
 			// At this point, we've successfully filtered/collapsed points if it was needed,
 			// Or we had no new points - in which case we have nothing to do.
@@ -433,6 +776,7 @@ bool FPCGSplineIntersectionElement::ExecuteInternal(FPCGContext* Context) const
 
 				// implementation note: this does not support typecasting, will just overwrite the attribute.
 				FPCGMetadataAttribute<int32>* IndexAttribute = nullptr;
+				FPCGMetadataAttribute<int32>* IntersectionTypeAttribute = nullptr;
 
 				if (Settings->bOutputSplineIndices)
 				{
@@ -441,7 +785,16 @@ bool FPCGSplineIntersectionElement::ExecuteInternal(FPCGContext* Context) const
 
 					if (!IndexAttribute)
 					{
-						LogErrorOnIndexAttributeCreation(IndexAttributeName);
+						LogErrorOnAttributeCreation(IndexAttributeName);
+					}
+				}
+
+				if (Settings->bOutputIntersectionType)
+				{
+					IntersectionTypeAttribute = NewSplineData->MutableMetadata()->FindOrCreateAttribute<int32>(Settings->IntersectionTypeAttribute.GetAttributeName(), static_cast<int32>(EPCGSplineIntersectionPointType::NotAnIntersection), /*bAllowInterpolation=*/false);
+					if (!IntersectionTypeAttribute)
+					{
+						LogErrorOnAttributeCreation(Settings->IntersectionTypeAttribute.GetAttributeName());
 					}
 				}
 
@@ -453,7 +806,7 @@ bool FPCGSplineIntersectionElement::ExecuteInternal(FPCGContext* Context) const
 				TArray<int64> SplinePointsEntryKeys = InputData.Spline->GetMetadataEntryKeysForSplinePoints();
 				check(SplinePoints.Num() == SplinePointsEntryKeys.Num() || SplinePointsEntryKeys.IsEmpty());
 				const bool bHasMetadata = !SplinePointsEntryKeys.IsEmpty();
-				const bool bNeedsMetadata = (IndexAttribute != nullptr);
+				const bool bNeedsMetadata = (IndexAttribute != nullptr || IntersectionTypeAttribute != nullptr);
 
 				// Prepare metadata entry keys if we're going to use them later
 				if (bNeedsMetadata && !bHasMetadata)
@@ -465,7 +818,7 @@ bool FPCGSplineIntersectionElement::ExecuteInternal(FPCGContext* Context) const
 					}
 				}
 
-				auto OptionalWriteIntersectionMetadata = [bNeedsMetadata, NewSplineData, IndexAttribute, &InputData](int64& EntryKey, const PCGSplineIntersection::IntersectionPoint& IntersectionPoint)
+				auto OptionalWriteIntersectionMetadata = [bNeedsMetadata, NewSplineData, IndexAttribute, IntersectionTypeAttribute, &InputData, &InputSplines](int64& EntryKey, const PCGSplineIntersection::FIntersectionRange& IntersectionPoint, EPCGSplineIntersectionPointType IntersectionType)
 				{
 					if (bNeedsMetadata)
 					{
@@ -473,94 +826,140 @@ bool FPCGSplineIntersectionElement::ExecuteInternal(FPCGContext* Context) const
 
 						if (IndexAttribute)
 						{
-							IndexAttribute->SetValue(EntryKey, IntersectionPoint.GetOtherInput(&InputData)->InputIndex);
+							check(!IntersectionPoint.OtherSplines.IsEmpty());
+							const PCGSplineIntersection::FInputData* IntersectedSplineData = InputSplines.FindByPredicate([&IntersectionPoint](const PCGSplineIntersection::FInputData& Data) { return Data.Spline == IntersectionPoint.OtherSplines[0]; });
+							if(IntersectedSplineData)
+							{
+								IndexAttribute->SetValue(EntryKey, IntersectedSplineData->InputIndex);
+							}
+						}
+
+						if (IntersectionTypeAttribute)
+						{
+							IntersectionTypeAttribute->SetValue(EntryKey, IntersectionType);
 						}
 					}
 				};
 				
-				const float MinKeyDifferenceThreshold = 1.0e-4f;
+				constexpr double MinKeyDifferenceThreshold = PCGSplineIntersection::Constants::KeyDifferenceThreshold;
 				
-				for (const PCGSplineIntersection::IntersectionPoint& PointToAdd : CollapsedPoints)
+				for (const PCGSplineIntersection::FIntersectionRange& PointToAdd : CollapsedPoints)
 				{
-					const float Key = PointToAdd.GetKey(InputSpline);
-					const FVector& Position = InputData.Spline->SplineStruct.GetTransform().InverseTransformPosition(PointToAdd.AveragePosition);
+					TArray<double, TInlineAllocator<2>> Keys;
+					TArray<EPCGSplineIntersectionPointType, TInlineAllocator<2>> KeyTypes;
 
-					// Find insertion index @todo_pcg: use bisection
-					int InsertionIndex = SplinePoints.Num();
-					for (int SPIndex = 0; SPIndex < SplinePoints.Num(); ++SPIndex)
+					double IntersectionKeyStart = PointToAdd.InputKeyStart;
+					double IntersectionKeyEnd = PointToAdd.InputKeyEnd;
+					double LoopKeyShift = 0;
+
+					if (IntersectionKeyEnd < IntersectionKeyStart)
 					{
-						if (SplinePoints[SPIndex].InputKey > Key)
-						{
-							InsertionIndex = SPIndex;
-							break;
-						}
+						ensure(PointToAdd.Spline->IsClosed());
+						LoopKeyShift = PCGSplineIntersection::Helpers::GetSplineKeyLength(InputSpline);
+						IntersectionKeyEnd += LoopKeyShift;
 					}
 
-					bool bIsOnLastSegment = (InsertionIndex == SplinePoints.Num());
-
-					// Skip points that are on control points - need to check against previous and next.
-					FSplinePoint& PreviousPoint = SplinePoints[InsertionIndex - 1];
-					FSplinePoint& NextPoint = SplinePoints[bIsOnLastSegment ? 0 : InsertionIndex];
-
-					const float NextPointVirtualKey = (bIsOnLastSegment ? FMath::CeilToInt(PreviousPoint.InputKey + 1.0e-4f) : NextPoint.InputKey);
-
-					if ((FMath::Abs(PreviousPoint.InputKey - Key) <= MinKeyDifferenceThreshold && (PreviousPoint.Position - Position).SquaredLength() <= SqrDistanceThreshold))
+					if ((IntersectionKeyEnd - IntersectionKeyStart) > MinKeyDifferenceThreshold && Settings->bOutputComplexIntersections)
 					{
-						if (bNeedsMetadata)
+						Keys.Add(IntersectionKeyStart);
+						Keys.Add(IntersectionKeyEnd - LoopKeyShift);
+						KeyTypes.Add(EPCGSplineIntersectionPointType::IntersectionStart);
+						KeyTypes.Add(EPCGSplineIntersectionPointType::IntersectionEnd);
+					}
+					else
+					{
+						double AverageKey = 0.5 * (IntersectionKeyStart + IntersectionKeyEnd);
+						if (AverageKey > LoopKeyShift)
 						{
-							// Mark previous point as an intersection
-							OptionalWriteIntersectionMetadata(SplinePointsEntryKeys[InsertionIndex - 1], PointToAdd);
+							AverageKey -= LoopKeyShift;
 						}
 
-						continue;
+						Keys.Add(AverageKey);
+						KeyTypes.Add(EPCGSplineIntersectionPointType::SinglePointIntersection);
 					}
-					else if ((FMath::Abs(NextPointVirtualKey - Key) <= MinKeyDifferenceThreshold && (NextPoint.Position - Position).SquaredLength() <= SqrDistanceThreshold))
+
+					for(int KeyIndex = 0; KeyIndex < Keys.Num(); ++KeyIndex)
 					{
-						if (bNeedsMetadata)
+						const double& Key = Keys[KeyIndex];
+						const FVector Position = InputData.Spline->SplineStruct.GetLocationAtSplineInputKey(Key, ESplineCoordinateSpace::Local);
+
+						// Find insertion index @todo_pcg: use bisection
+						int InsertionIndex = SplinePoints.Num();
+						for (int SPIndex = 1; SPIndex < SplinePoints.Num(); ++SPIndex)
 						{
-							// Mark next point as an intersection
-							OptionalWriteIntersectionMetadata(SplinePointsEntryKeys[bIsOnLastSegment ? 0 : InsertionIndex], PointToAdd);
+							if (SplinePoints[SPIndex].InputKey > Key)
+							{
+								InsertionIndex = SPIndex;
+								break;
+							}
 						}
 
-						continue;
-					}
+						bool bIsOnLastSegment = (InsertionIndex == SplinePoints.Num());
 
-					FSplinePoint SplinePoint;
-					SplinePoint.InputKey = Key;
-					SplinePoint.Position = Position;
-					// @todo_pcg: we don't need to add points that are on the control point, however, we might want to do metadata for those
-					SplinePoint.ArriveTangent = InputData.Spline->SplineStruct.GetTangentAtSplineInputKey(Key, ESplineCoordinateSpace::Local);
-					SplinePoint.LeaveTangent = SplinePoint.ArriveTangent;
-					SplinePoint.Rotation = InputData.Spline->SplineStruct.GetQuaternionAtSplineInputKey(Key, ESplineCoordinateSpace::Local).Rotator();
-					SplinePoint.Scale = InputData.Spline->SplineStruct.GetScaleAtSplineInputKey(Key);
-					SplinePoint.Type = ESplinePointType::CurveCustomTangent;
-
-					// Adjust tangents
-					const double KeyRange = FMath::Max(NextPointVirtualKey - PreviousPoint.InputKey, 1.0e-4);
-					const double PreviousRatio = (Key - PreviousPoint.InputKey) / KeyRange;
-					const double NextRatio = (NextPointVirtualKey - Key) / KeyRange;
-
-					PreviousPoint.LeaveTangent *= PreviousRatio;
-					NextPoint.ArriveTangent *= NextRatio;
-
-					// The spline point tangent is "always" taken from the original spline's perspective,
-					// so we will assume that the range is always one, so we need to scale accordingly.
-					// @todo_pcg: get key at the beginning of the segment
-					SplinePoint.ArriveTangent *= (Key - PreviousPoint.InputKey);
-					SplinePoint.LeaveTangent *= (NextPointVirtualKey - Key);
-
-					SplinePoints.Insert(SplinePoint, InsertionIndex);
-					if (bHasMetadata || bNeedsMetadata)
-					{
-						PCGMetadataEntryKey EntryKey = PCGInvalidEntryKey;
+						// Skip points that are on control points - need to check against previous and next.
+						FSplinePoint& PreviousPoint = SplinePoints[InsertionIndex - 1];
+						FSplinePoint& NextPoint = SplinePoints[bIsOnLastSegment ? 0 : InsertionIndex];
+
+						const float NextPointVirtualKey = (bIsOnLastSegment ? FMath::CeilToInt(PreviousPoint.InputKey + 1.0e-4f) : NextPoint.InputKey);
 
-						if (bHasMetadata)
+						if ((FMath::Abs(PreviousPoint.InputKey - Key) <= MinKeyDifferenceThreshold && (PreviousPoint.Position - Position).SquaredLength() <= SqrDistanceThreshold))
 						{
-							InputData.Spline->WriteMetadataToEntry(Key, EntryKey, NewSplineData->MutableMetadata());
+							if (bNeedsMetadata)
+							{
+								// Mark previous point as an intersection
+								OptionalWriteIntersectionMetadata(SplinePointsEntryKeys[InsertionIndex - 1], PointToAdd, KeyTypes[KeyIndex]);
+							}
+
+							continue;
+						}
+						else if ((FMath::Abs(NextPointVirtualKey - Key) <= MinKeyDifferenceThreshold && (NextPoint.Position - Position).SquaredLength() <= SqrDistanceThreshold))
+						{
+							if (bNeedsMetadata)
+							{
+								// Mark next point as an intersection
+								OptionalWriteIntersectionMetadata(SplinePointsEntryKeys[bIsOnLastSegment ? 0 : InsertionIndex], PointToAdd, KeyTypes[KeyIndex]);
+							}
+
+							continue;
 						}
 
-						OptionalWriteIntersectionMetadata(EntryKey, PointToAdd);
-						SplinePointsEntryKeys.Insert(EntryKey, InsertionIndex);
+						FSplinePoint SplinePoint;
+						SplinePoint.InputKey = Key;
+						SplinePoint.Position = Position;
+						// @todo_pcg: we don't need to add points that are on the control point, however, we might want to do metadata for those
+						SplinePoint.ArriveTangent = InputData.Spline->SplineStruct.GetTangentAtSplineInputKey(Key, ESplineCoordinateSpace::Local);
+						SplinePoint.LeaveTangent = SplinePoint.ArriveTangent;
+						SplinePoint.Rotation = InputData.Spline->SplineStruct.GetQuaternionAtSplineInputKey(Key, ESplineCoordinateSpace::Local).Rotator();
+						SplinePoint.Scale = InputData.Spline->SplineStruct.GetScaleAtSplineInputKey(Key);
+						SplinePoint.Type = ESplinePointType::CurveCustomTangent;
+
+						// Adjust tangents
+						const double KeyRange = FMath::Max(NextPointVirtualKey - PreviousPoint.InputKey, 1.0e-4);
+						const double PreviousRatio = (Key - PreviousPoint.InputKey) / KeyRange;
+						const double NextRatio = (NextPointVirtualKey - Key) / KeyRange;
+
+						PreviousPoint.LeaveTangent *= PreviousRatio;
+						NextPoint.ArriveTangent *= NextRatio;
+
+						// The spline point tangent is "always" taken from the original spline's perspective,
+						// so we will assume that the range is always one, so we need to scale accordingly.
+						// @todo_pcg: get key at the beginning of the segment
+						SplinePoint.ArriveTangent *= (Key - PreviousPoint.InputKey);
+						SplinePoint.LeaveTangent *= (NextPointVirtualKey - Key);
+
+						SplinePoints.Insert(SplinePoint, InsertionIndex);
+						if (bHasMetadata || bNeedsMetadata)
+						{
+							PCGMetadataEntryKey EntryKey = PCGInvalidEntryKey;
+
+							if (bHasMetadata)
+							{
+								InputData.Spline->WriteMetadataToEntry(Key, EntryKey, NewSplineData->MutableMetadata());
+							}
+
+							OptionalWriteIntersectionMetadata(EntryKey, PointToAdd, KeyTypes[KeyIndex]);
+							SplinePointsEntryKeys.Insert(EntryKey, InsertionIndex);
+						}
 					}
 				}
 
diff --git a/Engine/Plugins/PCG/Source/PCG/Public/Elements/PCGSplineIntersection.h b/Engine/Plugins/PCG/Source/PCG/Public/Elements/PCGSplineIntersection.h
--- a/Engine/Plugins/PCG/Source/PCG/Public/Elements/PCGSplineIntersection.h
+++ b/Engine/Plugins/PCG/Source/PCG/Public/Elements/PCGSplineIntersection.h
@@ -26,6 +26,15 @@ enum class EPCGSplineIntersectionOutput : uint8
 	OriginalSplinesWithIntersections
 };
 
+UENUM(BlueprintType)
+enum class EPCGSplineIntersectionPointType : uint8
+{
+	NotAnIntersection = 0,
+	SinglePointIntersection = 1,
+	IntersectionStart = 2,
+	IntersectionEnd = 3
+};
+
 /** 
 * Intersects splines against other splines (or themselves) and returns varied results
 * based on user need.
@@ -67,15 +76,31 @@ public:
 
 	/** Attribute that will contain the first spline index (only in the output case of intersection points.) */
 	UPROPERTY(BlueprintReadWrite, EditAnywhere, Category = Settings, meta=(PCG_Overridable, PCG_DiscardPropertySelection, PCG_DiscardExtraSelection, EditCondition="bOutputSplineIndices && Type!=EPCGSplineIntersectionType::Self && Output==EPCGSplineIntersectionOutput::IntersectionPointsOnly"))
-	FPCGAttributePropertyOutputSelector OriginatingSplineIndexAttribute;
+	FPCGAttributePropertyOutputNoSourceSelector OriginatingSplineIndexAttribute;
 
 	/** Attribute that will contain the first intersecting spline index, or -1 if not an intersection control point. */
 	UPROPERTY(BlueprintReadWrite, EditAnywhere, Category = Settings, meta=(PCG_Overridable, PCG_DiscardPropertySelection, PCG_DiscardExtraSelection, EditCondition="bOutputSplineIndices"))
-	FPCGAttributePropertyOutputSelector IntersectingSplineIndexAttribute;
+	FPCGAttributePropertyOutputNoSourceSelector IntersectingSplineIndexAttribute;
+
+	/** Controls whether we'll output intersection ranges (start/end) when relevant instead of singular points. */
+	UPROPERTY(BlueprintReadWrite, EditAnywhere, Category = Settings, meta=(PCG_Overridable))
+	bool bOutputComplexIntersections = false;
+
+	/** Controls whether intersection types will be written to the attributes. */
+	UPROPERTY(BlueprintReadWrite, EditAnywhere, Category = Settings, meta=(PCG_Overridable, InlineEditConditionToggle))
+	bool bOutputIntersectionType = false;
+
+	/** Attribute that will contain the intersection type (0: not an intersection, 1: point intersection, 2: intersection start, 3: intersection end. */
+	UPROPERTY(BlueprintReadWrite, EditAnywhere, Category = Settings, meta=(PCG_Overridable, PCG_DiscardPropertySelection, PCG_DiscardExtraSelection, EditCondition="bOutputIntersectionType"))
+	FPCGAttributePropertyOutputNoSourceSelector IntersectionTypeAttribute;
 
 	/** Maximum distance at which we report an intersection on the splines. */
 	UPROPERTY(BlueprintReadWrite, EditAnywhere, Category = Settings, meta=(ClampMin="0.0001", PCG_Overridable))
 	double DistanceThreshold = 10.0;
+
+	/** Controls whether the output for intersection points is oriented using a best fix box (using PCA - Principal Component Analysis) or if it's aligned on the longest intersection. */
+	UPROPERTY(BlueprintReadWrite, EditAnywhere, Category = Settings, meta=(PCG_Overridable, EditCondition="Output==EPCGSplineIntersectionOutput::IntersectionPointsOnly", EditConditionHides))
+	bool bUseBestFitBox = true;
 };
 
 class FPCGSplineIntersectionElement : public IPCGElement
