[PCG] pcg.Cache.LogStats is more verbose giving details instance count for each data type and total size

#rnx

#Synced-CL 51261860

diff --git a/Engine/Plugins/PCG/Source/PCG/Private/Graph/PCGGraphCache.cpp b/Engine/Plugins/PCG/Source/PCG/Private/Graph/PCGGraphCache.cpp
--- a/Engine/Plugins/PCG/Source/PCG/Private/Graph/PCGGraphCache.cpp
+++ b/Engine/Plugins/PCG/Source/PCG/Private/Graph/PCGGraphCache.cpp
@@ -140,8 +140,53 @@ void FPCGGraphCache::ResetStats()
 
 void FPCGGraphCache::LogStats()
 {
+	PCG::TScopeLock ScopedLock(CacheLock);
 	const float HitRate = Stats.GetCount > 0 ? ((float)Stats.HitCount / Stats.GetCount) * 100.f : 0.f;
-	UE_LOG(LogPCG, Log, TEXT("Cache Stats : %lld Stores, %lld Hits / %lld Gets (%f), %lld Grow count"), Stats.StoreCount, Stats.HitCount, Stats.GetCount, HitRate, Stats.GrowCount);
+
+	UE_LOG(LogPCG, Log, TEXT("Cache Stats : %lld Stores, %lld Hits / %lld Gets (%.2f %%), %lld Grow count / %lld Remove count / %d Entry count / %.2f Memory(MB)"), 
+		Stats.StoreCount, 
+		Stats.HitCount, 
+		Stats.GetCount, 
+		HitRate, 
+		Stats.GrowCount,
+		Stats.RemoveCount,
+		CacheData.Num(),
+		TotalMemoryUsed / 1024.f / 1024.f);
+
+	struct FPerClassStats
+	{
+		int32 Instances = 0;
+		uint64 TotalMemoryUsed = 0;
+	};
+
+	TMap<UClass*, FPerClassStats> PerClassStats;
+	FPerClassStats UnknownStats;
+
+	for (const auto& [UID, MemoryRecord] : MemoryRecords)
+	{
+		if (UObject* DataPtr = MemoryRecord.Data.ResolveObjectPtr())
+		{
+			FPerClassStats& ClassStats = PerClassStats.FindOrAdd(DataPtr->GetClass());
+			++ClassStats.Instances;
+			ClassStats.TotalMemoryUsed += MemoryRecord.MemoryPerInstance;
+		}
+		else
+		{
+			// Unknown object
+			++UnknownStats.Instances;
+			UnknownStats.TotalMemoryUsed += MemoryRecord.MemoryPerInstance;
+		}
+	}
+
+	for (const auto& [ClassPtr, ClassStats] : PerClassStats)
+	{
+		UE_LOG(LogPCG, Log, TEXT("Cache Stats: %s : %d Instances : %.2f Memory(MB)"), *ClassPtr->GetPathName(), ClassStats.Instances, ClassStats.TotalMemoryUsed / 1024.f / 1024.f);
+	}
+
+	if (UnknownStats.Instances > 0)
+	{
+		UE_LOG(LogPCG, Log, TEXT("Cache Stats: <Unknown type> : %d Instances : %.2f Memory(MB)"), UnknownStats.Instances, UnknownStats.TotalMemoryUsed / 1024.f / 1024.f);
+	}
 }
 
 bool FPCGGraphCache::IsEnabled() const
@@ -387,6 +432,8 @@ void FPCGGraphCache::AddToCacheInternal(const FPCGCacheEntryKey& InKey, const FP
 
 void FPCGGraphCache::RemoveFromCacheInternal(const FPCGCacheEntryKey& InKey)
 {
+	++Stats.RemoveCount;
+
 	if (TSet<FPCGCacheEntryKey>* ElementCacheEntryKeys = ElementToCacheEntryKeys.Find(InKey.GetElement()))
 	{
 		ElementCacheEntryKeys->Remove(InKey);
@@ -443,6 +490,7 @@ void FPCGGraphCache::AddDataToAccountedMemory(const FPCGDataCollection& InCollec
 						const SIZE_T DataSize = ResSize.GetDedicatedSystemMemoryBytes();
 
 						FCachedMemoryRecord& NewRecord = MemoryRecords.Add(Data->UID);
+						NewRecord.Data = Data;
 						NewRecord.MemoryPerInstance = DataSize;
 						NewRecord.InstanceCount = 1;
 						TotalMemoryUsed += DataSize;
diff --git a/Engine/Plugins/PCG/Source/PCG/Private/Graph/PCGGraphCache.h b/Engine/Plugins/PCG/Source/PCG/Private/Graph/PCGGraphCache.h
--- a/Engine/Plugins/PCG/Source/PCG/Private/Graph/PCGGraphCache.h
+++ b/Engine/Plugins/PCG/Source/PCG/Private/Graph/PCGGraphCache.h
@@ -9,6 +9,7 @@
 #include "Containers/LruCache.h"
 #include "Misc/SpinLock.h"
 #include "UObject/GCObject.h"
+#include "UObject/ObjectKey.h"
 
 class IPCGElement;
 class UPCGComponent;
@@ -25,6 +26,9 @@ struct FCachedMemoryRecord
 
 	/** Memory usage in bytes of one instance. Cached here to save recomputing later. */
 	SIZE_T MemoryPerInstance = 0;
+
+	/** Allows retrieval of data ptr for logging purposes. */
+	FObjectKey Data;
 };
 
 struct FPCGCacheEntryKey
@@ -145,6 +149,7 @@ private:
 		int64 StoreCount = 0;
 		int64 HitCount = 0;
 		int64 GrowCount = 0;
+		int64 RemoveCount = 0;
 	};
 	
 	mutable FPCGCacheStats Stats;
