Small code cleanup improvements

#rnx

#Synced-CL 50984763

diff --git a/Engine/Plugins/PCG/Source/PCG/Private/Elements/PCGSplineIntersection.cpp b/Engine/Plugins/PCG/Source/PCG/Private/Elements/PCGSplineIntersection.cpp
--- a/Engine/Plugins/PCG/Source/PCG/Private/Elements/PCGSplineIntersection.cpp
+++ b/Engine/Plugins/PCG/Source/PCG/Private/Elements/PCGSplineIntersection.cpp
@@ -434,14 +434,13 @@ namespace PCGSplineIntersection
 			return false;
 		}
 
-		const double ABHint = (X - A).Length() / (B - A).Length();
-		const double CDHint = (Y - C).Length() / (D - C).Length();
-
 		// TODO : we should early test out trivial cases
 		auto FindRange = [&SqrDistanceThreshold](const FVector& A, const FVector& B, const FVector& C, const FVector& D, const double& Hint, double& MinRange, double& MaxRange)
 		{
 			const FVector& Start = A;
 			const double Len = (B - A).Length();
+			check(Len >= UE_DOUBLE_SMALL_NUMBER);
+
 			const FVector Dir = (B - A) / Len;
 
 			auto FindEnd = [&SqrDistanceThreshold, &Start, &Len, &Dir, &C, &D, &Hint](bool bFindMin)
@@ -471,9 +470,30 @@ namespace PCGSplineIntersection
 			MinRange = FindEnd(/*bFindMin=*/true);
 			MaxRange = FindEnd(/*bFindMin=*/false);
 		};
-		
-		FindRange(A, B, C, D, ABHint, ABStart, ABEnd);
-		FindRange(C, D, A, B, CDHint, CDStart, CDEnd);
+
+		const double ABLen = (B - A).Length();
+		if (ABLen >= UE_DOUBLE_SMALL_NUMBER)
+		{
+			const double ABHint = (X - A).Length() / ABLen;
+			FindRange(A, B, C, D, ABHint, ABStart, ABEnd);
+		}
+		else
+		{
+			ABStart = 0.0;
+			ABEnd = 1.0;
+		}
+
+		const double CDLen = (D - C).Length();
+		if (CDLen >= UE_DOUBLE_SMALL_NUMBER)
+		{
+			const double CDHint = (Y - C).Length() / CDLen;
+			FindRange(C, D, A, B, CDHint, CDStart, CDEnd);
+		}
+		else
+		{
+			CDStart = 0.0;
+			CDEnd = 1.0;
+		}
 
 		check(ABStart <= ABEnd);
 		check(CDStart <= CDEnd);
diff --git a/Engine/Plugins/PCG/Source/PCG/Public/Elements/PCGSplineIntersection.h b/Engine/Plugins/PCG/Source/PCG/Public/Elements/PCGSplineIntersection.h
--- a/Engine/Plugins/PCG/Source/PCG/Public/Elements/PCGSplineIntersection.h
+++ b/Engine/Plugins/PCG/Source/PCG/Public/Elements/PCGSplineIntersection.h
@@ -90,7 +90,7 @@ public:
 	UPROPERTY(BlueprintReadWrite, EditAnywhere, Category = Settings, meta=(PCG_Overridable, InlineEditConditionToggle))
 	bool bOutputIntersectionType = false;
 
-	/** Attribute that will contain the intersection type (0: not an intersection, 1: point intersection, 2: intersection start, 3: intersection end. */
+	/** Attribute that will contain the intersection type (0: not an intersection, 1: point intersection, 2: intersection start, 3: intersection end.) */
 	UPROPERTY(BlueprintReadWrite, EditAnywhere, Category = Settings, meta=(PCG_Overridable, PCG_DiscardPropertySelection, PCG_DiscardExtraSelection, EditCondition="bOutputIntersectionType"))
 	FPCGAttributePropertyOutputNoSourceSelector IntersectionTypeAttribute;
 
@@ -98,7 +98,7 @@ public:
 	UPROPERTY(BlueprintReadWrite, EditAnywhere, Category = Settings, meta=(ClampMin="0.0001", PCG_Overridable))
 	double DistanceThreshold = 10.0;
 
-	/** Controls whether the output for intersection points is oriented using a best fix box (using PCA - Principal Component Analysis) or if it's aligned on the longest intersection. */
+	/** Controls whether the output for intersection points is oriented using a best fit box (using PCA - Principal Component Analysis) or if it's aligned on the longest intersection. */
 	UPROPERTY(BlueprintReadWrite, EditAnywhere, Category = Settings, meta=(PCG_Overridable, EditCondition="Output==EPCGSplineIntersectionOutput::IntersectionPointsOnly", EditConditionHides))
 	bool bUseBestFitBox = true;
 };
