[PCG] Expose new virtual GetDataOverridePhase in SMS

#Synced-CL 51274443

diff --git a/Engine/Plugins/PCG/Source/PCG/Public/Elements/PCGStaticMeshSpawner.h b/Engine/Plugins/PCG/Source/PCG/Public/Elements/PCGStaticMeshSpawner.h
--- a/Engine/Plugins/PCG/Source/PCG/Public/Elements/PCGStaticMeshSpawner.h
+++ b/Engine/Plugins/PCG/Source/PCG/Public/Elements/PCGStaticMeshSpawner.h
@@ -152,7 +152,7 @@ protected:
 	UE_API virtual void AbortInternal(FPCGContext* Context) const override;
 	UE_API void SpawnStaticMeshInstances(FPCGStaticMeshSpawnerContext* Context, const FPCGStaticMeshSpawnerContext::FPackedInstanceListData& InstanceListData, const FPCGMeshInstanceList& InstanceList, const FPCGPackedCustomData* InPackedCustomData) const;
 	virtual bool SupportsBasePointDataInputs(FPCGContext* InContext) const { return true; }
-	virtual EPCGDataOverridePhase GetDataOverridePhase() const override;
+	UE_API virtual EPCGDataOverridePhase GetDataOverridePhase() const override;
 };
 
 #undef UE_API
