[PCG] Fix minor null check issues in Attribute Partition

#Synced-CL 51274443

diff --git a/Engine/Plugins/PCG/Source/PCG/Private/Metadata/PCGMetadataPartitionCommon.cpp b/Engine/Plugins/PCG/Source/PCG/Private/Metadata/PCGMetadataPartitionCommon.cpp
--- a/Engine/Plugins/PCG/Source/PCG/Private/Metadata/PCGMetadataPartitionCommon.cpp
+++ b/Engine/Plugins/PCG/Source/PCG/Private/Metadata/PCGMetadataPartitionCommon.cpp
@@ -494,7 +494,7 @@ namespace PCGMetadataPartitionCommon
 
 	TArray<UPCGData*> AttributeParamSpatialPartition(const UPCGData* InData, const TArrayView<const FPCGAttributePropertySelector>& InSelectorArray, FPCGContext* InOptionalContext, bool bSilenceMissingAttributeErrors)
 	{
-		if (!InData->IsA<UPCGSpatialData>() && !InData->IsA<UPCGParamData>())
+		if (!InData || (!InData->IsA<UPCGSpatialData>() && !InData->IsA<UPCGParamData>()))
 		{
 			PCGLog::LogErrorOnGraph(LOCTEXT("InvalidDataType", "Input data is not an attribute set nor a spatial data. Operation not supported."), InOptionalContext);
 			return {};
@@ -566,7 +566,7 @@ namespace PCGMetadataPartitionCommon
 
 	UPCGData* RemoveDuplicatesParamSpatial(const UPCGData* InData, const TArrayView<const FPCGAttributePropertySelector>& InSelectorArray, FPCGContext* InOptionalContext, bool bSilenceMissingAttributeErrors)
 	{
-		if (!InData->IsA<UPCGSpatialData>() && !InData->IsA<UPCGParamData>())
+		if (!InData || (!InData->IsA<UPCGSpatialData>() && !InData->IsA<UPCGParamData>()))
 		{
 			PCGLog::LogErrorOnGraph(LOCTEXT("InvalidDataType", "Input data is not an attribute set nor a spatial data. Operation not supported."), InOptionalContext);
 			return nullptr;
