[PCG] Front end for the Manual edit data override.

+ Uses the Component visualizer.
+ First version, Unpolished, supports only EditData.

Committed on behalf of jeansimon.lapointe

#Synced-CL 51274443

diff --git a/Engine/Plugins/PCG/Source/PCG/Private/Graph/DataOverride/PCGDataOverrideHelpers.cpp b/Engine/Plugins/PCG/Source/PCG/Private/Graph/DataOverride/PCGDataOverrideHelpers.cpp
--- a/Engine/Plugins/PCG/Source/PCG/Private/Graph/DataOverride/PCGDataOverrideHelpers.cpp
+++ b/Engine/Plugins/PCG/Source/PCG/Private/Graph/DataOverride/PCGDataOverrideHelpers.cpp
@@ -137,8 +137,11 @@ namespace PCG::DataOverride
 			return false;
 		}
 
+		// @todo_pcg: Investigate the need to push the node to the stack
 		// First check if the node has data override deltas.
-		const FXxHash64 NodeKey = Helpers::ComputeDataHash(*Stack);
+		FPCGStack NodeStack = *Stack;
+		NodeStack.PushFrame(Context->Node);
+		const FXxHash64 NodeKey = Helpers::ComputeDataHash(NodeStack);
 		const FPCGSourceDataStorageKey NodeStorageKey(Constants::DefaultOverrideLabel, NodeKey.Hash);
 		return DataContainer->Get<FPCGDeltaCollection>(NodeStorageKey).IsValid();
 	}
@@ -152,7 +155,9 @@ namespace PCG::DataOverride
 			return false;
 		}
 
-		const FXxHash64 NodeKey = Helpers::ComputeDataHash(*Stack);
+		FPCGStack NodeStack = *Stack;
+		NodeStack.PushFrame(Context->Node);
+		const FXxHash64 NodeKey = Helpers::ComputeDataHash(NodeStack);
 		const FXxHash64 PinKey = Helpers::ComputeDataHashWithPin(NodeKey, InTaggedData.Pin);
 		const FPCGSourceDataStorageKey PinStorageKey(Constants::DefaultOverrideLabel, PinKey.Hash);
 		return DataContainer->Get<FPCGDeltaCollection>(PinStorageKey).IsValid();
@@ -191,7 +196,9 @@ namespace PCG::DataOverride
 			return;
 		}
 
-		const FXxHash64 NodeKey = Helpers::ComputeDataHash(*Stack);
+		FPCGStack NodeStack = *Stack;
+		NodeStack.PushFrame(Context->Node);
+		const FXxHash64 NodeKey = Helpers::ComputeDataHash(NodeStack);
 		for (FPCGTaggedData& TaggedData : OutTaggedData)
 		{
 			// @todo_pcg: This should be part of a registry. For now, default to points, but should be element
diff --git a/Engine/Plugins/PCG/Source/PCG/Private/PCGGraphExecutionInspection.cpp b/Engine/Plugins/PCG/Source/PCG/Private/PCGGraphExecutionInspection.cpp
--- a/Engine/Plugins/PCG/Source/PCG/Private/PCGGraphExecutionInspection.cpp
+++ b/Engine/Plugins/PCG/Source/PCG/Private/PCGGraphExecutionInspection.cpp
@@ -137,7 +137,7 @@ void FPCGGraphExecutionInspection::StoreInspectionData(const FPCGStack* InStack,
 		}
 	}
 
-	if (IsInspecting())
+	if (IsInspecting() || (InNode->GetSettingsInterface() && InNode->GetSettingsInterface()->bViewportEditingEnabled))
 	{
 		InInputData.MarkUsage(EPCGDataUsage::ComponentInspectionData);
 		InOutputData.MarkUsage(EPCGDataUsage::ComponentInspectionData);
diff --git a/Engine/Plugins/PCG/Source/PCG/Public/PCGSettings.h b/Engine/Plugins/PCG/Source/PCG/Public/PCGSettings.h
--- a/Engine/Plugins/PCG/Source/PCG/Public/PCGSettings.h
+++ b/Engine/Plugins/PCG/Source/PCG/Public/PCGSettings.h
@@ -218,6 +218,9 @@ public:
 	bool bDebug = false;
 
 #if WITH_EDITORONLY_DATA
+	UPROPERTY(Transient, EditAnywhere, BlueprintReadWrite, Category = Debug, meta = (EditCondition = bDisplayDebuggingProperties, EditConditionHides, HideEditConditionToggle))
+	bool bViewportEditingEnabled = false;
+
 	UPROPERTY(EditAnywhere, BlueprintReadWrite, Category = Debug, meta = (ShowOnlyInnerProperties, PCGNoHash))
 	FPCGDebugVisualizationSettings DebugSettings;
 
diff --git a/Engine/Plugins/PCG/Source/PCGEditor/PCGEditor.Build.cs b/Engine/Plugins/PCG/Source/PCGEditor/PCGEditor.Build.cs
--- a/Engine/Plugins/PCG/Source/PCGEditor/PCGEditor.Build.cs
+++ b/Engine/Plugins/PCG/Source/PCGEditor/PCGEditor.Build.cs
@@ -40,6 +40,7 @@ namespace UnrealBuildTool.Rules
 					"AssetRegistry",
 					"AssetTools",
 					"BlueprintGraph",
+					"ComponentVisualizers",
 					"ContentBrowser",
 					"ContentBrowserData",
 					"DesktopWidgets",
diff --git a/Engine/Plugins/PCG/Source/PCGEditor/Private/PCGComponentVisualizer.cpp b/Engine/Plugins/PCG/Source/PCGEditor/Private/PCGComponentVisualizer.cpp
new file mode 100644
--- /dev/null
+++ b/Engine/Plugins/PCG/Source/PCGEditor/Private/PCGComponentVisualizer.cpp
@@ -0,0 +1,539 @@
+// Copyright Epic Games, Inc. All Rights Reserved.
+
+#include "PCGComponentVisualizer.h"
+
+#include "PCGComponent.h"
+#include "PCGEditorModule.h"
+#include "Graph/DataOverride/PCGDataOverrideHelpers.h"
+#include "Graph/DataOverride/PCGDataOverridePoints.h"
+
+#include "CanvasItem.h"
+#include "CanvasTypes.h"
+#include "EditorModeManager.h"
+#include "EditorModes.h"
+#include "EditorViewportClient.h"
+#include "PCGGraph.h"
+#include "PrimitiveDrawInterface.h"
+#include "BaseGizmos/TransformGizmoUtil.h"
+#include "Data/PCGBasePointData.h"
+#include "Editor/UnrealEdEngine.h"
+#include "Materials/Material.h"
+#include "Materials/MaterialRenderProxy.h"
+#include "Tools/EdModeInteractiveToolsContext.h"
+
+IMPLEMENT_HIT_PROXY(HPCGComponentVisualizerHitProxy, HComponentVisProxy);
+IMPLEMENT_HIT_PROXY(HPCGManualEditHitProxy, HComponentVisProxy);
+
+#define LOCTEXT_NAMESPACE "PCGComponentVisualizer"
+
+namespace PCG::ComponentVisualizer
+{
+	namespace Defaults
+	{
+		const FVector ManualEditCubeExtent(100, 100, 50);
+		constexpr FLinearColor ManualEditDeltaCubeColor(1.0f, 0.0f, 0.5f, 0.4f);
+		constexpr FLinearColor ManualEditDefaultCubeColor(0.0f, 1.0f, 0.5f, 0.4f);
+		constexpr FLinearColor ManualEditSelectedCubeColor(1.0f, 1.0f, 0.0f, 0.4f);
+
+		constexpr const TCHAR* ManualEditOnText = TEXT("Manual Edit Elements : ON");
+		constexpr const TCHAR* ManualEditOffText = TEXT("Manual Edit Elements : OFF");
+	}
+
+	namespace Helpers
+	{
+		FPCGDeltaKey GetDeltaKeyFromCollection(const FPCGDeltaCollection& Collection, const FTransform& Transform)
+		{
+			FPCGDeltaKey LocationKey = PCG::DataOverride::Keys::FPCGLocationDeltaKey(Transform.GetLocation());
+			FPCGDeltaKey FoundDeltaKey = {};
+			Collection.ForEachDelta([&FoundDeltaKey, LocationKey](const FPCGDeltaKey& DeltaKey, const TInstancedStruct<FPCGDeltaBase>& DeltaValue)
+			{
+				if (const FPCGPointEditDelta* EditDelta = DeltaValue.GetPtr<FPCGPointEditDelta>())
+				{
+					if (LocationKey == PCG::DataOverride::Keys::FPCGLocationDeltaKey(EditDelta->TransformOverride.GetLocation()))
+					{
+						FoundDeltaKey = DeltaKey;
+						return false;
+					}
+				}
+				return true;
+			});
+
+			return FoundDeltaKey;
+		}
+
+		uint64 ComputePinKey(const FPCGStack& PinStack)
+		{
+			// Right now, the pin key for a given stack should be computed with a node key first.
+			// @todo_pcg: This will need to change if/when the Data Override API changes.
+			check(PinStack.GetStackFrames().Num() > 1);
+			FPCGStack NodeStack = PinStack;
+			NodeStack.PopFrame();
+			const FXxHash64 NodeKey = DataOverride::Helpers::ComputeDataHash(NodeStack);
+			const FPCGStackFrame& PinFrame = PinStack.GetStackFrames().Last();
+			const UPCGPin* TopPin = PinFrame.GetObject_AnyThread<UPCGPin>();
+			check(TopPin);
+			return DataOverride::Helpers::ComputeDataHashWithPin(NodeKey, TopPin->Properties.Label).Hash;
+		}
+	}
+} // PCG::ComponentVisualizer
+
+HPCGComponentVisualizerHitProxy::HPCGComponentVisualizerHitProxy(const UActorComponent* InComponent, const FPCGStack& InPinStack, int InDataIndex, int InPointIndex)
+	: HComponentVisProxy(InComponent)
+	, PinStack(InPinStack)
+	, DataIndex(InDataIndex)
+	, PointIndex(InPointIndex)
+{}
+
+HPCGManualEditHitProxy::HPCGManualEditHitProxy(const UActorComponent* InComponent)
+	: HComponentVisProxy(InComponent)
+{}
+
+/** Define commands for the PCG component visualizer */
+class FPCGComponentVisualizerCommands : public TCommands<FPCGComponentVisualizerCommands>
+{
+public:
+	FPCGComponentVisualizerCommands() : TCommands<FPCGComponentVisualizerCommands>
+		(
+			"PCGComponentVisualizer",                                      // Context name for fast lookup
+			LOCTEXT("PCGComponentVisualizer", "PCG Component Visualizer"), // Localized context name for displaying
+			NAME_None,                                                     // Parent
+			FAppStyle::GetAppStyleSetName()
+			)
+	{}
+
+	virtual void RegisterCommands() override
+	{
+		UI_COMMAND(DeleteAllDeltas, "Remove All Deltas", "TODO.", EUserInterfaceActionType::Button, FInputChord());
+		UI_COMMAND(DeleteSelectedDeltas, "Remove Selected Deltas", "TODO.", EUserInterfaceActionType::Button, FInputChord());
+	}
+
+	TSharedPtr<FUICommandInfo> DeleteAllDeltas;
+	TSharedPtr<FUICommandInfo> DeleteSelectedDeltas;
+};
+
+FPCGComponentVisualizer::FPCGComponentVisualizer()
+{
+	FPCGComponentVisualizerCommands::Register();
+
+	ComponentVisualizerActions = MakeShareable(new FUICommandList);
+}
+
+
+FPCGComponentVisualizer::~FPCGComponentVisualizer()
+{
+	ResetGizmo();
+}
+
+void FPCGComponentVisualizer::OnRegister()
+{
+	FComponentVisualizer::OnRegister();
+
+	const FPCGComponentVisualizerCommands& Commands = FPCGComponentVisualizerCommands::Get();
+
+	ComponentVisualizerActions->MapAction(
+		Commands.DeleteAllDeltas,
+		FExecuteAction::CreateSP(this, &FPCGComponentVisualizer::RemoveAllDeltas));
+	ComponentVisualizerActions->MapAction(
+		Commands.DeleteSelectedDeltas,
+		FExecuteAction::CreateSP(this, &FPCGComponentVisualizer::RemoveSelectedDeltas),
+		FCanExecuteAction::CreateSP(this, &FPCGComponentVisualizer::HasSelection));
+}
+
+void FPCGComponentVisualizer::DrawVisualization(const UActorComponent* Component, const FSceneView* View, FPrimitiveDrawInterface* PDI)
+{
+	using namespace PCG::ComponentVisualizer;
+	using namespace PCG::ComponentVisualizer::Defaults;
+
+	if (const UPCGComponent* PCGComp = Cast<const UPCGComponent>(Component))
+	{
+		if (bManualEditMode)
+		{
+			const FPCGGraphExecutionInspection& Inspection = PCGComp->GetExecutionState().GetInspection();
+
+			const FPCGSourceDataContainer* DataContainer = PCGComp->GetExecutionState().GetSourceDataContainer();
+			check(DataContainer);
+
+			// Not ideal to const_cast here, but we need to modify the component.
+			LastPCGComponent = const_cast<UPCGComponent*>(PCGComp);
+
+			TMap<TObjectKey<const UPCGNode>, TSet<FPCGGraphExecutionInspection::FNodeExecutedNotificationData>> NodeStacks = Inspection.GetExecutedNodeStacks();
+
+			// Find nodes with manual edit.
+			for (const auto& [ObjectKey, StackSet] : NodeStacks)
+			{
+				const UPCGNode* Node = ObjectKey.ResolveObjectPtr();
+				if (Node && Node->GetSettingsInterface() && Node->GetSettingsInterface()->bViewportEditingEnabled && Node->GetOutputPins().Num() > 0)
+				{
+					for (const FPCGGraphExecutionInspection::FNodeExecutedNotificationData& Stack : StackSet)
+					{
+						const UPCGPin* PinUsed = Node->GetOutputPins()[0];
+
+						FPCGStack PinStack = Stack.Stack;
+						PinStack.PushFrame(Node);
+						PinStack.PushFrame(PinUsed);
+						TSharedPtr<FPCGInspectionData> InspectionData = Inspection.GetInspectionDataPtr(PinStack);
+
+						// Compute the node key from the execution stack
+						FPCGSourceDataStorageKey StorageKey(PCG::DataOverride::Constants::DefaultOverrideLabel, Helpers::ComputePinKey(PinStack));
+
+						// Retrieve the delta collection
+						FConstSharedStruct SharedStruct = DataContainer->Get<FPCGDeltaCollection>(StorageKey);
+						const FPCGDeltaCollection* Collection = SharedStruct.IsValid()
+																	? SharedStruct.GetPtr<const FPCGDeltaCollection>()
+																	: nullptr;
+
+						if (InspectionData && InspectionData->Data)
+						{
+							const TArray<FPCGTaggedData>& DataArray = InspectionData->Data->GetAllInputs();
+							for (int DataIndex = 0; DataIndex < DataArray.Num(); DataIndex++)
+							{
+								const FPCGTaggedData& Input = DataArray[DataIndex];
+								if (const UPCGBasePointData* BasePointData = Cast<UPCGBasePointData>(Input.Data))
+								{
+									TConstPCGValueRange<FTransform> TransformRange = BasePointData->GetConstTransformValueRange();
+									for (int Index = 0; Index < BasePointData->GetNumPoints(); Index++)
+									{
+										const FTransform& Transform = TransformRange[Index];
+										bool bIsDelta = false;
+										bool bIsSelected = false;
+										if (Collection)
+										{
+											FPCGDeltaKey Key = Helpers::GetDeltaKeyFromCollection(*Collection, Transform);
+											bIsDelta = Key.Hash != 0;
+											bIsSelected = bIsDelta && EditTarget && EditTarget->DeltaKey == Key && EditTarget->Stack == PinStack;
+										}
+
+										FLinearColor ProxyColor = bIsSelected ? ManualEditSelectedCubeColor : bIsDelta ? ManualEditDeltaCubeColor : ManualEditDefaultCubeColor;
+
+										PDI->SetHitProxy(new HPCGComponentVisualizerHitProxy(Component, PinStack, DataIndex, Index));
+										DrawWireBox(PDI, Transform.ToMatrixWithScale(), FBox(-ManualEditCubeExtent, ManualEditCubeExtent), ProxyColor, SDPG_World);
+										FDynamicColoredMaterialRenderProxy* ColoredMaterialRenderProxy = new FDynamicColoredMaterialRenderProxy(GEngine->DebugEditorMaterial->GetRenderProxy(), ProxyColor);
+										PDI->RegisterDynamicResource(ColoredMaterialRenderProxy);
+										DrawBox(PDI, Transform.ToMatrixWithScale(), ManualEditCubeExtent, ColoredMaterialRenderProxy, SDPG_World);
+										PDI->SetHitProxy(nullptr);
+									}
+								}
+							}
+						}
+					}
+				}
+			}
+		}
+		// while dragging, selection is not visited with the previous loop. So we draw it here so we can have visual feedback on dragging.
+		// No need to add proxy since
+		if (EditTarget)
+		{
+			DrawWireBox(PDI, EditTarget->Transform.ToMatrixWithScale(), FBox(-ManualEditCubeExtent, ManualEditCubeExtent), ManualEditSelectedCubeColor, SDPG_World);
+		}
+	}
+	FComponentVisualizer::DrawVisualization(Component, View, PDI);
+}
+
+void FPCGComponentVisualizer::DrawVisualizationHUD(const UActorComponent* Component, const FViewport* Viewport, const FSceneView* View,
+												   FCanvas* Canvas)
+{
+	using namespace PCG::ComponentVisualizer;
+
+	if (GEngine)
+	{
+		Canvas->SetHitProxy(new HPCGManualEditHitProxy(Component));
+		const TCHAR* ManualEditText = bManualEditMode ? Defaults::ManualEditOnText : Defaults::ManualEditOffText;
+		int32 XL;
+		int32 YL;
+		const UFont* Font = GEngine->GetLargeFont();
+		StringSize(Font, XL, YL, ManualEditText);
+
+		// @todo_pcg: There seems to be a bug in the canvas hit proxy. it only works at position 0,0.
+		constexpr float DrawPositionX = 0;
+		constexpr float DrawPositionY = 0;
+		// const FIntRect CanvasRect = Canvas->GetViewRect();
+		// const float DrawPositionX = FMath::FloorToFloat(static_cast<float>(CanvasRect.Min.X) + static_cast<float>(CanvasRect.Width() - XL) * 0.5f);
+		// const float DrawPositionY = static_cast<float>(CanvasRect.Min.Y) + 10;
+		FVector2D DrawPosition{DrawPositionX, DrawPositionY};
+
+		FCanvasTileItem Background(DrawPosition, FVector2D{static_cast<float>(XL), static_cast<float>(YL)}, FLinearColor::Yellow);
+		Canvas->DrawItem(Background, DrawPosition);
+
+		FCanvasTextItem TextItem(FVector2D{DrawPositionX, DrawPositionY}, FText::FromString(ManualEditText), Font, FLinearColor::White);
+		TextItem.EnableShadow(FLinearColor::Black);
+		Canvas->DrawItem(TextItem, DrawPosition);
+		Canvas->SetHitProxy(nullptr);
+	}
+
+	FComponentVisualizer::DrawVisualizationHUD(Component, Viewport, View, Canvas);
+}
+
+bool FPCGComponentVisualizer::VisProxyHandleClick(FEditorViewportClient* InViewportClient, HComponentVisProxy* VisProxy, const FViewportClick& Click)
+{
+	using namespace PCG::ComponentVisualizer;
+
+	// We need to register the gizmo tool with the current mode tool so the gizmo will work.
+	//  Has to be done here since we need an interactive tool context.
+	UE::TransformGizmoUtil::RegisterTransformGizmoContextObject(InViewportClient->GetModeTools()->GetInteractiveToolsContext());
+
+	if (HitProxyCast<HPCGManualEditHitProxy>(VisProxy))
+	{
+		bManualEditMode = !bManualEditMode;
+		return true;
+	}
+	else if (HPCGComponentVisualizerHitProxy* HitProxy = HitProxyCast<HPCGComponentVisualizerHitProxy>(VisProxy))
+	{
+		const UPCGComponent* PCGComponent = Cast<UPCGComponent>(HitProxy->Component);
+		const FPCGGraphExecutionInspection& Inspection = PCGComponent->GetExecutionState().GetInspection();
+		TSharedPtr<FPCGInspectionData> InspectionData = Inspection.GetInspectionDataPtr(HitProxy->PinStack);
+		if (!InspectionData || !InspectionData->Data)
+		{
+			return false;
+		}
+
+		const TArray<FPCGTaggedData>& DataArray = InspectionData->Data->GetAllInputs();
+		if (!ensure(DataArray.IsValidIndex(HitProxy->DataIndex)))
+		{
+			return false;
+		}
+
+		const FPCGTaggedData& Input = DataArray[HitProxy->DataIndex];
+		if (const UPCGBasePointData* BasePointData = Cast<UPCGBasePointData>(Input.Data))
+		{
+			if (UInteractiveGizmoManager* GizmoManager = InViewportClient->GetModeTools()->GetInteractiveToolsContext()->GizmoManager)
+			{
+				// Not ideal to const_cast here, but we need to modify the PCGComp.
+				UPCGComponent* MutablePCGComponent = const_cast<UPCGComponent*>(PCGComponent);
+
+				FPCGSourceDataContainer* DataContainer = MutablePCGComponent->GetExecutionState().GetSourceDataContainer();
+				check(DataContainer);
+
+				FPCGSourceDataStorageKey StorageKey(PCG::DataOverride::Constants::DefaultOverrideLabel, Helpers::ComputePinKey(HitProxy->PinStack));
+
+				FPCGPointTransform::ConstValueRange TransformRange = BasePointData->GetConstTransformValueRange();
+				if (!ensure(TransformRange.IsValidIndex(HitProxy->PointIndex)))
+				{
+					return false;
+				}
+
+				const FTransform& CurTransform = TransformRange[HitProxy->PointIndex];
+				FPCGDeltaCollection* Collection = nullptr;
+				FPCGDeltaKey DeltaKey = {};
+				bool bWasInDelta = false;
+
+				// Check if the point is already a product of a delta.
+				{
+					FSharedStruct SharedStruct = DataContainer->GetMutable<FPCGDeltaCollection>(StorageKey);
+					Collection = SharedStruct.IsValid()
+									 ? SharedStruct.GetPtr<FPCGDeltaCollection>()
+									 : nullptr;
+					if (Collection)
+					{
+						DeltaKey = Helpers::GetDeltaKeyFromCollection(*Collection, CurTransform);
+						bWasInDelta = (DeltaKey.Hash != 0);
+					}
+				}
+
+				// It was not part of a delta...
+				if (!bWasInDelta)
+				{
+					// Create collection if there was not already one.
+					if (!Collection)
+					{
+						// Build a delta collection with spatial key settings
+						FPCGDeltaCollection NewCollection;
+						NewCollection.Settings.KeyPolicy = EPCGDataOverrideKeyPolicy::Spatial;
+						NewCollection.Settings.KeyTarget = EPCGSpatialKeyTarget::Position;
+						NewCollection.Settings.SpatialTolerance = 10.0;
+
+						DataContainer->Store(StorageKey, NewCollection);
+
+						FSharedStruct SharedStruct = DataContainer->GetMutable<FPCGDeltaCollection>(StorageKey);
+						Collection = SharedStruct.IsValid()
+										 ? SharedStruct.GetPtr<FPCGDeltaCollection>()
+										 : nullptr;
+					}
+
+					// Add current point in the delta, with the same transform as the point (it will be modified by the user later)
+					check(Collection);
+					FPCGPointEditDelta EditDelta;
+					EditDelta.OriginalTransform = CurTransform;
+					EditDelta.TransformOverride = CurTransform;
+					EditDelta.Bounds = PCGComponent->GetTotalBounds();
+
+					// Key the delta by spatial position and delta type label
+					DeltaKey = PCG::DataOverride::Keys::FPCGLocationDeltaKey(
+						CurTransform.GetLocation(),
+						Collection->Settings.SpatialTolerance,
+						FPCGPointEditDelta::GetDeltaNameStatic());
+
+					Collection->Add(DeltaKey, TInstancedStruct<FPCGPointEditDelta>::Make(EditDelta));
+				}
+
+				//Now that the point is part of the delta, set up the current selection
+				EditTarget.Emplace(MutablePCGComponent, HitProxy->PinStack, StorageKey, DeltaKey, bWasInDelta, CurTransform);
+
+				ResetGizmo();
+				ETransformGizmoSubElements TransformElements = ETransformGizmoSubElements::FullTranslateRotateScale;
+				Gizmo = UE::TransformGizmoUtil::CreateCustomTransformGizmo(GizmoManager, TransformElements, this, FString());
+				UTransformProxy* Proxy = NewObject<UTransformProxy>(Gizmo);
+				Proxy->SetTransform(CurTransform);
+				Gizmo->SetActiveTarget(Proxy, GizmoManager);
+				Gizmo->ReinitializeGizmoTransform(CurTransform);
+				Gizmo->SetVisibility(true);
+
+
+				Proxy->OnTransformChanged.AddLambda([this](UTransformProxy* InProxy, const FTransform& NewTransform)
+				{
+					SetSelectedTransform(NewTransform);
+				});
+
+				return true;
+			}
+		}
+	}
+	return false;
+}
+
+void FPCGComponentVisualizer::ResetGizmo()
+{
+	if (Gizmo)
+	{
+		Gizmo->SetVisibility(false);
+		Gizmo->ClearActiveTarget();
+
+		if (UInteractiveGizmoManager* GizmoManager = Gizmo->GetGizmoManager())
+		{
+			GizmoManager->DestroyGizmo(Gizmo);
+		}
+		Gizmo = nullptr;
+	}
+}
+
+void FPCGComponentVisualizer::EndEditing()
+{
+	EndCurrentElementManualEditing();
+
+	bManualEditMode = false;
+	LastPCGComponent.Reset();
+	FComponentVisualizer::EndEditing();
+}
+
+bool FPCGComponentVisualizer::GetWidgetLocation(const FEditorViewportClient* ViewportClient, FVector& OutLocation) const
+{
+	return FComponentVisualizer::GetWidgetLocation(ViewportClient, OutLocation);
+}
+
+bool FPCGComponentVisualizer::GetCustomInputCoordinateSystem(const FEditorViewportClient* ViewportClient, FMatrix& OutMatrix) const
+{
+	return FComponentVisualizer::GetCustomInputCoordinateSystem(ViewportClient, OutMatrix);
+}
+
+bool FPCGComponentVisualizer::HandleInputDelta(FEditorViewportClient* ViewportClient, FViewport* Viewport, FVector& DeltaTranslate,
+											   FRotator& DeltaRotate, FVector& DeltaScale)
+{
+	return FComponentVisualizer::HandleInputDelta(ViewportClient, Viewport, DeltaTranslate, DeltaRotate, DeltaScale);
+}
+
+bool FPCGComponentVisualizer::HandleInputKey(FEditorViewportClient* ViewportClient, FViewport* Viewport, FKey Key, EInputEvent Event)
+{
+	return FComponentVisualizer::HandleInputKey(ViewportClient, Viewport, Key, Event);
+}
+
+TSharedPtr<SWidget> FPCGComponentVisualizer::GenerateContextMenu() const
+{
+	UE_LOG(LogPCGEditor, Warning, TEXT("FPCGComponentVisualizer::GenerateContextMenu"));
+	FMenuBuilder MenuBuilder(/*bInShouldCloseWindowAfterMenuSelection=*/true, ComponentVisualizerActions);
+
+	MenuBuilder.AddMenuEntry(FPCGComponentVisualizerCommands::Get().DeleteAllDeltas);
+	MenuBuilder.AddMenuEntry(FPCGComponentVisualizerCommands::Get().DeleteSelectedDeltas);
+
+	TSharedPtr<SWidget> MenuWidget = MenuBuilder.MakeWidget();
+	return MenuWidget;
+}
+
+void FPCGComponentVisualizer::SetSelectedTransform(const FTransform& Transform)
+{
+	if (!LastPCGComponent.IsValid())
+	{
+		return;
+	}
+
+	if (FEditState* Selection = EditTarget.GetPtrOrNull(); Selection && EditTarget->PCGComponent.IsValid())
+	{
+		Selection->Transform = Transform;
+
+		FPCGSourceDataContainer* DataContainer = EditTarget->PCGComponent->GetExecutionState().GetSourceDataContainer();
+
+		FSharedStruct SharedStruct = DataContainer->GetMutable<FPCGDeltaCollection>(EditTarget->StorageKey);
+		FPCGDeltaCollection* Collection = SharedStruct.IsValid()
+											  ? SharedStruct.GetPtr<FPCGDeltaCollection>()
+											  : nullptr;
+
+		TInstancedStruct<FPCGDeltaBase>* DeltaInstancedStruct = Collection ? Collection->Find(EditTarget->DeltaKey) : nullptr;
+
+		if (FPCGPointEditDelta* EditDelta = DeltaInstancedStruct ? DeltaInstancedStruct->GetMutablePtr<FPCGPointEditDelta>() : nullptr)
+		{
+			EditDelta->TransformOverride = Transform;
+			EditTarget->bModified = true;
+			DataContainer->MarkDirty();
+
+			// TODO: This is needed here but why???
+			LastPCGComponent->Modify();
+			LastPCGComponent->GetGraphInstance()->GetMutablePCGGraph()->ForceNotificationForEditor(EPCGChangeType::ExternalModification);
+		}
+	}
+}
+
+void FPCGComponentVisualizer::RemoveAllDeltas()
+{
+	if (LastPCGComponent.IsValid())
+	{
+		FPCGSourceDataContainer* DataContainer = LastPCGComponent->GetExecutionState().GetSourceDataContainer();
+		DataContainer->Empty();
+		DataContainer->MarkDirty();
+
+		// TODO: This is needed here but why???
+		LastPCGComponent->Modify();
+		LastPCGComponent->GetGraphInstance()->GetMutablePCGGraph()->ForceNotificationForEditor(EPCGChangeType::ExternalModification);
+	}
+	EndCurrentElementManualEditing();
+}
+
+void FPCGComponentVisualizer::RemoveSelectedDeltas()
+{
+	if (LastPCGComponent.IsValid() && EditTarget.IsSet())
+	{
+		FPCGSourceDataContainer* DataContainer = EditTarget->PCGComponent->GetExecutionState().GetSourceDataContainer();
+
+		FSharedStruct SharedStruct = DataContainer->GetMutable<FPCGDeltaCollection>(EditTarget->StorageKey);
+		if (FPCGDeltaCollection* Collection = SharedStruct.IsValid() ? SharedStruct.GetPtr<FPCGDeltaCollection>() : nullptr)
+		{
+			Collection->Remove(EditTarget->DeltaKey);
+		}
+
+		DataContainer->MarkDirty();
+		EditTarget.Reset();
+		EndCurrentElementManualEditing();
+
+		// TODO: This is needed here but why???
+		LastPCGComponent->Modify();
+		LastPCGComponent->GetGraphInstance()->GetMutablePCGGraph()->ForceNotificationForEditor(EPCGChangeType::ExternalModification);
+	}
+}
+
+bool FPCGComponentVisualizer::HasSelection() const
+{
+	return LastPCGComponent.IsValid() && EditTarget.IsSet();
+}
+
+void FPCGComponentVisualizer::EndCurrentElementManualEditing()
+{
+	if (EditTarget.IsSet())
+	{
+		if (!EditTarget->bWasInDelta && !EditTarget->bModified)
+		{
+			RemoveSelectedDeltas();
+		}
+	}
+
+	EditTarget.Reset();
+	ResetGizmo();
+}
+
+#undef LOCTEXT_NAMESPACE
\ No newline at end of file
diff --git a/Engine/Plugins/PCG/Source/PCGEditor/Private/PCGEditorModule.cpp b/Engine/Plugins/PCG/Source/PCGEditor/Private/PCGEditorModule.cpp
--- a/Engine/Plugins/PCG/Source/PCGEditor/Private/PCGEditorModule.cpp
+++ b/Engine/Plugins/PCG/Source/PCGEditor/Private/PCGEditorModule.cpp
@@ -2,7 +2,9 @@
 
 #include "PCGEditorModule.h"
 
+#include "ComponentVisualizers.h"
 #include "PCGComponent.h"
+#include "PCGComponentVisualizer.h"
 #include "PCGEngineSettings.h"
 #include "PCGGraphExecutionStateInterface.h"
 #include "PCGModule.h"
@@ -83,9 +85,11 @@
 #include "SLevelViewport.h"
 #include "ToolMenus.h"
 #include "ToolMenuEntry.h"
+#include "UnrealEdGlobals.h"
 #include "Data/DataView/PCGDataViewData.h"
 #include "DataVisualizations/PCGDataViewVisualization.h"
 #include "Details/PCGRaycastFilterRuleCollectionDetails.h"
+#include "Editor/UnrealEdEngine.h"
 #include "Editor/SceneOutliner/Public/ISceneOutliner.h"
 #include "EditorMode/PCGEdMode.h"
 #include "EditorMode/Tools/PCGInteractiveToolSettings.h"
@@ -154,6 +158,8 @@ void FPCGEditorModule::StartupModule()
 
 	IPCGEditorModule::SetEditorModule(this);
 
+	FCoreDelegates::OnPostEngineInit.AddRaw(this, &FPCGEditorModule::OnPostEngineInit);
+
 	// Register onto the PreExit, because we need the class/struct to be still valid to remove them from the mapping
 	FCoreDelegates::OnPreExit.AddRaw(this, &FPCGEditorModule::OnPreExit);
 }
@@ -193,6 +199,15 @@ void FPCGEditorModule::ShutdownModule()
 		GEditor->OnSceneMaterialsModifiedEvent().RemoveAll(this);
 	}
 
+	if (GUnrealEd)
+	{
+		for (const FName& Name : VisualizersToUnregisterOnShutdown)
+		{
+			GUnrealEd->UnregisterComponentVisualizer(Name);
+		}
+		VisualizersToUnregisterOnShutdown.Empty();
+	}
+
 	if (FModuleManager::Get().IsModuleLoaded("LevelEditor"))
 	{
 		FLevelEditorModule& LevelEditorModule = FModuleManager::GetModuleChecked<FLevelEditorModule>("LevelEditor");
@@ -200,6 +215,8 @@ void FPCGEditorModule::ShutdownModule()
 	}
 
 	FEditorBuildUtils::UnregisterCustomBuildType(PCGEditorModule::PCGBuildType);
+
+	FCoreDelegates::OnPostEngineInit.RemoveAll(this);
 }
 
 void FPCGEditorModule::ExtendShowFlagsMenu()
@@ -332,6 +349,13 @@ bool FPCGEditorModule::ShouldTreatViewportAsGenerationSourceInternal(const APCGW
 	return false;
 }
 
+void FPCGEditorModule::OnPostEngineInit()
+{
+	FComponentVisualizersModule& VisualisersModule = FModuleManager::LoadModuleChecked<FComponentVisualizersModule>("ComponentVisualizers");
+	VisualisersModule.RegisterComponentVisualizer(UPCGComponent::StaticClass()->GetFName(), MakeShared<FPCGComponentVisualizer>());
+	VisualizersToUnregisterOnShutdown.Add(UPCGComponent::StaticClass()->GetFName());
+}
+
 void FPCGEditorModule::OnPreExit()
 {
 	UnregisterPinColorAndIcons();
diff --git a/Engine/Plugins/PCG/Source/PCGEditor/Public/PCGComponentVisualizer.h b/Engine/Plugins/PCG/Source/PCGEditor/Public/PCGComponentVisualizer.h
new file mode 100644
--- /dev/null
+++ b/Engine/Plugins/PCG/Source/PCGEditor/Public/PCGComponentVisualizer.h
@@ -0,0 +1,89 @@
+// Copyright Epic Games, Inc. All Rights Reserved.
+
+#pragma once
+
+#include "Graph/DataOverride/PCGDataOverride.h"
+#include "Graph/PCGSourceDataContainer.h"
+#include "Graph/PCGStackContext.h"
+
+#include "ComponentVisualizer.h"
+#include "CoreMinimal.h"
+#include "Engine/EngineBaseTypes.h"
+#include "HitProxies.h"
+#include "InputCoreTypes.h"
+#include "BaseGizmos/CombinedTransformGizmo.h"
+#include "Templates/SharedPointer.h"
+
+class FUICommandList;
+
+struct HPCGComponentVisualizerHitProxy : public HComponentVisProxy
+{
+	DECLARE_HIT_PROXY(PCGEDITOR_API);
+
+	HPCGComponentVisualizerHitProxy(const UActorComponent* InComponent, const FPCGStack& InPinStack, int InDataIndex, int InPointIndex);
+
+	FPCGStack PinStack;
+	int DataIndex;
+	int PointIndex;
+};
+
+struct HPCGManualEditHitProxy : public HComponentVisProxy
+{
+	DECLARE_HIT_PROXY(PCGEDITOR_API);
+
+	HPCGManualEditHitProxy(const UActorComponent* InComponent);
+};
+
+
+class FPCGComponentVisualizer : public FComponentVisualizer
+{
+public:
+	PCGEDITOR_API FPCGComponentVisualizer();
+	PCGEDITOR_API virtual ~FPCGComponentVisualizer() override;
+
+	//~ Begin FComponentVisualizer Interface
+	PCGEDITOR_API virtual void OnRegister() override;
+	PCGEDITOR_API virtual void DrawVisualization(const UActorComponent* Component, const FSceneView* View, FPrimitiveDrawInterface* PDI) override;
+	PCGEDITOR_API virtual void DrawVisualizationHUD(const UActorComponent* Component, const FViewport* Viewport, const FSceneView* View, FCanvas* Canvas) override;
+	PCGEDITOR_API virtual bool VisProxyHandleClick(FEditorViewportClient* InViewportClient, HComponentVisProxy* VisProxy, const FViewportClick& Click) override;
+	PCGEDITOR_API virtual void EndEditing() override;
+	PCGEDITOR_API virtual bool GetWidgetLocation(const FEditorViewportClient* ViewportClient, FVector& OutLocation) const override;
+	PCGEDITOR_API virtual bool GetCustomInputCoordinateSystem(const FEditorViewportClient* ViewportClient, FMatrix& OutMatrix) const override;
+	PCGEDITOR_API virtual bool HandleInputDelta(FEditorViewportClient* ViewportClient, FViewport* Viewport, FVector& DeltaTranslate, FRotator& DeltaRotate, FVector& DeltaScale) override;
+	PCGEDITOR_API virtual bool HandleInputKey(FEditorViewportClient* ViewportClient, FViewport* Viewport, FKey Key, EInputEvent Event) override;
+	PCGEDITOR_API virtual TSharedPtr<SWidget> GenerateContextMenu() const override;
+	//~ End FComponentVisualizer Interface
+
+
+protected:
+	void SetSelectedTransform(const FTransform& Transform);
+	void RemoveAllDeltas();
+	void RemoveSelectedDeltas();
+	bool HasSelection() const;
+	void ResetGizmo();
+
+	PCGEDITOR_API void EndCurrentElementManualEditing();
+
+	/** Registered Menu commands*/
+	TSharedPtr<FUICommandList> ComponentVisualizerActions;
+
+	bool bManualEditMode = false;
+
+	TWeakObjectPtr<UPCGComponent> LastPCGComponent;
+
+	struct FEditState
+	{
+		TWeakObjectPtr<UPCGComponent> PCGComponent;
+		FPCGStack Stack;
+		FPCGSourceDataStorageKey StorageKey;
+		FPCGDeltaKey DeltaKey;
+		bool bWasInDelta;
+		FTransform Transform;
+
+		bool bModified = false;
+	};
+
+	TOptional<FEditState> EditTarget = {};
+
+	UCombinedTransformGizmo* Gizmo = nullptr;
+};
diff --git a/Engine/Plugins/PCG/Source/PCGEditor/Public/PCGEditorModule.h b/Engine/Plugins/PCG/Source/PCGEditor/Public/PCGEditorModule.h
--- a/Engine/Plugins/PCG/Source/PCGEditor/Public/PCGEditorModule.h
+++ b/Engine/Plugins/PCG/Source/PCGEditor/Public/PCGEditorModule.h
@@ -74,6 +74,7 @@ protected:
 	void OnLevelEditorCreated(TSharedPtr<ILevelEditor> InLevelEditor);
 	void RegisterOnEditorModeChange();
 
+	void OnPostEngineInit();
 	void OnPreExit();
 
 	void ExtendShowFlagsMenu();
@@ -112,4 +113,6 @@ protected:
 	FPermissionModeChanged PermissionModeChangedEvent;
 
 	TUniquePtr<FPCGLoadDataAssetChangeTracker> LoadDataAssetChangeTracker;
+
+	TArray<FName> VisualizersToUnregisterOnShutdown;
 };
