[PCG] Corrected API on PCGDataOverridePoints.h

#Synced-CL 51279030

diff --git a/Engine/Plugins/PCG/Source/PCG/Public/Graph/DataOverride/PCGDataOverridePoints.h b/Engine/Plugins/PCG/Source/PCG/Public/Graph/DataOverride/PCGDataOverridePoints.h
--- a/Engine/Plugins/PCG/Source/PCG/Public/Graph/DataOverride/PCGDataOverridePoints.h
+++ b/Engine/Plugins/PCG/Source/PCG/Public/Graph/DataOverride/PCGDataOverridePoints.h
@@ -16,15 +16,15 @@ namespace PCG::DataOverride::Points::Constants
 
 /** Base delta for point data overrides. Uses an octree bounds query to filter candidate points spatially. */
 USTRUCT()
-struct PCG_API FPCGPointDeltaBase : public FPCGDeltaBase
+struct FPCGPointDeltaBase : public FPCGDeltaBase
 {
 	GENERATED_BODY()
 
 	/** Queries the point octree with Bounds to find candidate points. */
-	virtual PCGIndexing::FPCGIndexCollection FilterCandidates(const UPCGData* InData) const override;
+	PCG_API virtual PCGIndexing::FPCGIndexCollection FilterCandidates(const UPCGData* InData) const override;
 
 	/** Computes the spatial key for each candidate and compares against this delta's key to find the matching point. */
-	virtual int32 Resolve(const UPCGData* InData, const PCGIndexing::FPCGIndexCollection& FilteredCandidates, const FPCGDeltaSettings& DeltaSettings) const override;
+	PCG_API virtual int32 Resolve(const UPCGData* InData, const PCGIndexing::FPCGIndexCollection& FilteredCandidates, const FPCGDeltaSettings& DeltaSettings) const override;
 
 	/** The original transform of the point. For use with matching the key during override. */
 	UPROPERTY()
@@ -37,13 +37,13 @@ struct PCG_API FPCGPointDeltaBase : public FPCGDeltaBase
 
 /** Delta that overrides the transform of a specific point. */
 USTRUCT()
-struct PCG_API FPCGPointEditDelta : public FPCGPointDeltaBase
+struct FPCGPointEditDelta : public FPCGPointDeltaBase
 {
 	GENERATED_BODY()
 
 	static FName GetDeltaNameStatic() { return PCG::DataOverride::Points::Constants::PointEditDeltaName; }
-	virtual FName GetDeltaName() const override;
-	virtual bool Apply(UPCGData* InData, int32 ResolvedIndex) const override;
+	PCG_API virtual FName GetDeltaName() const override;
+	PCG_API virtual bool Apply(UPCGData* InData, int32 ResolvedIndex) const override;
 
 	/** The desired transform to apply as the override. */
 	UPROPERTY()
@@ -70,13 +70,13 @@ struct PCG_API FPCGPointEditDelta : public FPCGPointDeltaBase
 
 /** Delta that marks a specific point for deletion. */
 USTRUCT()
-struct PCG_API FPCGPointDeletionDelta : public FPCGPointDeltaBase
+struct FPCGPointDeletionDelta : public FPCGPointDeltaBase
 {
 	GENERATED_BODY()
 
 	static FName GetDeltaNameStatic() { return PCG::DataOverride::Points::Constants::PointDeleteDeltaName; }
-	virtual FName GetDeltaName() const override;
-	virtual bool Apply(UPCGData* InData, int32 ResolvedIndex) const override;
+	PCG_API virtual FName GetDeltaName() const override;
+	PCG_API virtual bool Apply(UPCGData* InData, int32 ResolvedIndex) const override;
 
 #if WITH_EDITORONLY_DATA
 	/** The original element index. For debugging purposes only. */
@@ -87,14 +87,14 @@ struct PCG_API FPCGPointDeletionDelta : public FPCGPointDeltaBase
 
 /** Delta that injects a new point. */
 USTRUCT()
-struct PCG_API FPCGPointInsertionDelta : public FPCGDeltaBase
+struct FPCGPointInsertionDelta : public FPCGDeltaBase
 {
 	GENERATED_BODY()
 
 	static FName GetDeltaNameStatic() { return PCG::DataOverride::Points::Constants::PointInsertDeltaName; }
-	virtual FName GetDeltaName() const override;
-	virtual PCGIndexing::FPCGIndexCollection FilterCandidates(const UPCGData* InData) const override;
-	virtual bool Apply(UPCGData* InData, int32 ResolvedIndex) const override;
+	PCG_API virtual FName GetDeltaName() const override;
+	PCG_API virtual PCGIndexing::FPCGIndexCollection FilterCandidates(const UPCGData* InData) const override;
+	PCG_API virtual bool Apply(UPCGData* InData, int32 ResolvedIndex) const override;
 
 	/** Points to be appended to the point data. */
 	UPROPERTY()
diff --git a/Engine/Plugins/PCG/Source/PCG/Public/Graph/PCGSourceDataContainer.h b/Engine/Plugins/PCG/Source/PCG/Public/Graph/PCGSourceDataContainer.h
--- a/Engine/Plugins/PCG/Source/PCG/Public/Graph/PCGSourceDataContainer.h
+++ b/Engine/Plugins/PCG/Source/PCG/Public/Graph/PCGSourceDataContainer.h
@@ -16,14 +16,14 @@ concept CPCGSourceDataCompatibleType = requires { T::StaticStruct(); };
  * Note: The Label can be used independently if desired, with the Hash left as 0.
  */
 USTRUCT()
-struct PCG_API FPCGSourceDataStorageKey
+struct FPCGSourceDataStorageKey
 {
 	GENERATED_BODY()
 
 	FPCGSourceDataStorageKey() = default;
-	FPCGSourceDataStorageKey(FName InLabel, uint64 InHash = 0u);
+	PCG_API FPCGSourceDataStorageKey(FName InLabel, uint64 InHash = 0u);
 
-	bool operator==(const FPCGSourceDataStorageKey& Other) const = default;
+	PCG_API bool operator==(const FPCGSourceDataStorageKey& Other) const = default;
 
 	PCG_API friend uint32 GetTypeHash(const FPCGSourceDataStorageKey& Key);
 
@@ -38,7 +38,7 @@ struct PCG_API FPCGSourceDataStorageKey
 
 /** Type-erased storage value. Wraps a FSharedStruct payload for runtime verification and key collision avoidance. */
 USTRUCT()
-struct PCG_API FPCGSourceDataStorageValue
+struct FPCGSourceDataStorageValue
 {
 	GENERATED_BODY()
 
@@ -58,7 +58,7 @@ struct PCG_API FPCGSourceDataStorageValue
  * NOTE: This container is READ-ONLY at runtime.
  */
 USTRUCT()
-struct PCG_API FPCGSourceDataContainer
+struct FPCGSourceDataContainer
 {
 	GENERATED_BODY()
 
@@ -66,10 +66,10 @@ struct PCG_API FPCGSourceDataContainer
 	~FPCGSourceDataContainer() = default;
 	// Required by USTRUCT
 	// NOT THREAD SAFE: These must only be called from a single thread with no concurrent readers/writers.
-	FPCGSourceDataContainer(const FPCGSourceDataContainer& Other);
-	FPCGSourceDataContainer(FPCGSourceDataContainer&& Other);
-	FPCGSourceDataContainer& operator=(const FPCGSourceDataContainer& Other);
-	FPCGSourceDataContainer& operator=(FPCGSourceDataContainer&& Other);
+	PCG_API FPCGSourceDataContainer(const FPCGSourceDataContainer& Other);
+	PCG_API FPCGSourceDataContainer(FPCGSourceDataContainer&& Other);
+	PCG_API FPCGSourceDataContainer& operator=(const FPCGSourceDataContainer& Other);
+	PCG_API FPCGSourceDataContainer& operator=(FPCGSourceDataContainer&& Other);
 
 #if WITH_EDITOR
 	/** Store data of type T within the storage container. */
@@ -77,30 +77,30 @@ struct PCG_API FPCGSourceDataContainer
 	void Store(const FPCGSourceDataStorageKey& DataKey, const T& Data);
 
 	/** Remove the data attached to a specific Data Key. */
-	bool Remove(const FPCGSourceDataStorageKey& DataKey);
+	PCG_API bool Remove(const FPCGSourceDataStorageKey& DataKey);
 
 	/** Get a mutable shared reference to stored data by type. Returns an empty FSharedStruct if key or type doesn't exist. */
 	template <CPCGSourceDataCompatibleType T>
 	FSharedStruct GetMutable(const FPCGSourceDataStorageKey& DataKey);
 
 	/** Mark the container dirty to invalidate the Crc cache. */
-	void MarkDirty();
+	PCG_API void MarkDirty();
 
 	/** Set the container to automatically dirty to invalidate the Crc cache. */
-	void SetShouldAutoDirty(bool bShouldAutoDirty);
+	PCG_API void SetShouldAutoDirty(bool bShouldAutoDirty);
 
 	/** Empty the contents of the container. */
-	void Empty();
+	PCG_API void Empty();
 #endif // WITH_EDITOR
 
 	/** Get the dirty generation value for Crc combination. */
-	uint32 GetDirtyGeneration() const;
+	PCG_API uint32 GetDirtyGeneration() const;
 
 	/** Get the number of storage values. */
-	int32 Num() const;
+	PCG_API int32 Num() const;
 
 	/** True if the storage is empty. */
-	bool IsEmpty() const;
+	PCG_API bool IsEmpty() const;
 
 	/** Get a const shared reference to stored data by type. Returns an empty FConstSharedStruct if key or type doesn't exist. */
 	template <CPCGSourceDataCompatibleType T>
