[PCG] Fix static analysis issue with FPCGDeltaCollection::Add

#Synced-CL 51279030

diff --git a/Engine/Plugins/PCG/Source/PCG/Public/Graph/DataOverride/PCGDataOverride.h b/Engine/Plugins/PCG/Source/PCG/Public/Graph/DataOverride/PCGDataOverride.h
--- a/Engine/Plugins/PCG/Source/PCG/Public/Graph/DataOverride/PCGDataOverride.h
+++ b/Engine/Plugins/PCG/Source/PCG/Public/Graph/DataOverride/PCGDataOverride.h
@@ -148,7 +148,7 @@ struct FPCGDeltaCollection
 	// For now, this will be editor only for storage mutation.
 #if WITH_EDITOR
 	/** Add a new delta entry to the collection, as a TInstancedStruct<FPCGDeltaBase>. */
-	PCG_API TInstancedStruct<FPCGDeltaBase>& Add(const FPCGDeltaKey& InKey, TInstancedStruct<FPCGDeltaBase>&& InDelta);
+	[[maybe_unused]] PCG_API TInstancedStruct<FPCGDeltaBase>& Add(const FPCGDeltaKey& InKey, TInstancedStruct<FPCGDeltaBase>&& InDelta);
 
 	/** Remove a delta entry by key. */
 	PCG_API bool Remove(const FPCGDeltaKey& InKey);
