[PCG] One more API issue with FPCGSourceDataStorageKey

#Synced-CL 51279030

diff --git a/Engine/Plugins/PCG/Source/PCG/Public/Graph/PCGSourceDataContainer.h b/Engine/Plugins/PCG/Source/PCG/Public/Graph/PCGSourceDataContainer.h
--- a/Engine/Plugins/PCG/Source/PCG/Public/Graph/PCGSourceDataContainer.h
+++ b/Engine/Plugins/PCG/Source/PCG/Public/Graph/PCGSourceDataContainer.h
@@ -23,7 +23,7 @@ struct FPCGSourceDataStorageKey
 	FPCGSourceDataStorageKey() = default;
 	PCG_API FPCGSourceDataStorageKey(FName InLabel, uint64 InHash = 0u);
 
-	PCG_API bool operator==(const FPCGSourceDataStorageKey& Other) const = default;
+	bool operator==(const FPCGSourceDataStorageKey& Other) const = default;
 
 	PCG_API friend uint32 GetTypeHash(const FPCGSourceDataStorageKey& Key);
 
