[PCG] Temporary build health fix for PCGComponentVisualizer.cpp

#Synced-CL 51279030

diff --git a/Engine/Plugins/PCG/Source/PCGEditor/Private/PCGComponentVisualizer.cpp b/Engine/Plugins/PCG/Source/PCGEditor/Private/PCGComponentVisualizer.cpp
--- a/Engine/Plugins/PCG/Source/PCGEditor/Private/PCGComponentVisualizer.cpp
+++ b/Engine/Plugins/PCG/Source/PCGEditor/Private/PCGComponentVisualizer.cpp
@@ -362,7 +362,8 @@ bool FPCGComponentVisualizer::VisProxyHandleClick(FEditorViewportClient* InViewp
 						Collection->Settings.SpatialTolerance,
 						FPCGPointEditDelta::GetDeltaNameStatic());
 
-					Collection->Add(DeltaKey, TInstancedStruct<FPCGPointEditDelta>::Make(EditDelta));
+					// @todo_pcg: This temporarily bypasses static analysis warning in PVS Studio. API to be revisited.
+					(void)Collection->Add(DeltaKey, TInstancedStruct<FPCGPointEditDelta>::Make(EditDelta));
 				}
 
 				//Now that the point is part of the delta, set up the current selection
