Use-After-Free fix: pass a weak pointer to DataBinding instead of a potentially stale DataBinding.

#rnx

#Synced-CL 51102399

diff --git a/Engine/Plugins/PCG/Source/PCG/Private/Compute/Elements/PCGComputeGraphElement.cpp b/Engine/Plugins/PCG/Source/PCG/Private/Compute/Elements/PCGComputeGraphElement.cpp
--- a/Engine/Plugins/PCG/Source/PCG/Private/Compute/Elements/PCGComputeGraphElement.cpp
+++ b/Engine/Plugins/PCG/Source/PCG/Private/Compute/Elements/PCGComputeGraphElement.cpp
@@ -51,16 +51,17 @@ FPCGComputeGraphContext::~FPCGComputeGraphContext()
 {
 	const FPCGStack* StackPtr = GetStack();
 	TWeakObjectPtr<const UPCGGraph> GraphWeak = StackPtr ? StackPtr->GetNearestNonInlinedGraphForCurrentFrame() : nullptr;
-
+	TWeakObjectPtr<UPCGDataBinding> DataBindingWeak = MakeWeakObjectPtr(DataBinding);
+	
 	// Can be here on a worker thread, after unpinning the context in one of the async tasks.
 	// Will execute immediately if on game thread, otherwise on next game thread frame.
-	PCGHelpers::ExecuteOnGameThread(TEXT("~FPCGComputeGraphContext"), [DataBinding=DataBinding, GraphInstance=ComputeGraphInstance, bInstanceInitialized=bComputeGraphInstanceInitialized, GraphWeak, GridSize=GenerationGridSize, GraphIndex=ComputeGraphIndex]
+	PCGHelpers::ExecuteOnGameThread(TEXT("~FPCGComputeGraphContext"), [DataBindingWeak, GraphInstance=ComputeGraphInstance, bInstanceInitialized=bComputeGraphInstanceInitialized, GraphWeak, GridSize=GenerationGridSize, GraphIndex=ComputeGraphIndex]
 	{
 		LLM_SCOPE_BYTAG(PCG);
 
-		if (DataBinding)
+		if (DataBindingWeak.IsValid())
 		{
-			DataBinding->ReleaseTransientResources();
+			DataBindingWeak->ReleaseTransientResources();
 		}
 
 		if (GraphInstance.IsValid())
