[PCG] Actor Tracker: Fix Level Instance PIE Crash
- Skip streamed out Level Instance in game worlds (PIE)
- Avoid registering editor only callbacks in game world
- Remove redundant WITH_EDITOR

#Synced-CL 51232499

diff --git a/Engine/Plugins/PCG/Source/PCGEditor/Private/ChangeTracking/PCGActorTracker.cpp b/Engine/Plugins/PCG/Source/PCGEditor/Private/ChangeTracking/PCGActorTracker.cpp
--- a/Engine/Plugins/PCG/Source/PCGEditor/Private/ChangeTracking/PCGActorTracker.cpp
+++ b/Engine/Plugins/PCG/Source/PCGEditor/Private/ChangeTracking/PCGActorTracker.cpp
@@ -53,6 +53,16 @@ namespace PCGActorTracker
 		{
 			if (const ULevelInstanceSubsystem* LevelInstanceSubsystem = InLevelInstance->GetLevelInstanceSubsystem())
 			{
+				// Handle streamed out level for game worlds 
+				if (LevelInstanceSubsystem->GetWorld()->IsGameWorld())
+				{
+					ULevel* LoadedLevel = LevelInstanceSubsystem->GetLevelInstanceLevel(InLevelInstance);
+					if (!LoadedLevel || !LoadedLevel->OwningWorld)
+					{
+						return;
+					} 
+				}
+
 				LevelInstanceSubsystem->ForEachActorInLevelInstance(InLevelInstance, [InFunc](const AActor* LevelActor)
 				{
 					if (!ShouldIgnoreActor(LevelActor))
@@ -84,7 +94,7 @@ FPCGActorTracker::FPCGActorTracker(FPCGTrackingManager* InOwner)
 	UWorld* World = Owner->GetWorld();
 	check(World);
 
-	if (World->PersistentLevel)
+	if (World->PersistentLevel && !World->IsGameWorld())
 	{
 		World->PersistentLevel->OnLoadedActorAddedToLevelEvent.AddRaw(this, &FPCGActorTracker::OnActorLoaded);
 		World->PersistentLevel->OnLoadedActorRemovedFromLevelEvent.AddRaw(this, &FPCGActorTracker::OnActorUnloaded);
@@ -107,13 +117,7 @@ FPCGActorTracker::~FPCGActorTracker()
 	UWorld* World = Owner->GetWorld();
 	check(World);
 
-	if (UWorldPartition* WorldPartition = World->GetWorldPartition())
-	{
-		WorldPartition->OnActorDescInstanceAddedEvent.RemoveAll(this);
-		WorldPartition->OnActorDescInstanceRemovedEvent.RemoveAll(this);
-	}
-
-	if (World->PersistentLevel)
+	if (World->PersistentLevel && !World->IsGameWorld())
 	{
 		World->PersistentLevel->OnLoadedActorAddedToLevelEvent.RemoveAll(this);
 		World->PersistentLevel->OnLoadedActorRemovedFromLevelEvent.RemoveAll(this);
@@ -407,13 +411,11 @@ void FPCGActorTracker::OnActorDeleted_Internal(const AActor* InActor, bool bShou
 		return;
 	}
 
-#if WITH_EDITOR
 	PCGActorTracker::PropagateToLevelInstanceActors(InActor, [this, bShouldDirty](const AActor* LevelActor)
 	{
 		OnActorDeleted_Internal(LevelActor, bShouldDirty);
 		return true;
 	});
-#endif // WITH_EDITOR
 
 	// Notify all components that the actor has changed (was removed), but the Refresh will only happen AFTER the actor was actually removed from the world (because of delayed refresh).
 	OnActorChanged_Impl(InActor, nullptr, PCGActorTracker::ActorDeleted, /*InPreviousData=*/nullptr);
