[PCG] Compute and cache bounds once in URuntimeGenExecutionSource::ExecutePreGraph()

#rnx
#Synced-CL 51323499

diff --git a/Engine/Plugins/PCG/Source/PCG/Internal/RuntimeGen/PCGRuntimeGenExecutionSource.h b/Engine/Plugins/PCG/Source/PCG/Internal/RuntimeGen/PCGRuntimeGenExecutionSource.h
--- a/Engine/Plugins/PCG/Source/PCG/Internal/RuntimeGen/PCGRuntimeGenExecutionSource.h
+++ b/Engine/Plugins/PCG/Source/PCG/Internal/RuntimeGen/PCGRuntimeGenExecutionSource.h
@@ -62,6 +62,7 @@ public:
 	const FPCGDataCollection* RetrieveOutputDataForPin(const FString& InResourceKey) const;
 	FPCGManagedResourceContainer* GetManagedResourceContainer();
 	FTransactionallySafeCriticalSection* GetManagedResourceContainerLock();
+	void ExecutePreGraph(FPCGContext* InContext);
 
 #if WITH_EDITOR
 	const PCGUtils::FExtraCapture& GetExtraCapture() const;
@@ -108,6 +109,10 @@ private:
 	void ClearPerPinGeneratedOutput();
 
 	void CreateSelfData();
+	FBox GetBoundsInternal() const;
+	FBox GetOriginalBoundsInternal() const;
+	FBox GetLocalSpaceBoundsInternal() const;
+	FBox GetOriginalLocalSpaceBoundsInternal() const;
 
 #if WITH_EDITOR
 	bool UpdateTrackingCache(TArray<FPCGSelectionKey>* InOptionalChangedKeys = nullptr);
diff --git a/Engine/Plugins/PCG/Source/PCG/Internal/RuntimeGen/PCGRuntimeGenExecutionState.h b/Engine/Plugins/PCG/Source/PCG/Internal/RuntimeGen/PCGRuntimeGenExecutionState.h
--- a/Engine/Plugins/PCG/Source/PCG/Internal/RuntimeGen/PCGRuntimeGenExecutionState.h
+++ b/Engine/Plugins/PCG/Source/PCG/Internal/RuntimeGen/PCGRuntimeGenExecutionState.h
@@ -36,6 +36,7 @@ public:
 	virtual const FPCGDataCollection* RetrieveOutputDataForPin(const FString& InResourceKey) const override;
 	virtual FPCGManagedResourceContainer* GetManagedResourceContainer() override;
 	virtual FTransactionallySafeCriticalSection* GetManagedResourceContainerLock() override;
+	virtual void ExecutePreGraph(FPCGContext* InContext) override;
 
 #if WITH_EDITOR
 	virtual const PCGUtils::FExtraCapture& GetExtraCapture() const override;
diff --git a/Engine/Plugins/PCG/Source/PCG/Private/RuntimeGen/PCGRuntimeGenExecutionSource.cpp b/Engine/Plugins/PCG/Source/PCG/Private/RuntimeGen/PCGRuntimeGenExecutionSource.cpp
--- a/Engine/Plugins/PCG/Source/PCG/Private/RuntimeGen/PCGRuntimeGenExecutionSource.cpp
+++ b/Engine/Plugins/PCG/Source/PCG/Private/RuntimeGen/PCGRuntimeGenExecutionSource.cpp
@@ -62,8 +62,6 @@ void UPCGRuntimeGenExecutionSource::SetOwningSource(IPCGGraphExecutionSource* In
 
 	OwningSource = InOwningSource;
 
-	CreateSelfData();
-
 #if WITH_EDITOR
 	UpdateTrackingCache();
 #endif
@@ -76,7 +74,6 @@ void UPCGRuntimeGenExecutionSource::SetOwningSource(IPCGGraphExecutionSource* In
 		}
 
 		LocalSource->OwningSource = OwningSource;
-		LocalSource->CreateSelfData();
 	}
 }
 
@@ -152,7 +149,7 @@ FTransform UPCGRuntimeGenExecutionSource::GetOriginalTransform() const
 	return OwningSource->GetExecutionState().GetOriginalTransform();
 }
 
-FBox UPCGRuntimeGenExecutionSource::GetBounds() const
+FBox UPCGRuntimeGenExecutionSource::GetBoundsInternal() const
 {
 	if (IsLocalSource())
 	{
@@ -167,31 +164,135 @@ FBox UPCGRuntimeGenExecutionSource::GetBounds() const
 			Extent.Z = HALF_WORLD_MAX1;
 		}
 
-		return FBox(Center - Extent, Center + Extent).Overlap(GetOriginalBounds());
+		return FBox(Center - Extent, Center + Extent).Overlap(GetOriginalBoundsInternal());
 	}
 	else
 	{
-		return GetOriginalBounds();
+		return GetOriginalBoundsInternal();
 	}
 }
 
-FBox UPCGRuntimeGenExecutionSource::GetOriginalBounds() const
+FBox UPCGRuntimeGenExecutionSource::GetBounds() const
+{
+	if (GenerationTaskId != InvalidPCGTaskId)
+	{
+		UPCGSubsystem* PCGSubsystem = Cast<UPCGSubsystem>(ExecutionState.GetSubsystem());
+		check(PCGSubsystem);
+
+		TOptional<FBox> GridBounds = PCGSubsystem->GetBounds(GenerationTaskId);
+
+		if (!GridBounds.IsSet())
+		{
+			GridBounds = GetBoundsInternal();
+			PCGSubsystem->SetBounds(GenerationTaskId, GridBounds.GetValue());
+		}
+
+		check(GridBounds.IsSet());
+		return GridBounds.GetValue();
+	}
+	else
+	{
+		return GetBoundsInternal();
+	}
+}
+
+FBox UPCGRuntimeGenExecutionSource::GetOriginalBoundsInternal() const
 {
 	check(OwningSource);
 	return OwningSource->GetExecutionState().GetOriginalBounds();
 }
 
+FBox UPCGRuntimeGenExecutionSource::GetOriginalBounds() const
+{
+	if (GenerationTaskId != InvalidPCGTaskId)
+	{
+		UPCGSubsystem* PCGSubsystem = Cast<UPCGSubsystem>(ExecutionState.GetSubsystem());
+		check(PCGSubsystem);
+
+		TOptional<FBox> GridBounds = PCGSubsystem->GetOriginalBounds(GenerationTaskId);
+
+		if (!GridBounds.IsSet())
+		{
+			GridBounds = GetOriginalBoundsInternal();
+			PCGSubsystem->SetOriginalBounds(GenerationTaskId, GridBounds.GetValue());
+		}
+
+		check(GridBounds.IsSet());
+		return GridBounds.GetValue();
+	}
+	else
+	{
+		return GetOriginalBoundsInternal();
+	}
+}
+
+FBox UPCGRuntimeGenExecutionSource::GetLocalSpaceBoundsInternal() const
+{
+	if (IsLocalSource())
+	{
+		const FBox Bounds = GetBounds();
+		return Bounds.MoveTo(FVector::ZeroVector);
+	}
+	else
+	{
+		return GetOriginalLocalSpaceBoundsInternal();
+	}
+}
+
 FBox UPCGRuntimeGenExecutionSource::GetLocalSpaceBounds() const
 {
-	return GetBounds();
+	if (GenerationTaskId != InvalidPCGTaskId)
+	{
+		UPCGSubsystem* PCGSubsystem = Cast<UPCGSubsystem>(ExecutionState.GetSubsystem());
+		check(PCGSubsystem);
+
+		TOptional<FBox> GridBounds = PCGSubsystem->GetLocalSpaceBounds(GenerationTaskId);
+
+		if (!GridBounds.IsSet())
+		{
+			GridBounds = GetLocalSpaceBoundsInternal();
+			PCGSubsystem->SetLocalSpaceBounds(GenerationTaskId, GridBounds.GetValue());
+		}
+
+		check(GridBounds.IsSet());
+		return GridBounds.GetValue();
+	}
+	else
+	{
+		return GetLocalSpaceBoundsInternal();
+	}
 }
 
-FBox UPCGRuntimeGenExecutionSource::GetOriginalLocalSpaceBounds() const
+FBox UPCGRuntimeGenExecutionSource::GetOriginalLocalSpaceBoundsInternal() const
 {
 	check(OwningSource);
 	return OwningSource->GetExecutionState().GetOriginalLocalSpaceBounds();
 }
 
+FBox UPCGRuntimeGenExecutionSource::GetOriginalLocalSpaceBounds() const
+{
+	if (GenerationTaskId != InvalidPCGTaskId)
+	{
+		UPCGSubsystem* PCGSubsystem = Cast<UPCGSubsystem>(ExecutionState.GetSubsystem());
+		check(PCGSubsystem);
+
+		TOptional<FBox> GridBounds = PCGSubsystem->GetOriginalLocalSpaceBounds(GenerationTaskId);
+
+		if (!GridBounds.IsSet())
+		{
+			GridBounds = GetOriginalLocalSpaceBoundsInternal();
+			PCGSubsystem->SetOriginalLocalSpaceBounds(GenerationTaskId, GridBounds.GetValue());
+		}
+
+		check(GridBounds.IsSet());
+		return GridBounds.GetValue();
+	}
+	else
+	{
+		return GetOriginalLocalSpaceBoundsInternal();
+	}
+}
+
 UPCGGraph* UPCGRuntimeGenExecutionSource::GetGraph() const
 {
 	check(OwningSource);
@@ -371,6 +472,15 @@ FTransactionallySafeCriticalSection* UPCGRuntimeGenExecutionSource::GetManagedRe
 	return &ManagedResourcesLock;
 }
 
+void UPCGRuntimeGenExecutionSource::ExecutePreGraph(FPCGContext* InContext)
+{
+	CreateSelfData();
+	GetBounds();
+	GetOriginalBounds();
+	GetLocalSpaceBounds();
+	GetOriginalLocalSpaceBounds();
+}
+
 #if WITH_EDITOR
 const PCGUtils::FExtraCapture& UPCGRuntimeGenExecutionSource::GetExtraCapture() const
 {
@@ -538,7 +648,6 @@ UPCGRuntimeGenExecutionSource* UPCGRuntimeGenExecutionSource::CreateLocalSource(
 	LocalSource->OriginalSource = this;
 	LocalSource->GridSize = InGridSize;
 	LocalSource->GridCoords = InGridCoords;
-	LocalSource->CreateSelfData();
 
 	LocalSources.Add(LocalSource);
 
diff --git a/Engine/Plugins/PCG/Source/PCG/Private/RuntimeGen/PCGRuntimeGenExecutionState.cpp b/Engine/Plugins/PCG/Source/PCG/Private/RuntimeGen/PCGRuntimeGenExecutionState.cpp
--- a/Engine/Plugins/PCG/Source/PCG/Private/RuntimeGen/PCGRuntimeGenExecutionState.cpp
+++ b/Engine/Plugins/PCG/Source/PCG/Private/RuntimeGen/PCGRuntimeGenExecutionState.cpp
@@ -160,6 +160,12 @@ FTransactionallySafeCriticalSection* FPCGRuntimeGenExecutionState::GetManagedRes
 	return Source->GetManagedResourceContainerLock();
 }
 
+void FPCGRuntimeGenExecutionState::ExecutePreGraph(FPCGContext* InContext)
+{
+	check(Source);
+	Source->ExecutePreGraph(InContext);
+}
+
 #if WITH_EDITOR
 const PCGUtils::FExtraCapture& FPCGRuntimeGenExecutionState::GetExtraCapture() const
 {
