[PCG] Added fire-and-forget capability to execution so we can run on actors without relying on a PCG component.

#Synced-CL 50984763

#Synced-CL 50984763

diff --git a/Engine/Plugins/PCG/Source/PCG/Private/PCGDefaultActorExecutionSource.cpp b/Engine/Plugins/PCG/Source/PCG/Private/PCGDefaultActorExecutionSource.cpp
new file mode 100644
--- /dev/null
+++ b/Engine/Plugins/PCG/Source/PCG/Private/PCGDefaultActorExecutionSource.cpp
@@ -0,0 +1,104 @@
+// Copyright Epic Games, Inc. All Rights Reserved.
+
+#include "PCGDefaultActorExecutionSource.h"
+
+#include "PCGSubsystem.h"
+#include "Helpers/PCGHelpers.h"
+
+#include "GameFramework/Actor.h"
+
+#include UE_INLINE_GENERATED_CPP_BY_NAME(PCGDefaultActorExecutionSource)
+
+FPCGDefaultActorExecutionState::FPCGDefaultActorExecutionState(UPCGDefaultActorExecutionSource* InSource)
+	: FPCGDefaultExecutionState(InSource)
+{
+}
+
+UPCGData* FPCGDefaultActorExecutionState::GetSelfData() const
+{
+	if (AActor* Actor = GetActor())
+	{
+		if (UPCGSubsystem* PCGSubsystem = UPCGSubsystem::GetInstance(Actor->GetWorld()))
+		{
+			bool bFound = false;
+			return PCGSubsystem->GetSelfPCGData(GetTypedSource()->CurrentGenerationTask, bFound);
+		}
+	}
+
+	return FPCGDefaultExecutionState::GetSelfData();
+}
+
+FTransform FPCGDefaultActorExecutionState::GetTransform() const
+{
+	if (AActor* Actor = GetActor())
+	{
+		return Actor->GetTransform();
+	}
+	else
+	{
+		return FPCGDefaultExecutionState::GetTransform();
+	}
+}
+
+UObject* FPCGDefaultActorExecutionState::GetTarget() const
+{
+	return GetActor();
+}
+
+UWorld* FPCGDefaultActorExecutionState::GetWorld() const
+{
+	if (AActor* Actor = GetActor())
+	{
+		return Actor->GetWorld();
+	}
+	else
+	{
+		return FPCGDefaultExecutionState::GetWorld();
+	}
+}
+
+bool FPCGDefaultActorExecutionState::HasAuthority() const
+{
+	if (AActor* Actor = GetActor())
+	{
+		return Actor->HasAuthority();
+	}
+	else
+	{
+		return FPCGDefaultExecutionState::HasAuthority();
+	}
+}
+
+FBox FPCGDefaultActorExecutionState::GetBounds() const
+{
+	if (AActor* Actor = GetActor())
+	{
+		return PCGHelpers::GetActorBounds(Actor);
+	}
+	else
+	{
+		return FPCGDefaultExecutionState::GetBounds();
+	}
+}
+
+UPCGDefaultActorExecutionSource::UPCGDefaultActorExecutionSource()
+{
+	// Inefficient but more foolproof - we'll replace the previously created state from the base class by our version.
+	State = MakeUnique<FPCGDefaultActorExecutionState>(this);
+}
+
+void UPCGDefaultActorExecutionSource::Initialize(const FPCGDefaultActorExecutionSourceParams& InParams)
+{
+	Super::Initialize(InParams);
+	SetActor(InParams.Actor);
+}
+
+void UPCGDefaultActorExecutionSource::SetActor(AActor* InActor)
+{
+	Actor = InActor;
+}
+
+AActor* UPCGDefaultActorExecutionSource::GetActor()
+{
+	return Actor.Get();
+}
\ No newline at end of file
diff --git a/Engine/Plugins/PCG/Source/PCG/Private/PCGDefaultExecutionSource.cpp b/Engine/Plugins/PCG/Source/PCG/Private/PCGDefaultExecutionSource.cpp
--- a/Engine/Plugins/PCG/Source/PCG/Private/PCGDefaultExecutionSource.cpp
+++ b/Engine/Plugins/PCG/Source/PCG/Private/PCGDefaultExecutionSource.cpp
@@ -183,7 +183,7 @@ void UPCGDefaultExecutionSource::Generate()
 		return;
 	}
 
-	if (IPCGBaseSubsystem* Subsystem = State.GetSubsystem())
+	if (IPCGBaseSubsystem* Subsystem = State->GetSubsystem())
 	{
 		FPCGScheduleGraphParams GraphParams{
 			/*InGraph=*/ GetGraph(),
@@ -246,8 +246,8 @@ void UPCGDefaultExecutionSource::Sunset()
 }
 
 UPCGDefaultExecutionSource::UPCGDefaultExecutionSource()
-	: State(this)
 {
+	State = MakeUnique<FPCGDefaultExecutionState>(this);
 }
 
 UPCGDefaultExecutionSource::~UPCGDefaultExecutionSource()
diff --git a/Engine/Plugins/PCG/Source/PCG/Private/Subsystems/PCGEngineSubsystem.cpp b/Engine/Plugins/PCG/Source/PCG/Private/Subsystems/PCGEngineSubsystem.cpp
--- a/Engine/Plugins/PCG/Source/PCG/Private/Subsystems/PCGEngineSubsystem.cpp
+++ b/Engine/Plugins/PCG/Source/PCG/Private/Subsystems/PCGEngineSubsystem.cpp
@@ -51,18 +51,6 @@ void UPCGEngineSubsystem::Deinitialize()
 	Super::Deinitialize();
 }
 
-UPCGDefaultExecutionSource* UPCGEngineSubsystem::CreateExecutionSource(const FPCGDefaultExecutionSourceParams& InParams)
-{
-	if (!InParams.GraphInterface)
-	{
-		return nullptr;
-	}
-
-	UPCGDefaultExecutionSource* NewSource = NewObject<UPCGDefaultExecutionSource>(this);
-	NewSource->Initialize(InParams);
-	return NewSource;
-}
-
 void UPCGEngineSubsystem::NotifyGraphChanged(UPCGGraph* InGraph, EPCGChangeType ChangeType)
 {
 	if (IPCGEditorModule* PCGEditorModule = IPCGEditorModule::Get())
@@ -88,7 +76,7 @@ void UPCGEngineSubsystem::OnPCGSourceGenerationDone(IPCGGraphExecutionSource* In
 	if (Index != INDEX_NONE)
 	{
 		FGraphExecution& GraphExecution = GraphExecutions[Index];
-		GraphExecution.GenerationParams.GenerationCallback.ExecuteIfBound(ExecutionSource, InStatus);
+		GraphExecution.GenerationCallback.ExecuteIfBound(ExecutionSource, InStatus);
 		GraphExecutions.RemoveAtSwap(Index);
 
 		if (IPCGEditorModule* PCGEditorModule = IPCGEditorModule::Get())
@@ -118,12 +106,7 @@ void UPCGEngineSubsystem::GenerateGraph(const FPCGGenerateGraphParams& InParams)
 		return;
 	}
 
-	UPCGDefaultExecutionSource* ExecutionSource = CreateExecutionSource({ InParams.Graph, InParams.Seed });
-	if (ExecutionSource)
-	{
-		GraphExecutions.Add({ ExecutionSource, InParams });
-		ExecutionSource->Generate();
-	}
+	CreateExecutionSource({ InParams.Graph, InParams.Seed, /*bFireAndForget=*/true, InParams.GenerationCallback });
 }
 
 #endif // WITH_EDITOR
diff --git a/Engine/Plugins/PCG/Source/PCG/Public/PCGDefaultActorExecutionSource.h b/Engine/Plugins/PCG/Source/PCG/Public/PCGDefaultActorExecutionSource.h
new file mode 100644
--- /dev/null
+++ b/Engine/Plugins/PCG/Source/PCG/Public/PCGDefaultActorExecutionSource.h
@@ -0,0 +1,58 @@
+// Copyright Epic Games, Inc. All Rights Reserved.
+
+#pragma once
+
+#include "PCGDefaultExecutionSource.h"
+#include "PCGManagedResourceContainer.h"
+
+#include "PCGDefaultActorExecutionSource.generated.h"
+
+class FPCGDefaultActorExecutionState;
+
+struct FPCGDefaultActorExecutionSourceParams : public FPCGDefaultExecutionSourceParams
+{
+	AActor* Actor = nullptr;
+};
+
+UCLASS(MinimalAPI)
+class UPCGDefaultActorExecutionSource : public UPCGDefaultExecutionSource
+{
+	friend FPCGDefaultActorExecutionState;
+
+	GENERATED_BODY()
+
+public:
+	PCG_API UPCGDefaultActorExecutionSource();
+	PCG_API void Initialize(const FPCGDefaultActorExecutionSourceParams& InParams);
+	PCG_API AActor* GetActor();
+	PCG_API void SetActor(AActor* InActor);
+
+private:
+	TWeakObjectPtr<AActor> Actor = nullptr;
+};
+
+class FPCGDefaultActorExecutionState : public FPCGDefaultExecutionState
+{
+public:
+	FPCGDefaultActorExecutionState() = default;
+	explicit FPCGDefaultActorExecutionState(UPCGDefaultActorExecutionSource* InSource);
+
+	PCG_API virtual UPCGData* GetSelfData() const override;
+	PCG_API virtual FTransform GetTransform() const override;
+	PCG_API virtual UObject* GetTarget() const override;
+	PCG_API virtual UWorld* GetWorld() const override;
+	PCG_API virtual bool HasAuthority() const override;
+	PCG_API virtual FBox GetBounds() const override;
+
+	PCG_API virtual FPCGManagedResourceContainer* GetManagedResourceContainer() override { return &ManagedResourceContainer; }
+	PCG_API virtual FTransactionallySafeCriticalSection* GetManagedResourceContainerLock() override { return &ManagedResourcesLock; }
+
+protected:
+	UPCGDefaultActorExecutionSource* GetTypedSource() const { return CastChecked<UPCGDefaultActorExecutionSource>(GetSource()); }
+
+private:
+	AActor* GetActor() const { return GetTypedSource()->GetActor(); }
+
+	FPCGManagedResourceContainer ManagedResourceContainer;
+	mutable FTransactionallySafeCriticalSection ManagedResourcesLock;
+};
\ No newline at end of file
diff --git a/Engine/Plugins/PCG/Source/PCG/Public/PCGDefaultExecutionSource.h b/Engine/Plugins/PCG/Source/PCG/Public/PCGDefaultExecutionSource.h
--- a/Engine/Plugins/PCG/Source/PCG/Public/PCGDefaultExecutionSource.h
+++ b/Engine/Plugins/PCG/Source/PCG/Public/PCGDefaultExecutionSource.h
@@ -4,6 +4,7 @@
 
 #include "PCGGraphExecutionInspection.h"
 #include "PCGGraphExecutionStateInterface.h"
+#include "Subsystems/PCGEngineSubsystem.h"
 
 #include "PCGDefaultExecutionSource.generated.h"
 
@@ -48,6 +49,9 @@ public:
 	virtual FPCGDynamicTrackingPriority GetDynamicTrackingPriority() const override { return FPCGDynamicTrackingPriority(); }
 #endif // WITH_EDITOR
 
+protected:
+	UPCGDefaultExecutionSource* GetSource() const { return Source; }
+
 private:
 	UPCGDefaultExecutionSource* Source = nullptr;
 };
@@ -56,6 +60,10 @@ struct FPCGDefaultExecutionSourceParams
 {
 	UPCGGraphInterface* GraphInterface = nullptr;
 	int32 Seed = 42;
+	bool bFireAndForgetExecution = false;
+
+	/** Optional callback when generation completes. */
+	FPCGOnEditorGenerationDone GenerationCallback;
 };
 
 UCLASS(MinimalAPI)
@@ -73,8 +81,8 @@ public:
 
 	PCG_API static void AddReferencedObjects(UObject* InThis, FReferenceCollector& Collector);
 	
-	virtual IPCGGraphExecutionState& GetExecutionState() override { return State; }
-	virtual const IPCGGraphExecutionState& GetExecutionState() const override { return State; }
+	virtual IPCGGraphExecutionState& GetExecutionState() override { return *State; }
+	virtual const IPCGGraphExecutionState& GetExecutionState() const override { return *State; }
 
 	PCG_API void SetGraphInterface(UPCGGraphInterface* InGraphInterface);
 
@@ -92,13 +100,13 @@ public:
 
 	PCG_API void Sunset();
 
-private:
+protected:
 	FPCGTaskId CurrentGenerationTask = InvalidPCGTaskId;
 	
 	UPROPERTY()
 	TObjectPtr<UPCGGraphInterface> GraphInterface;
 	
-	FPCGDefaultExecutionState State;
+	TUniquePtr<FPCGDefaultExecutionState> State;
 	int32 Seed = 42;
 
 #if WITH_EDITOR
diff --git a/Engine/Plugins/PCG/Source/PCG/Public/Subsystems/PCGEngineSubsystem.h b/Engine/Plugins/PCG/Source/PCG/Public/Subsystems/PCGEngineSubsystem.h
--- a/Engine/Plugins/PCG/Source/PCG/Public/Subsystems/PCGEngineSubsystem.h
+++ b/Engine/Plugins/PCG/Source/PCG/Public/Subsystems/PCGEngineSubsystem.h
@@ -45,7 +45,27 @@ public:
 	PCG_API virtual void Initialize(FSubsystemCollectionBase& Collection) override;
 	PCG_API virtual void Deinitialize() override;
 
-	PCG_API UPCGDefaultExecutionSource* CreateExecutionSource(const FPCGDefaultExecutionSourceParams& InParams);
+	/** Creates an execution source from the params and optionally tracks its lifetime and generation. Note this call is NOT THREADSAFE and as such this should be called only from the game thread. */
+	template<typename SourceType = UPCGDefaultExecutionSource, typename ParamsType = FPCGDefaultExecutionSourceParams> 
+		requires std::is_base_of_v<IPCGGraphExecutionSource, SourceType> && std::is_base_of_v<FPCGDefaultExecutionSourceParams, ParamsType>
+	SourceType* CreateExecutionSource(const ParamsType& InParams)
+	{
+		if (!InParams.GraphInterface)
+		{
+			return nullptr;
+		}
+
+		SourceType* NewSource = NewObject<SourceType>(this);
+		NewSource->Initialize(InParams);
+
+		if (InParams.bFireAndForgetExecution)
+		{
+			GraphExecutions.Emplace(NewSource, InParams.GenerationCallback);
+			NewSource->Generate();
+		}
+
+		return NewSource;
+	}
 
 	PCG_API void GenerateGraph(const FPCGGenerateGraphParams& InParams);
 
@@ -66,7 +86,7 @@ protected:
 	struct FGraphExecution
 	{
 		TObjectPtr<UPCGDefaultExecutionSource> ExecutionSource;
-		FPCGGenerateGraphParams GenerationParams;
+		FPCGOnEditorGenerationDone GenerationCallback;
 	};
 
 	TArray<FGraphExecution> GraphExecutions;
