[PCG] Fix bounds computation for SceneCapture FromExecutionSource orientation mode

Also add warning when capturing scene depth with an incompatible format, to clear up confusion for users who expect 8 bit texture format to be compatible with depth captures

#rnx

#Synced-CL 51323499

diff --git a/Engine/Plugins/PCG/Source/PCG/Private/Elements/PCGSceneCapture.cpp b/Engine/Plugins/PCG/Source/PCG/Private/Elements/PCGSceneCapture.cpp
--- a/Engine/Plugins/PCG/Source/PCG/Private/Elements/PCGSceneCapture.cpp
+++ b/Engine/Plugins/PCG/Source/PCG/Private/Elements/PCGSceneCapture.cpp
@@ -145,12 +145,13 @@ bool FPCGSceneCaptureElement::ExecuteInternal(FPCGContext* InContext) const
 		}
 		else if (Settings->OrientationMode == EPCGSceneCaptureOrientationMode::FromExecutionSource)
 		{
-			const FBox OriginalActorLocalBounds = ExecutionState.GetOriginalLocalSpaceBounds();
+			const FBox OriginalActorBounds = ExecutionState.GetOriginalBounds();
+			const FBox ActorBounds = ExecutionState.GetBounds();
+			const FBox ExecutionBounds = OriginalActorBounds.Overlap(ActorBounds);
 			const FTransform ActorTransform = ExecutionState.GetTransform();
 
-			CaptureLocation = ActorTransform.GetLocation();
+			ExecutionBounds.GetCenterAndExtents(CaptureLocation, CaptureHalfExtents);
 			CaptureRotation = ActorTransform.GetRotation();
-			CaptureHalfExtents = ActorTransform.GetScale3D() * 0.5 * (OriginalActorLocalBounds.Max - OriginalActorLocalBounds.Min);
 		}
 
 		if (CaptureHalfExtents.Size() <= UE_DOUBLE_SMALL_NUMBER)
@@ -209,6 +210,20 @@ bool FPCGSceneCaptureElement::ExecuteInternal(FPCGContext* InContext) const
 		Context->SceneCaptureComponent->SetWorldLocation(CaptureLocation);
 		Context->SceneCaptureComponent->SetWorldRotation(CaptureComponentRotation);
 
+		const bool bCaptureDepth = Settings->CaptureSource == ESceneCaptureSource::SCS_SceneColorSceneDepth || Settings->CaptureSource == ESceneCaptureSource::SCS_SceneDepth;
+		const bool bUNorm = Settings->PixelFormat == ETextureRenderTargetFormat::RTF_R8
+			|| Settings->PixelFormat == ETextureRenderTargetFormat::RTF_RG8
+			|| Settings->PixelFormat == ETextureRenderTargetFormat::RTF_RGBA8
+			|| Settings->PixelFormat == ETextureRenderTargetFormat::RTF_RGBA8_SRGB
+			|| Settings->PixelFormat == ETextureRenderTargetFormat::RTF_RGB10A2;
+
+		if (bCaptureDepth && bUNorm)
+		{
+			PCGLog::LogWarningOnGraph(FText::Format(LOCTEXT("IncompatibleFormats", "Performing scene capture for '{0}', but render target format is '{1}'. Depth will be clamped to 1."),
+				StaticEnum<ESceneCaptureSource>()->GetDisplayNameTextByValue((int64)Settings->CaptureSource),
+				StaticEnum<ETextureRenderTargetFormat>()->GetDisplayNameTextByValue((int64)Settings->PixelFormat)));
+		}
+
 		// Gather actors and primitives for inclusion/exclusion filtering.
 		{
 			TRACE_CPUPROFILER_EVENT_SCOPE(FilterPrimitivesForCapture);
