[PCG] Fixed bug with Data Overrides storing in Data Visualizer, preventing overrides from being stored properly

* Modified API for Data overrides slightly to allow for Add and Add_GetRef mutations

#Synced-CL 51342292

diff --git a/Engine/Plugins/PCG/Source/PCG/Private/Graph/DataOverride/PCGDataOverride.cpp b/Engine/Plugins/PCG/Source/PCG/Private/Graph/DataOverride/PCGDataOverride.cpp
--- a/Engine/Plugins/PCG/Source/PCG/Private/Graph/DataOverride/PCGDataOverride.cpp
+++ b/Engine/Plugins/PCG/Source/PCG/Private/Graph/DataOverride/PCGDataOverride.cpp
@@ -38,7 +38,12 @@ bool FPCGDeltaCollection::Contains(const FPCGDeltaKey& InKey) const
 }
 
 #if WITH_EDITOR
-TInstancedStruct<FPCGDeltaBase>& FPCGDeltaCollection::Add(const FPCGDeltaKey& InKey, TInstancedStruct<FPCGDeltaBase>&& InDelta)
+void FPCGDeltaCollection::Add(const FPCGDeltaKey& InKey, TInstancedStruct<FPCGDeltaBase>&& InDelta)
+{
+	(void)Add_GetRef(InKey, MoveTemp(InDelta));
+}
+
+TInstancedStruct<FPCGDeltaBase>& FPCGDeltaCollection::Add_GetRef(const FPCGDeltaKey& InKey, TInstancedStruct<FPCGDeltaBase>&& InDelta)
 {
 	TInstancedStruct<FPCGDeltaBase>& Struct = Deltas.Add(InKey, MoveTemp(InDelta));
 	if (FPCGDeltaBase* Delta = Struct.GetMutablePtr<FPCGDeltaBase>())
diff --git a/Engine/Plugins/PCG/Source/PCG/Public/Graph/DataOverride/PCGDataOverride.h b/Engine/Plugins/PCG/Source/PCG/Public/Graph/DataOverride/PCGDataOverride.h
--- a/Engine/Plugins/PCG/Source/PCG/Public/Graph/DataOverride/PCGDataOverride.h
+++ b/Engine/Plugins/PCG/Source/PCG/Public/Graph/DataOverride/PCGDataOverride.h
@@ -148,7 +148,9 @@ struct FPCGDeltaCollection
 	// For now, this will be editor only for storage mutation.
 #if WITH_EDITOR
 	/** Add a new delta entry to the collection, as a TInstancedStruct<FPCGDeltaBase>. */
-	[[maybe_unused]] PCG_API TInstancedStruct<FPCGDeltaBase>& Add(const FPCGDeltaKey& InKey, TInstancedStruct<FPCGDeltaBase>&& InDelta);
+	PCG_API void Add(const FPCGDeltaKey& InKey, TInstancedStruct<FPCGDeltaBase>&& InDelta);
+	/** Add a new delta entry to the collection, as a TInstancedStruct<FPCGDeltaBase>. Returns a mutable reference to the delta. */
+	PCG_API TInstancedStruct<FPCGDeltaBase>& Add_GetRef(const FPCGDeltaKey& InKey, TInstancedStruct<FPCGDeltaBase>&& InDelta);
 
 	/** Remove a delta entry by key. */
 	PCG_API bool Remove(const FPCGDeltaKey& InKey);
diff --git a/Engine/Plugins/PCG/Source/PCGEditor/Private/PCGComponentVisualizer.cpp b/Engine/Plugins/PCG/Source/PCGEditor/Private/PCGComponentVisualizer.cpp
--- a/Engine/Plugins/PCG/Source/PCGEditor/Private/PCGComponentVisualizer.cpp
+++ b/Engine/Plugins/PCG/Source/PCGEditor/Private/PCGComponentVisualizer.cpp
@@ -162,11 +162,30 @@ void FPCGComponentVisualizer::DrawVisualization(const UActorComponent* Component
 			for (const auto& [ObjectKey, StackSet] : NodeStacks)
 			{
 				const UPCGNode* Node = ObjectKey.ResolveObjectPtr();
-				if (Node && Node->GetSettingsInterface() && Node->GetSettingsInterface()->bViewportEditingEnabled && Node->GetOutputPins().Num() > 0)
+				if (!Node || !Node->GetSettingsInterface() || !Node->GetSettingsInterface()->bViewportEditingEnabled)
+				{
+					continue;
+				}
+
+				FPCGElementPtr Element = Node->GetSettings() ? Node->GetSettings()->GetElement() : nullptr;
+				if (!Element)
+				{
+					continue;
+				}
+
+				// Only override pre or post execute.
+				EPCGDataOverridePhase Phase = Element->GetDataOverridePhase();
+				if (Phase != EPCGDataOverridePhase::PrepareData && Phase != EPCGDataOverridePhase::PostExecute)
+				{
+					continue;
+				}
+
+				const TArray<TObjectPtr<UPCGPin>>& Pins = (Phase == EPCGDataOverridePhase::PrepareData) ? Node->GetInputPins() : Node->GetOutputPins();
+				if (Pins.Num() > 0)
 				{
 					for (const FPCGGraphExecutionInspection::FNodeExecutedNotificationData& Stack : StackSet)
 					{
-						const UPCGPin* PinUsed = Node->GetOutputPins()[0];
+						const UPCGPin* PinUsed = Pins[0];
 
 						FPCGStack PinStack = Stack.Stack;
 						PinStack.PushFrame(Node);
@@ -362,8 +381,7 @@ bool FPCGComponentVisualizer::VisProxyHandleClick(FEditorViewportClient* InViewp
 						Collection->Settings.SpatialTolerance,
 						FPCGPointEditDelta::GetDeltaNameStatic());
 
-					// @todo_pcg: This temporarily bypasses static analysis warning in PVS Studio. API to be revisited.
-					(void)Collection->Add(DeltaKey, TInstancedStruct<FPCGPointEditDelta>::Make(EditDelta));
+					Collection->Add(DeltaKey, TInstancedStruct<FPCGPointEditDelta>::Make(EditDelta));
 				}
 
 				//Now that the point is part of the delta, set up the current selection
