[PCG] Refresh runtime gen from FPCGLandscapeTracker::ApplyLandscapeChanges()

This is an upstream fix that was not robomerged in the past

#rnx

#Synced-CL 51323499

#Synced-CL 51323499

#Synced-CL 51323499

diff --git a/Engine/Plugins/PCG/Source/PCGEditor/Private/ChangeTracking/PCGLandscapeTracker.cpp b/Engine/Plugins/PCG/Source/PCGEditor/Private/ChangeTracking/PCGLandscapeTracker.cpp
--- a/Engine/Plugins/PCG/Source/PCGEditor/Private/ChangeTracking/PCGLandscapeTracker.cpp
+++ b/Engine/Plugins/PCG/Source/PCGEditor/Private/ChangeTracking/PCGLandscapeTracker.cpp
@@ -7,6 +7,7 @@
 #include "PCGTrackingManager.h"
 #include "ChangeTracking/PCGActorTracker.h"
 #include "Landscape/PCGLandscapeCVars.h"
+#include "Subsystems/PCGSubsystem.h"
 
 #include "Editor.h"
 #include "EditorModeManager.h"
@@ -135,6 +136,16 @@ void FPCGLandscapeTracker::ApplyLandscapeChanges(ALandscapeProxy* InLandscape)
 		return;
 	}
 
+#if WITH_EDITOR
+	if (UPCGSubsystem* Subsystem = UPCGSubsystem::GetInstance(Owner->GetWorld()))
+	{
+		// @todo_pcg: Temporarily inject a runtime-only refresh path when applying changes from landscape actors.
+		// This is necessary for clearing and regenerating grass/detail on landscapes while painting.
+		// Should be removed when change tracking can propagate changes for runtime gen specifically.
+		Subsystem->RefreshAllRuntimeGenExecutionSources(EPCGChangeType::GenerationGrid);
+	}
+#endif
+
 	// Forward as an object change event that can be processed by another tracker
 	Owner->NotifyObjectChanged(InLandscape, PCGActorTracker::ActorChanged);
 }
