[PCG] Fix texture data transforms always using the original actor transform even when partitioned

#rnx

#Synced-CL 51323499

diff --git a/Engine/Plugins/PCG/Source/PCG/Private/Compute/Elements/PCGCustomHLSLKernel.cpp b/Engine/Plugins/PCG/Source/PCG/Private/Compute/Elements/PCGCustomHLSLKernel.cpp
--- a/Engine/Plugins/PCG/Source/PCG/Private/Compute/Elements/PCGCustomHLSLKernel.cpp
+++ b/Engine/Plugins/PCG/Source/PCG/Private/Compute/Elements/PCGCustomHLSLKernel.cpp
@@ -300,10 +300,7 @@ TSharedPtr<const FPCGDataCollectionDesc> UPCGCustomHLSLKernel::ComputeOutputBind
 			TSharedPtr<FPCGContextHandle> ContextHandle = InBinding->GetContextHandle().Pin();
 			if (FPCGContext* Context = ContextHandle ? ContextHandle->GetContext() : nullptr)
 			{
-				const FBox OriginalLocalBounds = Context->ExecutionSource->GetExecutionState().GetOriginalLocalSpaceBounds();
-				FTransform Transform = Context->ExecutionSource->GetExecutionState().GetOriginalTransform();
-				Transform.SetScale3D(Transform.GetScale3D() * 0.5 * (OriginalLocalBounds.Max - OriginalLocalBounds.Min));
-				OutDataDesc->SetTextureTransformForAllData(Transform);
+				OutDataDesc->SetTextureTransformForAllData(UPCGTexture2DBaseData::ComputeTransform(Context->ExecutionSource.Get()));
 			}
 		}
 	}
diff --git a/Engine/Plugins/PCG/Source/PCG/Private/Compute/PCGComputeKernel.cpp b/Engine/Plugins/PCG/Source/PCG/Private/Compute/PCGComputeKernel.cpp
--- a/Engine/Plugins/PCG/Source/PCG/Private/Compute/PCGComputeKernel.cpp
+++ b/Engine/Plugins/PCG/Source/PCG/Private/Compute/PCGComputeKernel.cpp
@@ -590,10 +590,7 @@ void UPCGComputeKernel::ComputeDataDescFromPinProperties(const FPCGPinProperties
 				TSharedPtr<FPCGContextHandle> ContextHandle = InBinding->GetContextHandle().Pin();
 				if (FPCGContext* Context = ContextHandle ? ContextHandle->GetContext() : nullptr)
 				{
-					const FBox OriginalLocalBounds = Context->ExecutionSource->GetExecutionState().GetOriginalLocalSpaceBounds();
-					FTransform Transform = Context->ExecutionSource->GetExecutionState().GetOriginalTransform();
-					Transform.SetScale3D(Transform.GetScale3D() * 0.5 * (OriginalLocalBounds.Max - OriginalLocalBounds.Min));
-					OutPinDesc->SetTextureTransformForAllData(Transform);
+					OutPinDesc->SetTextureTransformForAllData(UPCGTexture2DBaseData::ComputeTransform(Context->ExecutionSource.Get()));
 				}
 			}
 		}
@@ -640,10 +637,7 @@ void UPCGComputeKernel::ComputeDataDescFromPinProperties(const FPCGPinProperties
 				TSharedPtr<FPCGContextHandle> ContextHandle = InBinding->GetContextHandle().Pin();
 				if (FPCGContext* Context = ContextHandle ? ContextHandle->GetContext() : nullptr)
 				{
-					const FBox OriginalLocalBounds = Context->ExecutionSource->GetExecutionState().GetOriginalLocalSpaceBounds();
-					FTransform Transform = Context->ExecutionSource->GetExecutionState().GetOriginalTransform();
-					Transform.SetScale3D(Transform.GetScale3D() * 0.5 * (OriginalLocalBounds.Max - OriginalLocalBounds.Min));
-					OutPinDesc->SetTextureTransformForAllData(Transform);
+					OutPinDesc->SetTextureTransformForAllData(UPCGTexture2DBaseData::ComputeTransform(Context->ExecutionSource.Get()));
 				}
 			}
 			else
diff --git a/Engine/Plugins/PCG/Source/PCG/Private/Data/PCGTexture2DBaseData.cpp b/Engine/Plugins/PCG/Source/PCG/Private/Data/PCGTexture2DBaseData.cpp
--- a/Engine/Plugins/PCG/Source/PCG/Private/Data/PCGTexture2DBaseData.cpp
+++ b/Engine/Plugins/PCG/Source/PCG/Private/Data/PCGTexture2DBaseData.cpp
@@ -23,3 +23,17 @@ void UPCGTexture2DBaseData::CopyTexture2DBaseData(UPCGTexture2DBaseData* NewText
 	NewTextureData->TextureHandle = TextureHandle;
 	NewTextureData->Bounds = Bounds;
 }
+
+FTransform UPCGTexture2DBaseData::ComputeTransform(const IPCGGraphExecutionSource* InExecutionSource)
+{
+	FTransform Transform = FTransform::Identity;
+
+	if (InExecutionSource)
+	{
+		const FBox Bounds = InExecutionSource->GetExecutionState().GetBounds();
+		Transform = InExecutionSource->GetExecutionState().GetTransform();
+		Transform.SetScale3D(Transform.GetScale3D() * Bounds.GetExtent());
+	}
+
+	return Transform;
+}
diff --git a/Engine/Plugins/PCG/Source/PCG/Private/Elements/PCGGetTexture2DArrayData.cpp b/Engine/Plugins/PCG/Source/PCG/Private/Elements/PCGGetTexture2DArrayData.cpp
--- a/Engine/Plugins/PCG/Source/PCG/Private/Elements/PCGGetTexture2DArrayData.cpp
+++ b/Engine/Plugins/PCG/Source/PCG/Private/Elements/PCGGetTexture2DArrayData.cpp
@@ -118,13 +118,9 @@ bool FPCGGetTexture2DArrayDataElement::ExecuteInternal(FPCGContext* Context) con
 		return true;
 	}
 
-	FTransform Transform = Context->ExecutionSource->GetExecutionState().GetOriginalTransform();
-	const FBox OriginalActorLocalBounds = Context->ExecutionSource->GetExecutionState().GetOriginalLocalSpaceBounds();
-	Transform.SetScale3D(Transform.GetScale3D() * 0.5 * (OriginalActorLocalBounds.Max - OriginalActorLocalBounds.Min));
-
 	FPCGTexture2DArrayDataInitParams InitParams;
 	InitParams.Filter = Settings->Filter;
-	InitParams.Transform = MoveTemp(Transform);
+	InitParams.Transform = UPCGTexture2DBaseData::ComputeTransform(Context->ExecutionSource.Get());
 
 	UPCGTexture2DArrayData* Texture2DArrayData = FPCGContext::NewObject_AnyThread<UPCGTexture2DArrayData>(Context);
 	Texture2DArrayData->Initialize(Texture, InitParams);
diff --git a/Engine/Plugins/PCG/Source/PCG/Private/Elements/PCGTextureSampler.cpp b/Engine/Plugins/PCG/Source/PCG/Private/Elements/PCGTextureSampler.cpp
--- a/Engine/Plugins/PCG/Source/PCG/Private/Elements/PCGTextureSampler.cpp
+++ b/Engine/Plugins/PCG/Source/PCG/Private/Elements/PCGTextureSampler.cpp
@@ -235,11 +235,11 @@ bool FPCGTextureSamplerElement::ExecuteInternal(FPCGContext* InContext) const
 		}
 		else
 		{
-			const FTransform OriginalTransform = Context->ExecutionSource->GetExecutionState().GetOriginalTransform();
-			Context->Transform = Settings->Transform * OriginalTransform;
+			const FTransform Transform = Context->ExecutionSource->GetExecutionState().GetTransform();
+			Context->Transform = Settings->Transform * Transform;
 
-			const FBox OriginalActorLocalBounds = Context->ExecutionSource->GetExecutionState().GetOriginalLocalSpaceBounds();
-			Context->Transform.SetScale3D(Context->Transform.GetScale3D() * 0.5 * (OriginalActorLocalBounds.Max - OriginalActorLocalBounds.Min));
+			const FBox Bounds = Context->ExecutionSource->GetExecutionState().GetBounds();
+			Context->Transform.SetScale3D(Context->Transform.GetScale3D() * Bounds.GetExtent());
 		}
 
 		Context->bTransformInitialized = true;
diff --git a/Engine/Plugins/PCG/Source/PCG/Public/Data/PCGTexture2DBaseData.h b/Engine/Plugins/PCG/Source/PCG/Public/Data/PCGTexture2DBaseData.h
--- a/Engine/Plugins/PCG/Source/PCG/Public/Data/PCGTexture2DBaseData.h
+++ b/Engine/Plugins/PCG/Source/PCG/Public/Data/PCGTexture2DBaseData.h
@@ -43,6 +43,9 @@ public:
 	virtual FIntPoint GetResolution() const PURE_VIRTUAL(UPCGTexture2DBaseData::GetResolution, return FIntPoint::ZeroValue;);
 	virtual TRefCountPtr<IPooledRenderTarget> GetRefCountedTexture() const { return TextureHandle; }
 
+	/** Helper to compute the texture data transform based on the execution source transform and bounds. */
+	static FTransform ComputeTransform(const IPCGGraphExecutionSource* InExecutionSource);
+
 protected:
 	PCG_API void CopyTexture2DBaseData(UPCGTexture2DBaseData* NewTextureData) const;
 
