[PCG] Added self data caching on default actor execution source and local bounds and improved default local bounds behavior.

#Synced-CL 51379665

diff --git a/Engine/Plugins/PCG/Source/PCG/Private/PCGDefaultActorExecutionSource.cpp b/Engine/Plugins/PCG/Source/PCG/Private/PCGDefaultActorExecutionSource.cpp
--- a/Engine/Plugins/PCG/Source/PCG/Private/PCGDefaultActorExecutionSource.cpp
+++ b/Engine/Plugins/PCG/Source/PCG/Private/PCGDefaultActorExecutionSource.cpp
@@ -3,6 +3,7 @@
 #include "PCGDefaultActorExecutionSource.h"
 
 #include "PCGSubsystem.h"
+#include "Data/PCGUnionData.h"
 #include "Helpers/PCGHelpers.h"
 
 #include "GameFramework/Actor.h"
@@ -14,6 +15,16 @@ FPCGDefaultActorExecutionState::FPCGDefaultActorExecutionState(UPCGDefaultActorE
 {
 }
 
+FString FPCGDefaultActorExecutionState::GetDebugName() const
+{
+	return FString(TEXT("PCGDefaultActorExecutionSource - "))
+#if WITH_EDITOR
+		+ FString(GetActor() ? GetActor()->GetActorLabel() : TEXT("Unknown actor"));
+#else
+		+ FString(GetActor() ? GetActor()->GetName() : TEXT("Unknown actor"));
+#endif
+}
+
 UPCGData* FPCGDefaultActorExecutionState::GetSelfData() const
 {
 	if (AActor* Actor = GetActor())
@@ -28,6 +39,46 @@ UPCGData* FPCGDefaultActorExecutionState::GetSelfData() const
 	return FPCGDefaultExecutionState::GetSelfData();
 }
 
+void FPCGDefaultActorExecutionState::ExecutePreGraph(FPCGContext* Context)
+{
+	// Actor execution states have only self-data, and will not fill in other pins.
+	// @todo_pcg: should we also set bounds?
+	AActor* Actor = GetActor();
+	UPCGSubsystem* PCGSubsystem = Actor ? UPCGSubsystem::GetInstance(Actor->GetWorld()) : nullptr;
+	if (PCGSubsystem)
+	{
+		bool bFound = false;
+		UPCGData* SelfData = PCGSubsystem->GetSelfPCGData(GetTypedSource()->CurrentGenerationTask, bFound);
+		if (!bFound)
+		{
+			FPCGGetDataFunctionRegistryParams Params;
+			Params.Source = GetSource();
+			Params.bParseActor = true;
+			Params.DataTypeFilter = EPCGDataType::Spatial;
+
+			FPCGGetDataFunctionRegistryOutput Output;
+			FPCGModule::ConstGetDataFunctionRegistry().GetDataFromActor(/*Context=*/nullptr, Params, Actor, Output);
+
+			if (Output.Collection.TaggedData.Num() > 1)
+			{
+				UPCGUnionData* Union = NewObject<UPCGUnionData>();
+				for (const FPCGTaggedData& TaggedData : Output.Collection.TaggedData)
+				{
+					Union->AddData(CastChecked<const UPCGSpatialData>(TaggedData.Data));
+				}
+
+				SelfData = Union;
+			}
+			else if (Output.Collection.TaggedData.Num() == 1)
+			{
+				SelfData = const_cast<UPCGData*>(Output.Collection.TaggedData[0].Data.Get());
+			}
+			
+			PCGSubsystem->SetSelfPCGData(GetTypedSource()->CurrentGenerationTask, SelfData);
+		}
+	}
+}
+
 FTransform FPCGDefaultActorExecutionState::GetTransform() const
 {
 	if (AActor* Actor = GetActor())
@@ -81,6 +132,18 @@ FBox FPCGDefaultActorExecutionState::GetBounds() const
 	}
 }
 
+FBox FPCGDefaultActorExecutionState::GetLocalSpaceBounds() const
+{
+	if (AActor* Actor = GetActor())
+	{
+		return PCGHelpers::GetActorLocalBounds(Actor);
+	}
+	else
+	{
+		return FPCGDefaultExecutionState::GetLocalSpaceBounds();
+	}
+}
+
 UPCGDefaultActorExecutionSource::UPCGDefaultActorExecutionSource()
 {
 	// Inefficient but more foolproof - we'll replace the previously created state from the base class by our version.
diff --git a/Engine/Plugins/PCG/Source/PCG/Public/PCGDefaultActorExecutionSource.h b/Engine/Plugins/PCG/Source/PCG/Public/PCGDefaultActorExecutionSource.h
--- a/Engine/Plugins/PCG/Source/PCG/Public/PCGDefaultActorExecutionSource.h
+++ b/Engine/Plugins/PCG/Source/PCG/Public/PCGDefaultActorExecutionSource.h
@@ -37,15 +37,18 @@ public:
 	FPCGDefaultActorExecutionState() = default;
 	explicit FPCGDefaultActorExecutionState(UPCGDefaultActorExecutionSource* InSource);
 
+	PCG_API virtual FString GetDebugName() const override;
 	PCG_API virtual UPCGData* GetSelfData() const override;
 	PCG_API virtual FTransform GetTransform() const override;
 	PCG_API virtual UObject* GetTarget() const override;
 	PCG_API virtual UWorld* GetWorld() const override;
 	PCG_API virtual bool HasAuthority() const override;
 	PCG_API virtual FBox GetBounds() const override;
+	PCG_API virtual FBox GetLocalSpaceBounds() const override;
 
 	PCG_API virtual FPCGManagedResourceContainer* GetManagedResourceContainer() override { return &ManagedResourceContainer; }
 	PCG_API virtual FTransactionallySafeCriticalSection* GetManagedResourceContainerLock() override { return &ManagedResourcesLock; }
+	PCG_API virtual void ExecutePreGraph(FPCGContext* InContext) override;
 
 protected:
 	UPCGDefaultActorExecutionSource* GetTypedSource() const { return CastChecked<UPCGDefaultActorExecutionSource>(GetSource()); }
diff --git a/Engine/Plugins/PCG/Source/PCG/Public/PCGGraphExecutionStateInterface.h b/Engine/Plugins/PCG/Source/PCG/Public/PCGGraphExecutionStateInterface.h
--- a/Engine/Plugins/PCG/Source/PCG/Public/PCGGraphExecutionStateInterface.h
+++ b/Engine/Plugins/PCG/Source/PCG/Public/PCGGraphExecutionStateInterface.h
@@ -106,7 +106,7 @@ public:
 	virtual FBox GetOriginalBounds() const { return GetBounds(); }
 
 	/** Returns the ExecutionState's local bounds if this ExecutionState is a spatial one. */
-	virtual FBox GetLocalSpaceBounds() const { return GetBounds(); }
+	virtual FBox GetLocalSpaceBounds() const { return GetBounds().InverseTransformBy(GetTransform()); }
 
 	/** Returns the ExecutionState's original source local space bounds if this ExecutionState is a local source. */
 	virtual FBox GetOriginalLocalSpaceBounds() const { return GetLocalSpaceBounds(); }
