[PCG] Update some artifact creation nodes so they are execution source-based rather than PCG component based.

#rnx

#Synced-CL 51379665

diff --git a/Engine/Plugins/PCG/Source/PCG/Private/Elements/Landscape/PCGWaitLandscapeReady.cpp b/Engine/Plugins/PCG/Source/PCG/Private/Elements/Landscape/PCGWaitLandscapeReady.cpp
--- a/Engine/Plugins/PCG/Source/PCG/Private/Elements/Landscape/PCGWaitLandscapeReady.cpp
+++ b/Engine/Plugins/PCG/Source/PCG/Private/Elements/Landscape/PCGWaitLandscapeReady.cpp
@@ -2,7 +2,6 @@
 
 #include "Elements/Landscape/PCGWaitLandscapeReady.h"
 
-#include "PCGComponent.h"
 #include "PCGSubsystem.h"
 #include "Elements/PCGActorSelector.h"
 
@@ -56,8 +55,7 @@ bool FPCGWaitLandscapeReadyElement::ExecuteInternal(FPCGContext* InContext) cons
 		auto BoundsCheck = [](const AActor*) -> bool { return true; };
 		auto SelfIgnoreCheck = [](const AActor*) -> bool { return true; };
 
-		const UPCGComponent* SourceComponent = Cast<UPCGComponent>(Context->ExecutionSource.Get());
-		TArray<AActor*> FoundLandscapeProxies = PCGActorSelector::FindActors(ActorSelector, SourceComponent, BoundsCheck, SelfIgnoreCheck);
+		TArray<AActor*> FoundLandscapeProxies = PCGActorSelector::FindActors(ActorSelector, Context->ExecutionSource.Get(), BoundsCheck, SelfIgnoreCheck);
 
 		for (AActor* Proxy : FoundLandscapeProxies)
 		{
diff --git a/Engine/Plugins/PCG/Source/PCG/Private/Elements/PCGActorSelector.cpp b/Engine/Plugins/PCG/Source/PCG/Private/Elements/PCGActorSelector.cpp
--- a/Engine/Plugins/PCG/Source/PCG/Private/Elements/PCGActorSelector.cpp
+++ b/Engine/Plugins/PCG/Source/PCG/Private/Elements/PCGActorSelector.cpp
@@ -122,8 +122,7 @@ namespace PCGActorSelector
 		TRACE_CPUPROFILER_EVENT_SCOPE(PCGActorSelector::FindActor);
 
 		UWorld* World = InExecutionSource ? InExecutionSource->GetExecutionState().GetWorld() : nullptr;
-		const UPCGComponent* PCGComponent = Cast<UPCGComponent>(InExecutionSource);
-		AActor* Self = PCGComponent ? PCGComponent->GetOwner() : nullptr;
+		AActor* Self = InExecutionSource ? InExecutionSource->GetExecutionState().GetTypedTarget<AActor>() : nullptr;
 
 		TArray<AActor*> FoundActors;
 
@@ -251,8 +250,7 @@ namespace PCGActorSelector
 		case EPCGActorFilter::Original:
 		{
 			APCGPartitionActor* PartitionActor = Cast<APCGPartitionActor>(Self);
-			UPCGComponent* OriginalComponent = (PartitionActor && PCGComponent) ? PartitionActor->GetOriginalComponent(PCGComponent) : nullptr;
-			AActor* OriginalActor = OriginalComponent ? OriginalComponent->GetOwner() : nullptr;
+			AActor* OriginalActor = PartitionActor && InExecutionSource ? InExecutionSource->GetExecutionState().GetTypedTarget<AActor>() : nullptr;
 			if (OriginalActor)
 			{
 				ActorsToCheck.Add(OriginalActor);
@@ -451,15 +449,12 @@ bool FPCGSelectionKey::IsMatching(const FPCGSelectionKey& InSelectionKey, const
 			return false;
 		}
 
-		// @todo_pcg: support execution source
-		const UPCGComponent* PCGComponent = Cast<UPCGComponent>(InExecutionSource);
-
 		// InKey provide the info for selecting a given actor.
 		// We reconstruct the selector settings from this key, and we also force it to SelectMultiple, since
 		// we want to gather all the actors that matches this given key, to find if ours matches.
 		FPCGActorSelectorSettings SelectorSettings = FPCGActorSelectorSettings::ReconstructFromKey(*this);
 		SelectorSettings.bSelectMultiple = true;
-		TArray<AActor*> AllActors = PCGActorSelector::FindActors(SelectorSettings, PCGComponent, [](const AActor*) { return true; }, [](const AActor*) { return true; });
+		TArray<AActor*> AllActors = PCGActorSelector::FindActors(SelectorSettings, InExecutionSource, [](const AActor*) { return true; }, [](const AActor*) { return true; });
 		return AllActors.Contains(InActor);
 	}
 
diff --git a/Engine/Plugins/PCG/Source/PCG/Private/Elements/PCGAddComponent.cpp b/Engine/Plugins/PCG/Source/PCG/Private/Elements/PCGAddComponent.cpp
--- a/Engine/Plugins/PCG/Source/PCG/Private/Elements/PCGAddComponent.cpp
+++ b/Engine/Plugins/PCG/Source/PCG/Private/Elements/PCGAddComponent.cpp
@@ -2,8 +2,8 @@
 
 #include "Elements/PCGAddComponent.h"
 
-#include "PCGComponent.h"
 #include "PCGManagedResource.h"
+#include "PCGManagedResourceContainer.h"
 #include "Elements/Metadata/PCGMetadataElementCommon.h"
 #include "Helpers/PCGHelpers.h"
 #include "Metadata/Accessors/PCGAttributeAccessorHelpers.h"
@@ -244,8 +244,8 @@ bool FPCGAddComponentElement::ExecuteInternal(FPCGContext* InContext) const
 	TArray<FPCGTaggedData> TargetActors = Context->InputData.GetInputsByPin(PCGAddComponentConstants::TargetActorPinLabel);
 
 	// Early outs & validation
-	const UPCGComponent* SourceComponent = Cast<UPCGComponent>(Context->ExecutionSource.Get());
-	if (Inputs.Num() == 0 || !SourceComponent)
+	IPCGGraphExecutionSource* ExecutionSource = Context->ExecutionSource.Get();
+	if (Inputs.Num() == 0 || !ExecutionSource)
 	{
 		return true;
 	}
@@ -351,8 +351,14 @@ bool FPCGAddComponentElement::ExecuteInternal(FPCGContext* InContext) const
 		}
 		else
 		{
-			// Wrap on default actor namely the owner of the PCG component.
-			TargetActorAccessor = MakeUnique<FPCGConstantValueAccessor<FSoftObjectPath>>(FSoftObjectPath(SourceComponent->GetOwner()));
+			if (!ExecutionSource->GetExecutionState().GetTarget())
+			{
+				PCGE_LOG(Error, GraphAndLog, LOCTEXT("DefaultTargetSourceInvalid", "Unable to retrieve the owner of the execution source."));
+				continue;
+			}
+
+			// Wrap on the execution source target (in most cases, this will be the owner of the PCG component)
+			TargetActorAccessor = MakeUnique<FPCGConstantValueAccessor<FSoftObjectPath>>(FSoftObjectPath(ExecutionSource->GetExecutionState().GetTarget()));
 			TargetActorKeys = MakeUnique<FPCGAttributeAccessorKeysSingleObjectPtr<void>>();
 		}
 
@@ -418,9 +424,9 @@ void FPCGAddComponentElement::AddComponents(FPCGAddComponentContext* Context,
 	check(OutputComponentAccessor && OutputComponentKeys);
 
 	// TODO: implement component reuse
-	UPCGComponent* SourceComponent = Cast<UPCGComponent>(Context->ExecutionSource.Get());
-	check(SourceComponent);
-	const bool bIsInPreviewMode = SourceComponent->IsInPreviewMode();
+	IPCGGraphExecutionSource* ExecutionSource = Context->ExecutionSource.Get();
+	check(ExecutionSource);
+	const bool bIsInPreviewMode = ExecutionSource->GetExecutionState().IsInPreviewMode();
 	const EObjectFlags ObjectFlags = bIsInPreviewMode ? RF_Transient : RF_NoFlags;
 
 	// Similar approach to PCGMetadataCommon::ApplyOnAccessorRange
@@ -433,8 +439,16 @@ void FPCGAddComponentElement::AddComponents(FPCGAddComponentContext* Context,
 		return;
 	}
 
+	FPCGManagedResourceContainerHelper ManagedResourcesContainerHelper(ExecutionSource);
+	UObject* ExecutionSourceAsObject = Cast<UObject>(ExecutionSource);
+	if (!ManagedResourcesContainerHelper.IsValid() || !ExecutionSourceAsObject)
+	{
+		// Need resource container for tracking purposes
+		return;
+	}
+
 	// TODO: root location? (needed for reuse) + settings UID
-	UPCGManagedComponentList* Resource = NewObject<UPCGManagedComponentList>(SourceComponent);
+	UPCGManagedComponentList* Resource = NewObject<UPCGManagedComponentList>(ExecutionSourceAsObject);
 	check(Resource);
 	Resource->GeneratedComponents.Reserve(NumberOfEntries);
 
@@ -532,7 +546,11 @@ void FPCGAddComponentElement::AddComponents(FPCGAddComponentContext* Context,
 					SceneComponent->SetWorldTransform(Transform);
 				}
 
-				Component->ComponentTags.Add(SourceComponent->GetFName());
+				if (UObject* SourceObject = Cast<UObject>(ExecutionSource))
+				{
+					Component->ComponentTags.Add(SourceObject->GetFName());
+				}
+				
 				Component->ComponentTags.Add(PCGHelpers::DefaultPCGTag);
 
 				// Add component reference on output data
@@ -550,7 +568,7 @@ void FPCGAddComponentElement::AddComponents(FPCGAddComponentContext* Context,
 	// Finally add the resource to the component
 	if (!Resource->GeneratedComponents.IsEmpty())
 	{
-		SourceComponent->AddToManagedResources(Resource);
+		ManagedResourcesContainerHelper.AddManagedResource(Resource);
 	}	
 }
 
diff --git a/Engine/Plugins/PCG/Source/PCG/Private/Elements/PCGCreatePoints.cpp b/Engine/Plugins/PCG/Source/PCG/Private/Elements/PCGCreatePoints.cpp
--- a/Engine/Plugins/PCG/Source/PCG/Private/Elements/PCGCreatePoints.cpp
+++ b/Engine/Plugins/PCG/Source/PCG/Private/Elements/PCGCreatePoints.cpp
@@ -2,7 +2,6 @@
 
 #include "Elements/PCGCreatePoints.h"
 
-#include "PCGComponent.h"
 #include "PCGContext.h"
 #include "Helpers/PCGAsync.h"
 #include "Helpers/PCGBlueprintHelpers.h"
@@ -89,16 +88,14 @@ bool FPCGCreatePointsElement::ExecuteInternal(FPCGContext* Context) const
 		return true;
 	}
 
-	FTransform LocalTransform = Settings->CoordinateSpace == EPCGCoordinateSpace::World ? FTransform::Identity : Context->ExecutionSource->GetExecutionState().GetTransform();
-
-	if (Settings->CoordinateSpace == EPCGCoordinateSpace::OriginalComponent)
+	FTransform LocalTransform = FTransform::Identity;
+	if (Settings->CoordinateSpace == EPCGCoordinateSpace::LocalComponent)
 	{
-		UPCGComponent* SourceComponent = Cast<UPCGComponent>(Context->ExecutionSource.Get());
-		if (SourceComponent)
-		{
-			check(SourceComponent->GetOriginalComponent() && SourceComponent->GetOriginalComponent()->GetOwner());
-			LocalTransform = SourceComponent->GetOriginalComponent()->GetOwner()->GetActorTransform();
-		}
+		LocalTransform = Context->ExecutionSource->GetExecutionState().GetTransform();
+	}
+	else if (Settings->CoordinateSpace == EPCGCoordinateSpace::OriginalComponent)
+	{
+		LocalTransform = Context->ExecutionSource->GetExecutionState().GetOriginalTransform();
 	}
 
 	// Reset scale as we are not going to derive the points size from it
@@ -229,10 +226,9 @@ namespace PCGCreatePointsHelper
 			}
 			else if (InCoordinateSpace == EPCGCoordinateSpace::OriginalComponent)
 			{
-				const UPCGComponent* SourceComponent = Cast<UPCGComponent>(InExecutionSource);
-				if (SourceComponent)
+				if (IPCGGraphExecutionSource* OriginalExecutionSource = InExecutionSource ? InExecutionSource->GetExecutionState().GetOriginalSource() : nullptr)
 				{
-					return SourceComponent->GetOriginalActorPCGData();
+					return OriginalExecutionSource->GetExecutionState().GetSelfData();
 				}
 				else
 				{
diff --git a/Engine/Plugins/PCG/Source/PCG/Private/Elements/PCGCreatePointsGrid.cpp b/Engine/Plugins/PCG/Source/PCG/Private/Elements/PCGCreatePointsGrid.cpp
--- a/Engine/Plugins/PCG/Source/PCG/Private/Elements/PCGCreatePointsGrid.cpp
+++ b/Engine/Plugins/PCG/Source/PCG/Private/Elements/PCGCreatePointsGrid.cpp
@@ -2,7 +2,6 @@
 
 #include "Elements/PCGCreatePointsGrid.h"
 
-#include "PCGComponent.h"
 #include "PCGContext.h"
 #include "Data/PCGPointArrayData.h"
 #include "Elements/PCGCreatePoints.h"
@@ -73,16 +72,14 @@ bool FPCGCreatePointsGridElement::ExecuteInternal(FPCGContext* Context) const
 		return true;
 	}
 
-	FTransform LocalTransform = Settings->CoordinateSpace == EPCGCoordinateSpace::World ? FTransform::Identity : Context->ExecutionSource->GetExecutionState().GetTransform();
-
-	if (Settings->CoordinateSpace == EPCGCoordinateSpace::OriginalComponent)
+	FTransform LocalTransform = FTransform::Identity;
+	if (Settings->CoordinateSpace == EPCGCoordinateSpace::LocalComponent)
 	{
-		UPCGComponent* SourceComponent = Cast<UPCGComponent>(Context->ExecutionSource.Get());
-		if (SourceComponent)
-		{
-			check(SourceComponent->GetOriginalComponent() && SourceComponent->GetOriginalComponent()->GetOwner());
-			LocalTransform = SourceComponent->GetOriginalComponent()->GetOwner()->GetActorTransform();
-		}
+		LocalTransform = Context->ExecutionSource->GetExecutionState().GetTransform();
+	}
+	else if (Settings->CoordinateSpace == EPCGCoordinateSpace::OriginalComponent)
+	{
+		LocalTransform = Context->ExecutionSource->GetExecutionState().GetOriginalTransform();
 	}
 
 	// Reset the scale on the local transform since we don't want to derive the points scale from the referential
diff --git a/Engine/Plugins/PCG/Source/PCG/Private/Elements/PCGCreatePointsSphere.cpp b/Engine/Plugins/PCG/Source/PCG/Private/Elements/PCGCreatePointsSphere.cpp
--- a/Engine/Plugins/PCG/Source/PCG/Private/Elements/PCGCreatePointsSphere.cpp
+++ b/Engine/Plugins/PCG/Source/PCG/Private/Elements/PCGCreatePointsSphere.cpp
@@ -2,7 +2,6 @@
 
 #include "Elements/PCGCreatePointsSphere.h"
 
-#include "PCGComponent.h"
 #include "Data/PCGBasePointData.h"
 #include "Elements/PCGCreatePoints.h"
 #include "Helpers/PCGAsync.h"
@@ -393,6 +392,13 @@ FPCGElementPtr UPCGCreatePointsSphereSettings::CreateElement() const
 bool FPCGCreatePointsSphereElement::ExecuteInternal(FPCGContext* InContext) const
 {
 	TRACE_CPUPROFILER_EVENT_SCOPE(FPCGCreatePointsSphereElement::Execute);
+	check(InContext);
+
+	if (!InContext->ExecutionSource.IsValid())
+	{
+		return true;
+	}
+
 	const UPCGCreatePointsSphereSettings* Settings = InContext->GetInputSettings<UPCGCreatePointsSphereSettings>();
 
 	// Used for culling, regardless of generation coordinate space
@@ -462,16 +468,14 @@ bool FPCGCreatePointsSphereElement::ExecuteInternal(FPCGContext* InContext) cons
 	UPCGBasePointData* PointData = FPCGContext::NewPointData_AnyThread(InContext);
 	Output.Data = PointData;
 
-	FTransform LocalTransform = Settings->CoordinateSpace == EPCGCoordinateSpace::World ? FTransform::Identity : InContext->ExecutionSource->GetExecutionState().GetTransform();
-
-	if (Settings->CoordinateSpace == EPCGCoordinateSpace::OriginalComponent)
+	FTransform LocalTransform = FTransform::Identity;
+	if (Settings->CoordinateSpace == EPCGCoordinateSpace::LocalComponent)
 	{
-		UPCGComponent* SourceComponent = Cast<UPCGComponent>(InContext->ExecutionSource.Get());
-		if (SourceComponent)
-		{
-			check(SourceComponent->GetOriginalComponent() && SourceComponent->GetOriginalComponent()->GetOwner());
-			LocalTransform = SourceComponent->GetOriginalComponent()->GetOwner()->GetActorTransform();
-		}
+		LocalTransform = InContext->ExecutionSource->GetExecutionState().GetTransform();
+	}
+	else if (Settings->CoordinateSpace == EPCGCoordinateSpace::OriginalComponent)
+	{
+		LocalTransform = InContext->ExecutionSource->GetExecutionState().GetOriginalTransform();
 	}
 
 	// Point scale should not be relative to the coordinate space
diff --git a/Engine/Plugins/PCG/Source/PCG/Private/Elements/PCGCreateSpline.cpp b/Engine/Plugins/PCG/Source/PCG/Private/Elements/PCGCreateSpline.cpp
--- a/Engine/Plugins/PCG/Source/PCG/Private/Elements/PCGCreateSpline.cpp
+++ b/Engine/Plugins/PCG/Source/PCG/Private/Elements/PCGCreateSpline.cpp
@@ -3,9 +3,9 @@
 #include "Elements/PCGCreateSpline.h"
 
 #include "PCGCommon.h"
-#include "PCGComponent.h"
 #include "PCGContext.h"
 #include "PCGManagedResource.h"
+#include "PCGManagedResourceContainer.h"
 #include "PCGPin.h"
 #include "Data/PCGBasePointData.h"
 #include "Data/PCGSplineData.h"
@@ -207,10 +207,15 @@ bool FPCGCreateSplineElement::ExecuteInternal(FPCGContext* Context) const
 		if (Settings->Mode != EPCGCreateSplineMode::CreateDataOnly)
 		{
 			check(IsInGameThread());
-			if (UPCGComponent* SourceComponent = CastChecked<UPCGComponent>(Context->ExecutionSource.Get()))
+
+			IPCGGraphExecutionSource* ExecutionSource = Context->ExecutionSource.Get();
+			FPCGManagedResourceContainerHelper ManagedResourcesContainerHelper(ExecutionSource);
+			UObject* ExecutionSourceAsObject = Cast<UObject>(ExecutionSource);
+
+			if(ExecutionSourceAsObject && ManagedResourcesContainerHelper.IsValid())
 			{
 				SplineComponent = NewObject<USplineComponent>(SplineActor);
-				SplineComponent->ComponentTags.Add(SourceComponent->GetFName());
+				SplineComponent->ComponentTags.Add(ExecutionSourceAsObject->GetFName());
 				SplineComponent->ComponentTags.Add(PCGHelpers::DefaultPCGTag);
 				for (const FName& Tag : Settings->TagsToAddOnComponents)
 				{
@@ -227,10 +232,10 @@ bool FPCGCreateSplineElement::ExecuteInternal(FPCGContext* Context) const
 
 				SplineData->ApplyTo(SplineComponent);
 
-				UPCGManagedComponent* ManagedComponent = NewObject<UPCGManagedComponent>(SourceComponent);
+				UPCGManagedComponent* ManagedComponent = NewObject<UPCGManagedComponent>(ExecutionSourceAsObject);
 				ManagedComponent->GeneratedComponent = SplineComponent;
 
-				SourceComponent->AddToManagedResources(ManagedComponent);
+				ManagedResourcesContainerHelper.AddManagedResource(ManagedComponent);
 			}
 		}
 
diff --git a/Engine/Plugins/PCG/Source/PCG/Private/Elements/PCGCreateTargetActor.cpp b/Engine/Plugins/PCG/Source/PCG/Private/Elements/PCGCreateTargetActor.cpp
--- a/Engine/Plugins/PCG/Source/PCG/Private/Elements/PCGCreateTargetActor.cpp
+++ b/Engine/Plugins/PCG/Source/PCG/Private/Elements/PCGCreateTargetActor.cpp
@@ -2,14 +2,15 @@
 
 #include "Elements/PCGCreateTargetActor.h"
 
-#include "PCGComponent.h"
 #include "PCGContext.h"
 #include "PCGCrc.h"
 #include "PCGElement.h"
 #include "PCGManagedResource.h"
+#include "PCGManagedResourceContainer.h"
 #include "PCGModule.h"
 #include "PCGParamData.h"
 #include "Data/PCGPointData.h"
+#include "Graph/PCGStackContext.h"
 #include "Helpers/PCGActorHelpers.h"
 #include "Helpers/PCGDataLayerHelpers.h"
 #include "Helpers/PCGHelpers.h"
@@ -368,8 +369,8 @@ bool FPCGCreateTargetActorElement::ExecuteInternal(FPCGContext* Context) const
 		return true;
 	}
 
-	UPCGComponent* SourceComponent = Cast<UPCGComponent>(Context->ExecutionSource.Get());
-	const bool bHasAuthority = !SourceComponent || (SourceComponent->GetOwner() && SourceComponent->GetOwner()->HasAuthority());
+	IPCGGraphExecutionSource* ExecutionSource = Context->ExecutionSource.Get();
+	const bool bHasAuthority = !ExecutionSource || ExecutionSource->GetExecutionState().HasAuthority();
 	const bool bSpawnedActorRequiresAuthority = CastChecked<AActor>(Settings->TemplateActorClass->GetDefaultObject())->GetIsReplicated();
 
 	if (!bHasAuthority && bSpawnedActorRequiresAuthority)
@@ -392,9 +393,9 @@ bool FPCGCreateTargetActorElement::ExecuteInternal(FPCGContext* Context) const
 	}
 
 	UPCGActorHelpers::FSpawnDefaultActorParams SpawnDefaultActorParams(TargetActor->GetWorld(), Settings->TemplateActorClass, Transform, SpawnParams);
-	SpawnDefaultActorParams.bIsRuntime = (SourceComponent && SourceComponent->IsManagedByRuntimeGenSystem());
+	SpawnDefaultActorParams.bIsRuntime = (ExecutionSource && ExecutionSource->GetExecutionState().IsManagedByRuntimeGenSystem());
 #if WITH_EDITOR
-	SpawnDefaultActorParams.bIsPreviewActor = (SourceComponent && SourceComponent->IsInPreviewMode());
+	SpawnDefaultActorParams.bIsPreviewActor = (ExecutionSource && ExecutionSource->GetExecutionState().IsInPreviewMode());
 
 	int32 DataLayerCrc;
 	SpawnDefaultActorParams.DataLayerInstances = PCGDataLayerHelpers::GetDataLayerInstancesAndCrc(Context, Settings->DataLayerSettings, TargetActor, DataLayerCrc);
@@ -404,6 +405,7 @@ bool FPCGCreateTargetActorElement::ExecuteInternal(FPCGContext* Context) const
 	SpawnDefaultActorParams.HLODLayer = PCGHLODHelpers::GetHLODLayerAndCrc(Context, Settings->HLODSettings, TargetActor, TemplateActorOrDefault, HLODLayerCrc);
 #endif
 
+	FPCGManagedResourceContainerHelper ManagedResourcesContainerHelper(ExecutionSource);
 	UPCGManagedActors* ReusableResource = nullptr;
 	FPCGCrc ResourceCrc;
 
@@ -443,9 +445,9 @@ bool FPCGCreateTargetActorElement::ExecuteInternal(FPCGContext* Context) const
 			const FPCGCrc AdditionalCRC = PCGCreateTargetActor::GetAdditionalDependenciesCrc(SpawnDefaultActorParams, TargetActor);
 			ResourceCrc.Combine(AdditionalCRC);
 
-			if (SourceComponent)
+			if (ManagedResourcesContainerHelper.IsValid())
 			{
-				SourceComponent->ForEachManagedResource([&ReusableResource, ResourceCrc, &Context, bIsPreview = SpawnDefaultActorParams.bIsPreviewActor](UPCGManagedResource* InResource)
+				ManagedResourcesContainerHelper.ForEachManagedResource([&ReusableResource, ResourceCrc, &Context, bIsPreview = SpawnDefaultActorParams.bIsPreviewActor](UPCGManagedResource* InResource)
 				{
 					if (ReusableResource)
 					{
@@ -530,9 +532,9 @@ bool FPCGCreateTargetActorElement::ExecuteInternal(FPCGContext* Context) const
 	// Create Resource if it isn't reused
 	if (!ReusableResource)
 	{
-		if (SourceComponent)
+		if(UObject* ExecutionSourceObject = Cast<UObject>(ExecutionSource); ExecutionSourceObject && ManagedResourcesContainerHelper.IsValid())
 		{
-			UPCGManagedActors* ManagedActors = NewObject<UPCGManagedActors>(SourceComponent);
+			UPCGManagedActors* ManagedActors = NewObject<UPCGManagedActors>(ExecutionSourceObject);
 			if (ResourceCrc.IsValid())
 			{
 				ManagedActors->SetCrc(ResourceCrc);
@@ -543,7 +545,7 @@ bool FPCGCreateTargetActorElement::ExecuteInternal(FPCGContext* Context) const
 			ManagedActors->GetMutableGeneratedActors().AddUnique(GeneratedActor);
 			ManagedActors->bSupportsReset = !Settings->bDeleteActorsBeforeGeneration;
 
-			SourceComponent->AddToManagedResources(ManagedActors);
+			ManagedResourcesContainerHelper.AddManagedResource(ManagedActors);
 		}
 	}
 
diff --git a/Engine/Plugins/PCG/Source/PCG/Private/Elements/PCGDataFromActor.cpp b/Engine/Plugins/PCG/Source/PCG/Private/Elements/PCGDataFromActor.cpp
--- a/Engine/Plugins/PCG/Source/PCG/Private/Elements/PCGDataFromActor.cpp
+++ b/Engine/Plugins/PCG/Source/PCG/Private/Elements/PCGDataFromActor.cpp
@@ -367,8 +367,7 @@ bool FPCGDataFromActorElement::ExecuteInternal(FPCGContext* InContext) const
 		return true;
 	}
 
-	UPCGComponent* PCGComponent = Cast<UPCGComponent>(ExecutionSource);
-	AActor* Self = PCGComponent ? PCGComponent->GetOwner() : nullptr;
+	AActor* Self = ExecutionSource ? ExecutionSource->GetExecutionState().GetTypedTarget<AActor>() : nullptr;
 
 	if (!Context->bPerformedQuery)
 	{
@@ -598,6 +597,7 @@ bool FPCGDataFromActorElement::ExecuteInternal(FPCGContext* InContext) const
 	if (Context->bPerformedQuery)
 	{
 #if WITH_EDITOR
+		UPCGComponent* PCGComponent = Cast<UPCGComponent>(ExecutionSource);
 		// Remove ignored change origins now that we've completed the wait tasks.
 		UPCGComponent* OriginalComponent = PCGComponent ? PCGComponent->GetOriginalComponent() : nullptr;
 		if (OriginalComponent)
diff --git a/Engine/Plugins/PCG/Source/PCG/Private/Elements/PCGSkinnedMeshSpawner.cpp b/Engine/Plugins/PCG/Source/PCG/Private/Elements/PCGSkinnedMeshSpawner.cpp
--- a/Engine/Plugins/PCG/Source/PCG/Private/Elements/PCGSkinnedMeshSpawner.cpp
+++ b/Engine/Plugins/PCG/Source/PCG/Private/Elements/PCGSkinnedMeshSpawner.cpp
@@ -2,9 +2,9 @@
 
 #include "Elements/PCGSkinnedMeshSpawner.h"
 
-#include "PCGComponent.h"
 #include "PCGCustomVersion.h"
 #include "PCGManagedResource.h"
+#include "PCGManagedResourceContainer.h"
 #include "Compute/PCGKernelHelpers.h"
 #include "Compute/BuiltInKernels/PCGCountUniqueAttributeValuesKernel.h"
 #include "Data/PCGPointData.h"
@@ -109,11 +109,8 @@ bool FPCGSkinnedMeshSpawnerElement::PrepareDataInternal(FPCGContext* InContext)
 		return true;
 	}
 
-	UPCGComponent* SourceComponent = Cast<UPCGComponent>(Context->ExecutionSource.Get());
-	if (!SourceComponent)
-	{
-		return true;
-	}
+	IPCGGraphExecutionSource* ExecutionSource = Context->ExecutionSource.Get();
+	FPCGManagedResourceContainerHelper ManagedResourcesContainerHelper(ExecutionSource);
 
 #if WITH_EDITOR
 	// In editor, we always want to generate this data for inspection & to prevent caching issues
@@ -133,10 +130,10 @@ bool FPCGSkinnedMeshSpawnerElement::PrepareDataInternal(FPCGContext* InContext)
 			GetDependenciesCrc(FPCGGetDependenciesCrcParams(&Context->InputData, Settings, Context->ExecutionSource.Get()), Context->DependenciesCrc);
 		}
 		
-		if (Context->DependenciesCrc.IsValid())
+		if (Context->DependenciesCrc.IsValid() && ManagedResourcesContainerHelper.IsValid())
 		{
 			TArray<UPCGManagedISKMComponent*> MISKMCs;
-			SourceComponent->ForEachManagedResource([&MISKMCs, &Context, Settings](UPCGManagedResource* InResource)
+			ManagedResourcesContainerHelper.ForEachManagedResource([&MISKMCs, &Context, Settings](UPCGManagedResource* InResource)
 			{
 				if (UPCGManagedISKMComponent* Resource = Cast<UPCGManagedISKMComponent>(InResource))
 				{
@@ -457,8 +454,7 @@ void FPCGSkinnedMeshSpawnerElement::SpawnSkinnedMeshInstances(FPCGSkinnedMeshSpa
 	Params.SettingsCrc = Settings->GetSettingsCrc();
 	ensure(Params.SettingsCrc.IsValid());
 
-	UPCGComponent* SourceComponent = Cast<UPCGComponent>(Context->ExecutionSource.Get());
-	UPCGManagedISKMComponent* MISKMC = UPCGActorHelpers::GetOrCreateManagedABMC(TargetActor, SourceComponent, Params, Context);
+	UPCGManagedISKMComponent* MISKMC = UPCGActorHelpers::GetOrCreateManagedABMC(TargetActor, Context->ExecutionSource.Get(), Params, Context);
 
 	check(MISKMC);
 	MISKMC->SetCrc(Context->DependenciesCrc);
diff --git a/Engine/Plugins/PCG/Source/PCG/Private/Elements/PCGSpawnActor.cpp b/Engine/Plugins/PCG/Source/PCG/Private/Elements/PCGSpawnActor.cpp
--- a/Engine/Plugins/PCG/Source/PCG/Private/Elements/PCGSpawnActor.cpp
+++ b/Engine/Plugins/PCG/Source/PCG/Private/Elements/PCGSpawnActor.cpp
@@ -471,18 +471,21 @@ bool FPCGSpawnActorElement::SpawnAndPrepareSubgraphs(FPCGSubgraphContext* Contex
 			return true;
 		}
 	}
-	UPCGComponent* SourceComponent = Cast<UPCGComponent>(Context->ExecutionSource.Get());
-	if (!SourceComponent)
+
+	IPCGGraphExecutionSource* ExecutionSource = Context->ExecutionSource.Get();
+	if (!ExecutionSource)
 	{
-		PCGE_LOG(Error, GraphAndLog, LOCTEXT("InvalidSourceComponent", "Invalid SourceComponent."));
+		PCGE_LOG(Error, GraphAndLog, LOCTEXT("InvalidSource", "Invalid Source."));
 		return true;
 	}
 
+	FPCGManagedResourceContainerHelper ManagedResourcesContainerHelper(ExecutionSource);
+
 	// Check if we can reuse existing resources - note that this is done on a per-settings basis when collapsed,
 	// Otherwise we'll check against merged crc 
 	bool bFullySkippedDueToReuse = false;
 
-	if (CVarAllowActorReuse.GetValueOnAnyThread())
+	if (CVarAllowActorReuse.GetValueOnAnyThread() && ManagedResourcesContainerHelper.IsValid())
 	{
 		// Compute CRC if it has not been computed (it likely isn't, but this is to futureproof this)
 		if (!Context->DependenciesCrc.IsValid())
@@ -495,7 +498,7 @@ bool FPCGSpawnActorElement::SpawnAndPrepareSubgraphs(FPCGSubgraphContext* Contex
 			if (Settings->Option == EPCGSpawnActorOption::CollapseActors)
 			{
 				TArray<UPCGManagedISMComponent*> MISMCs;
-				SourceComponent->ForEachManagedResource([&MISMCs, &Context](UPCGManagedResource* InResource)
+				ManagedResourcesContainerHelper.ForEachManagedResource([&MISMCs, &Context](UPCGManagedResource* InResource)
 				{
 					if (UPCGManagedISMComponent* Resource = Cast<UPCGManagedISMComponent>(InResource))
 					{
@@ -534,7 +537,7 @@ bool FPCGSpawnActorElement::SpawnAndPrepareSubgraphs(FPCGSubgraphContext* Contex
 	const bool bGenerateOutputsWithActorReference = (Settings->Option != EPCGSpawnActorOption::CollapseActors) && Context->Node && Context->Node->IsOutputPinConnected(PCGPinConstants::DefaultOutputLabel);
 #endif
 
-	const bool bHasAuthority = SourceComponent->GetOwner() && SourceComponent->GetOwner()->HasAuthority();
+	const bool bHasAuthority = ExecutionSource->GetExecutionState().HasAuthority();
 
 	TArray<FPCGTaggedData> Inputs = Context->InputData.GetInputsByPin(PCGPinConstants::DefaultInputLabel);
 	for (const FPCGTaggedData& Input : Inputs)
@@ -697,6 +700,12 @@ void FPCGSpawnActorElement::CollapseIntoTargetActor(FPCGSubgraphContext* Context
 		return;
 	}
 
+	IPCGGraphExecutionSource* ExecutionSource = Context->ExecutionSource.Get();
+	if (!ExecutionSource)
+	{
+		return;
+	}
+
 	const UPCGSpawnActorSettings* Settings = Context->GetInputSettings<UPCGSpawnActorSettings>();
 	check(Settings);
 	
@@ -742,13 +751,7 @@ void FPCGSpawnActorElement::CollapseIntoTargetActor(FPCGSubgraphContext* Context
 	{
 		const FPCGISMComponentBuilderParams& ISMCParams = ISMCBuilderTransforms.Key;
 
-		UPCGComponent* SourceComponent = Cast<UPCGComponent>(Context->ExecutionSource.Get());
-		if (!SourceComponent)
-		{
-			continue;
-		}
-
-		UPCGManagedISMComponent* MISMC = UPCGActorHelpers::GetOrCreateManagedISMC(TargetActor, SourceComponent, ISMCParams, Context);
+		UPCGManagedISMComponent* MISMC = UPCGActorHelpers::GetOrCreateManagedISMC(TargetActor, ExecutionSource, ISMCParams, Context);
 		if (!MISMC)
 		{
 			continue;
@@ -860,13 +863,18 @@ void FPCGSpawnActorElement::SpawnActors(FPCGSubgraphContext* Context, AActor* Ta
 	SpawnParams.SpawnCollisionHandlingOverride = ESpawnActorCollisionHandlingMethod::AlwaysSpawn;
 	SpawnParams.OverrideLevel = TargetActor->GetLevel();
 
-	UPCGComponent* SourceComponent = Cast<UPCGComponent>(Context->ExecutionSource.Get());
+	IPCGGraphExecutionSource* Source = Context->ExecutionSource.Get();
+
+	if (!Source)
+	{
+		return; // won't be able to create resources correctly
+	}
 
 	UPCGActorHelpers::FSpawnDefaultActorParams SpawnDefaultActorParams(TargetActor->GetWorld(), InTemplateActorClass, FTransform::Identity, SpawnParams);
 	SpawnDefaultActorParams.bForceStaticMobility = false; // Always respect the actor's mobility
-	SpawnDefaultActorParams.bIsRuntime = (SourceComponent && SourceComponent->IsManagedByRuntimeGenSystem());
+	SpawnDefaultActorParams.bIsRuntime = (Source && Source->GetExecutionState().IsManagedByRuntimeGenSystem());
 #if WITH_EDITOR
-	SpawnDefaultActorParams.bIsPreviewActor = (SourceComponent && SourceComponent->IsInPreviewMode());
+	SpawnDefaultActorParams.bIsPreviewActor = (Source && Source->GetExecutionState().IsInPreviewMode());
 
 	int32 DataLayerCrc;
 	SpawnDefaultActorParams.DataLayerInstances = PCGDataLayerHelpers::GetDataLayerInstancesAndCrc(Context, Settings->DataLayerSettings, TargetActor, DataLayerCrc);
@@ -882,7 +890,14 @@ void FPCGSpawnActorElement::SpawnActors(FPCGSubgraphContext* Context, AActor* Ta
 	const bool bOnLoadCallGenerate = (Settings->GenerationTrigger == EPCGSpawnActorGenerationTrigger::Default ||
 		Settings->GenerationTrigger == EPCGSpawnActorGenerationTrigger::DoNotGenerateInEditor);
 #endif
-	UPCGSubsystem* Subsystem = UWorld::GetSubsystem<UPCGSubsystem>(SourceComponent->GetWorld());
+	UPCGSubsystem* Subsystem = UWorld::GetSubsystem<UPCGSubsystem>(Source ? Source->GetExecutionState().GetWorld() : nullptr);
+
+	FPCGManagedResourceContainerHelper ManagedResourcesContainerHelper(Source);
+	
+	if (!ManagedResourcesContainerHelper.IsValid())
+	{
+		return; // won't be able to manage resources correctly
+	}
 
 	// Try to reuse actors if they are preexisting
 	UPCGManagedActors* ReusedManagedActorsResource = nullptr;
@@ -911,7 +926,7 @@ void FPCGSpawnActorElement::SpawnActors(FPCGSubgraphContext* Context, AActor* Ta
 
 		if (InputDependenciesCrc.IsValid())
 		{
-			SourceComponent->ForEachManagedResource([&ReusedManagedActorsResource, &InputDependenciesCrc, &Context, NumPoints = PointData->GetNumPoints(), bIsPreview = SpawnDefaultActorParams.bIsPreviewActor](UPCGManagedResource* InResource)
+			ManagedResourcesContainerHelper.ForEachManagedResource([&ReusedManagedActorsResource, &InputDependenciesCrc, &Context, NumPoints = PointData->GetNumPoints(), bIsPreview = SpawnDefaultActorParams.bIsPreviewActor](UPCGManagedResource* InResource)
 			{
 				if (ReusedManagedActorsResource)
 				{
@@ -979,7 +994,7 @@ void FPCGSpawnActorElement::SpawnActors(FPCGSubgraphContext* Context, AActor* Ta
 		TArray<FName> NewActorTags = GetNewActorTags(Context, TargetActor, Settings->bInheritActorTags, Settings->TagsToAddOnActors);
 
 		// Create managed resource for actor tracking
-		UPCGManagedActors* ManagedActors = NewObject<UPCGManagedActors>(SourceComponent);
+		UPCGManagedActors* ManagedActors = NewObject<UPCGManagedActors>(Cast<UObject>(Source));
 #if WITH_EDITOR
 		ManagedActors->SetIsPreview(SpawnDefaultActorParams.bIsPreviewActor);
 #endif
@@ -1073,11 +1088,13 @@ void FPCGSpawnActorElement::SpawnActors(FPCGSubgraphContext* Context, AActor* Ta
 			PCGE_LOG(Error, GraphAndLog, LOCTEXT("ActorOverridesFailed", "At least one actor property override failed."));
 		}
 
-		SourceComponent->AddToManagedResources(ManagedActors);
+		ManagedResourcesContainerHelper.AddManagedResource(ManagedActors);
 
 		PCGE_LOG(Verbose, LogOnly, FText::Format(LOCTEXT("GenerationInfo", "Generated {0} actors"), PointData->GetNumPoints()));
 	}
 
+	UPCGComponent* SourceComponent = Cast<UPCGComponent>(Source);
+
 	// Setup & Generate on PCG components if needed
 	for (AActor* Actor : ProcessedActors)
 	{
@@ -1088,7 +1105,7 @@ void FPCGSpawnActorElement::SpawnActors(FPCGSubgraphContext* Context, AActor* Ta
 		{
 #if WITH_EDITOR
 			// For both pre-existing and new actors, we need to make sure we're inline with loading/generation as needed
-			if (PCGComponent->GetEditingMode() != SourceComponent->GetEditingMode())
+			if (SourceComponent && PCGComponent->GetEditingMode() != SourceComponent->GetEditingMode())
 			{
 				PCGComponent->SetEditingMode(/*CurrentEditingMode=*/SourceComponent->GetEditingMode(), /*SerializedEditingMode=*/SourceComponent->GetEditingMode());
 				PCGComponent->ChangeTransientState(SourceComponent->GetEditingMode());
@@ -1134,6 +1151,7 @@ TArray<FName> FPCGSpawnActorElement::GetNewActorTags(FPCGContext* Context, AActo
 		// and find the original actor tags
 		if (APCGPartitionActor* PartitionActor = Cast<APCGPartitionActor>(TargetActor))
 		{
+			// @todo_pcg review for execution source
 			if (UPCGComponent* SourceComponent = Cast<UPCGComponent>(Context->ExecutionSource.Get()))
 			{
 				if (UPCGComponent* OriginalComponent = PartitionActor->GetOriginalComponent(SourceComponent))
diff --git a/Engine/Plugins/PCG/Source/PCG/Private/Elements/PCGSpawnSpline.cpp b/Engine/Plugins/PCG/Source/PCG/Private/Elements/PCGSpawnSpline.cpp
--- a/Engine/Plugins/PCG/Source/PCG/Private/Elements/PCGSpawnSpline.cpp
+++ b/Engine/Plugins/PCG/Source/PCG/Private/Elements/PCGSpawnSpline.cpp
@@ -2,9 +2,9 @@
 
 #include "Elements/PCGSpawnSpline.h"
 
-#include "PCGComponent.h"
 #include "PCGContext.h"
 #include "PCGManagedResource.h"
+#include "PCGManagedResourceContainer.h"
 #include "Data/PCGSplineData.h"
 #include "Elements/Metadata/PCGMetadataElementCommon.h"
 #include "Helpers/PCGHelpers.h"
@@ -53,8 +53,11 @@ bool FPCGSpawnSplineElement::ExecuteInternal(FPCGContext* InContext) const
 	const UPCGSpawnSplineSettings* Settings = InContext->GetInputSettings<UPCGSpawnSplineSettings>();
 	check(Settings);
 
-	UPCGComponent* SourceComponent = Cast<UPCGComponent>(InContext->ExecutionSource.Get());
-	if (!*Settings->SplineComponent || !SourceComponent)
+	IPCGGraphExecutionSource* ExecutionSource = InContext->ExecutionSource.Get();
+	UObject* ExecutionSourceObject = Cast<UObject>(ExecutionSource);
+	FPCGManagedResourceContainerHelper ManagedResourcesContainerHelper(ExecutionSource);
+
+	if (!*Settings->SplineComponent || !ExecutionSource || !ManagedResourcesContainerHelper.IsValid())
 	{
 		return true;
 	}
@@ -117,9 +120,9 @@ bool FPCGSpawnSplineElement::ExecuteInternal(FPCGContext* InContext) const
 
 		InSplineData->ApplyTo(SplineComponent);
 
-		UPCGManagedComponent* ManagedComponent = NewObject<UPCGManagedComponent>(SourceComponent);
+		UPCGManagedComponent* ManagedComponent = NewObject<UPCGManagedComponent>(ExecutionSourceObject);
 		ManagedComponent->GeneratedComponent = SplineComponent;
-		SourceComponent->AddToManagedResources(ManagedComponent);
+		ManagedResourcesContainerHelper.AddManagedResource(ManagedComponent);
 
 		if (!Settings->PropertyOverrideDescriptions.IsEmpty())
 		{
@@ -137,7 +140,11 @@ bool FPCGSpawnSplineElement::ExecuteInternal(FPCGContext* InContext) const
 
 		// Add tags, both the default ones (source name + PCG Generated) and the ones from the override pin.
 		// We do this after the override in the case that the override was writing the tags.
-		SplineComponent->ComponentTags.Add(SourceComponent->GetFName());
+		if (ExecutionSourceObject)
+		{
+			SplineComponent->ComponentTags.Add(ExecutionSourceObject->GetFName());
+		}
+		
 		SplineComponent->ComponentTags.Add(PCGHelpers::DefaultPCGTag);
 		for (const FName& Tag : Settings->TagsToAddOnComponents)
 		{
diff --git a/Engine/Plugins/PCG/Source/PCG/Private/Elements/PCGSpawnSplineMesh.cpp b/Engine/Plugins/PCG/Source/PCG/Private/Elements/PCGSpawnSplineMesh.cpp
--- a/Engine/Plugins/PCG/Source/PCG/Private/Elements/PCGSpawnSplineMesh.cpp
+++ b/Engine/Plugins/PCG/Source/PCG/Private/Elements/PCGSpawnSplineMesh.cpp
@@ -2,7 +2,6 @@
 
 #include "Elements/PCGSpawnSplineMesh.h"
 
-#include "PCGComponent.h"
 #include "PCGContext.h"
 #include "PCGManagedResource.h"
 #include "PCGPin.h"
@@ -178,13 +177,13 @@ bool FPCGSpawnSplineMeshElement::ExecuteInternal(FPCGContext* InContext) const
 		return true;
 	}
 
-	UPCGComponent* SourceComponent = Cast<UPCGComponent>(Context->ExecutionSource.Get());
-	if (!SourceComponent)
+	IPCGGraphExecutionSource* ExecutionSource = Context->ExecutionSource.Get();
+	if (!ExecutionSource)
 	{
 		return true;
 	}
 
-	return ExecuteSlice(Context, [this, Settings, SourceComponent](ContextType* Context, const ExecStateType& ExecState, IterStateType& IterState, const uint32 IterIndex)
+	return ExecuteSlice(Context, [this, Settings](ContextType* Context, const ExecStateType& ExecState, IterStateType& IterState, const uint32 IterIndex)
 	{
 		if (Context->GetIterationStateResult(IterIndex) != EPCGTimeSliceInitResult::Success)
 		{
@@ -306,7 +305,10 @@ bool FPCGSpawnSplineMeshElement::ExecuteInternal(FPCGContext* InContext) const
 			IterState.SMCBuilderParams.SettingsCrc = Settings->GetSettingsCrc();
 			ensure(IterState.SMCBuilderParams.SettingsCrc.IsValid());
 
-			if (USplineMeshComponent* SplineMeshComponent = UPCGActorHelpers::GetOrCreateSplineMeshComponent(ExecState.TargetActor, SourceComponent, IterState.SMCBuilderParams, Context))
+			IPCGGraphExecutionSource* ExecutionSource = Context->ExecutionSource.Get();
+			check(ExecutionSource); // not likely to happen since it was tested at the beginning of the execute.
+
+			if (USplineMeshComponent* SplineMeshComponent = UPCGActorHelpers::GetOrCreateSplineMeshComponent(ExecState.TargetActor, ExecutionSource, IterState.SMCBuilderParams, Context))
 			{
 				// Apply the overrides on the spawned component
 				if (IterState.ComponentOverrides.IsValid())
@@ -315,7 +317,11 @@ bool FPCGSpawnSplineMeshElement::ExecuteInternal(FPCGContext* InContext) const
 					IterState.ComponentOverrides.Apply(IterState.ElementIndex);
 
 					// Since the component overrides can override the component tags, we need to ensure that the default tags are still present
-					SplineMeshComponent->ComponentTags.AddUnique(SourceComponent->GetFName());
+					if (UObject* ExecutionSourceObject = Cast<UObject>(ExecutionSource))
+					{
+						SplineMeshComponent->ComponentTags.AddUnique(ExecutionSourceObject->GetFName());
+					}
+					
 					SplineMeshComponent->ComponentTags.AddUnique(PCGHelpers::DefaultPCGTag);
 				}
 
diff --git a/Engine/Plugins/PCG/Source/PCG/Private/Elements/PCGStaticMeshSpawner.cpp b/Engine/Plugins/PCG/Source/PCG/Private/Elements/PCGStaticMeshSpawner.cpp
--- a/Engine/Plugins/PCG/Source/PCG/Private/Elements/PCGStaticMeshSpawner.cpp
+++ b/Engine/Plugins/PCG/Source/PCG/Private/Elements/PCGStaticMeshSpawner.cpp
@@ -2,9 +2,9 @@
 
 #include "Elements/PCGStaticMeshSpawner.h"
 
-#include "PCGComponent.h"
 #include "PCGCustomVersion.h"
 #include "PCGManagedResource.h"
+#include "PCGManagedResourceContainer.h"
 #include "Compute/PCGComputeCommon.h"
 #include "Compute/PCGKernelHelpers.h"
 #include "Compute/BuiltInKernels/PCGCountUniqueAttributeValuesKernel.h"
@@ -283,12 +283,14 @@ bool FPCGStaticMeshSpawnerElement::PrepareDataInternal(FPCGContext* InContext) c
 		return true;
 	}
 
-	UPCGComponent* SourceComponent = Cast<UPCGComponent>(Context->ExecutionSource.Get());
-	if (!SourceComponent)
+	IPCGGraphExecutionSource* ExecutionSource = Context->ExecutionSource.Get();
+	if (!ExecutionSource)
 	{
 		return true;
 	}
 
+	FPCGManagedResourceContainerHelper ManagedResourcesContainerHelper(ExecutionSource);
+
 #if WITH_EDITOR
 	// In editor, we always want to generate this data for inspection & to prevent caching issues
 	const bool bGenerateOutput = true;
@@ -299,7 +301,7 @@ bool FPCGStaticMeshSpawnerElement::PrepareDataInternal(FPCGContext* InContext) c
 	// Check if we can reuse existing resources on an "all-or-nothing" basis
 	bool& bSkippedDueToReuse = Context->bSkippedDueToReuse;
 
-	if (!Context->bReuseCheckDone && PCGStaticMeshSpawner::CVarAllowISMReuse.GetValueOnAnyThread())
+	if (!Context->bReuseCheckDone && PCGStaticMeshSpawner::CVarAllowISMReuse.GetValueOnAnyThread() && ManagedResourcesContainerHelper.IsValid())
 	{
 		// Compute CRC if it has not been computed (it likely isn't, but this is to futureproof this)
 		if (!Context->DependenciesCrc.IsValid())
@@ -322,7 +324,7 @@ bool FPCGStaticMeshSpawnerElement::PrepareDataInternal(FPCGContext* InContext) c
 		{
 			int32 CrcReuseCount = 0;
 			TArray<UPCGManagedISMComponent*> MISMCs;
-			SourceComponent->ForEachManagedResource([&MISMCs, &Context, Settings, &CrcReuseCount](UPCGManagedResource* InResource)
+			ManagedResourcesContainerHelper.ForEachManagedResource([&MISMCs, &Context, Settings, &CrcReuseCount](UPCGManagedResource* InResource)
 			{
 				if (UPCGManagedISMComponent* Resource = Cast<UPCGManagedISMComponent>(InResource))
 				{
@@ -414,7 +416,7 @@ bool FPCGStaticMeshSpawnerElement::PrepareDataInternal(FPCGContext* InContext) c
 			// Prior to selection, if we have MISMC resources that have the same settings crc & the same data crc,
 			// We can skip the selection here, unless we need to output the data (again, could/should be cached instead).
 			FPCGCrc DataCrc;
-			if(!bSkippedDueToReuse && PCGStaticMeshSpawner::CVarAllowISMReuse.GetValueOnAnyThread())
+			if(!bSkippedDueToReuse && PCGStaticMeshSpawner::CVarAllowISMReuse.GetValueOnAnyThread() && ManagedResourcesContainerHelper.IsValid())
 			{
 				ensure(Context->DependenciesCrc.IsValid()); // should have been done earlier
 
@@ -437,7 +439,7 @@ bool FPCGStaticMeshSpawnerElement::PrepareDataInternal(FPCGContext* InContext) c
 
 				int32 DataCrcReuseCount = 0;
 				TArray<UPCGManagedISMComponent*> MISMCs;
-				SourceComponent->ForEachManagedResource([&MISMCs, &SettingsCrc, &DataCrc, &Context, &DataCrcReuseCount](UPCGManagedResource* InResource)
+				ManagedResourcesContainerHelper.ForEachManagedResource([&MISMCs, &SettingsCrc, &DataCrc, &Context, &DataCrcReuseCount](UPCGManagedResource* InResource)
 				{
 					if (UPCGManagedISMComponent* Resource = Cast<UPCGManagedISMComponent>(InResource))
 					{
@@ -809,13 +811,13 @@ void FPCGStaticMeshSpawnerElement::SpawnStaticMeshInstances(FPCGStaticMeshSpawne
 
 	Params.bAllowDescriptorChanges = Settings->bAllowDescriptorChanges;
 
-	UPCGComponent* SourceComponent = Cast<UPCGComponent>(Context->ExecutionSource.Get());
-	if (!ensure(SourceComponent))
+	IPCGGraphExecutionSource* ExecutionSource = Context->ExecutionSource.Get();
+	if (!ensure(ExecutionSource))
 	{
 		return;
 	}
 
-	UPCGManagedISMComponent* MISMC = UPCGActorHelpers::GetOrCreateManagedISMC(TargetActor, SourceComponent, Params, Context);
+	UPCGManagedISMComponent* MISMC = UPCGActorHelpers::GetOrCreateManagedISMC(TargetActor, ExecutionSource, Params, Context);
 
 	check(MISMC);
 
diff --git a/Engine/Plugins/PCG/Source/PCG/Private/Elements/PCGWorldQuery.cpp b/Engine/Plugins/PCG/Source/PCG/Private/Elements/PCGWorldQuery.cpp
--- a/Engine/Plugins/PCG/Source/PCG/Private/Elements/PCGWorldQuery.cpp
+++ b/Engine/Plugins/PCG/Source/PCG/Private/Elements/PCGWorldQuery.cpp
@@ -2,7 +2,6 @@
 
 #include "Elements/PCGWorldQuery.h"
 
-#include "PCGComponent.h"
 #include "PCGContext.h"
 #include "PCGSubsystem.h"
 #include "Grid/PCGLandscapeCache.h"
@@ -250,8 +249,7 @@ bool FPCGWorldRayHitQueryElement::ExecuteInternal(FPCGContext* Context) const
 			{
 				BoundsCheck = [Data, ExecutionSource](const AActor* OtherActor) -> bool
 				{
-					// @todo_pcg: adapt to IPCGGraphExecutionSource (currently only for partitioning which isn't supported by other types than UPCGComponent)
-					const FBox OtherActorBounds = OtherActor ? PCGHelpers::GetGridBounds(OtherActor, Cast<UPCGComponent>(ExecutionSource)) : FBox(EForceInit::ForceInit);
+					const FBox OtherActorBounds = OtherActor ? PCGHelpers::GetGridBounds(OtherActor, ExecutionSource) : FBox(EForceInit::ForceInit);
 					return OtherActorBounds.Intersect(Data->Bounds);
 				};
 			}
diff --git a/Engine/Plugins/PCG/Source/PCG/Private/Graph/PCGGraphCache.h b/Engine/Plugins/PCG/Source/PCG/Private/Graph/PCGGraphCache.h
--- a/Engine/Plugins/PCG/Source/PCG/Private/Graph/PCGGraphCache.h
+++ b/Engine/Plugins/PCG/Source/PCG/Private/Graph/PCGGraphCache.h
@@ -12,7 +12,6 @@
 #include "UObject/ObjectKey.h"
 
 class IPCGElement;
-class UPCGComponent;
 class UPCGNode;
 class UPCGSettings;
 
