[PCG Tests] Unit Tests for Generalized Component Data Storage

#Synced-CL 51342292

diff --git a/Engine/Plugins/PCG/Source/PCG/Private/Tests/ExecutionSource/PCGSourceDataContainerTest.cpp b/Engine/Plugins/PCG/Source/PCG/Private/Tests/ExecutionSource/PCGSourceDataContainerTest.cpp
new file mode 100644
--- /dev/null
+++ b/Engine/Plugins/PCG/Source/PCG/Private/Tests/ExecutionSource/PCGSourceDataContainerTest.cpp
@@ -0,0 +1,376 @@
+// Copyright Epic Games, Inc. All Rights Reserved.
+
+#include "Graph/PCGSourceDataContainer.h"
+#include "Tests/PCGTestsCommon.h"
+
+#if WITH_EDITOR
+
+// -------------------------------------------------------------------
+// Store, Get, GetMutable, Remove — basic editor API
+// -------------------------------------------------------------------
+
+IMPLEMENT_CUSTOM_SIMPLE_AUTOMATION_TEST(FPCGSourceDataContainerStoreAndGetTest, FPCGTestBaseClass, "Plugins.PCG.SourceDataContainer.StoreAndGet", PCGTestsCommon::TestFlags)
+
+bool FPCGSourceDataContainerStoreAndGetTest::RunTest(const FString& Parameters)
+{
+	const FPCGSourceDataStorageKey KeyA(TEXT("LabelA"), 42);
+	const FPCGSourceDataStorageKey KeyB(TEXT("LabelB"), 99);
+	const FPCGSourceDataStorageKey PayloadA(TEXT("PayloadA"), 100);
+	const FPCGSourceDataStorageKey PayloadB(TEXT("PayloadB"), 200);
+
+	FPCGSourceDataContainer Container;
+
+	// Store two entries
+	Container.Store<FPCGSourceDataStorageKey>(KeyA, PayloadA);
+	Container.Store<FPCGSourceDataStorageKey>(KeyB, PayloadB);
+
+	UTEST_EQUAL("Num after two stores", Container.Num(), 2);
+	UTEST_FALSE("Not empty", Container.IsEmpty());
+
+	// Retrieve first entry
+	{
+		const FConstSharedStruct ResultA = Container.Get<FPCGSourceDataStorageKey>(KeyA);
+		UTEST_TRUE("KeyA result valid", ResultA.IsValid());
+
+		const FPCGSourceDataStorageKey* PtrA = ResultA.GetPtr<FPCGSourceDataStorageKey>();
+		UTEST_NOT_NULL("KeyA ptr not null", PtrA);
+		UTEST_TRUE("KeyA payload matches", *PtrA == PayloadA);
+	}
+
+	// Retrieve second entry
+	{
+		const FConstSharedStruct ResultB = Container.Get<FPCGSourceDataStorageKey>(KeyB);
+		UTEST_TRUE("KeyB result valid", ResultB.IsValid());
+
+		const FPCGSourceDataStorageKey* PtrB = ResultB.GetPtr<FPCGSourceDataStorageKey>();
+		UTEST_NOT_NULL("KeyB ptr not null", PtrB);
+		UTEST_TRUE("KeyB payload matches", *PtrB == PayloadB);
+	}
+
+	return true;
+}
+
+// -------------------------------------------------------------------
+// Store overwrites existing key
+// -------------------------------------------------------------------
+
+IMPLEMENT_CUSTOM_SIMPLE_AUTOMATION_TEST(FPCGSourceDataContainerStoreOverwriteTest, FPCGTestBaseClass, "Plugins.PCG.SourceDataContainer.StoreOverwrite", PCGTestsCommon::TestFlags)
+
+bool FPCGSourceDataContainerStoreOverwriteTest::RunTest(const FString& Parameters)
+{
+	const FPCGSourceDataStorageKey Key(TEXT("Label"), 42);
+	const FPCGSourceDataStorageKey OriginalPayload(TEXT("Original"), 1);
+	const FPCGSourceDataStorageKey ReplacementPayload(TEXT("Replacement"), 2);
+
+	FPCGSourceDataContainer Container;
+	Container.Store<FPCGSourceDataStorageKey>(Key, OriginalPayload);
+	Container.Store<FPCGSourceDataStorageKey>(Key, ReplacementPayload);
+
+	UTEST_EQUAL("Num stays 1 after overwrite", Container.Num(), 1);
+
+	const FConstSharedStruct Result = Container.Get<FPCGSourceDataStorageKey>(Key);
+	UTEST_TRUE("Result valid after overwrite", Result.IsValid());
+
+	const FPCGSourceDataStorageKey* Ptr = Result.GetPtr<FPCGSourceDataStorageKey>();
+	UTEST_NOT_NULL("Ptr not null", Ptr);
+	UTEST_TRUE("Payload is replacement, not original", *Ptr == ReplacementPayload);
+
+	return true;
+}
+
+// -------------------------------------------------------------------
+// GetMutable — mutation persists
+// -------------------------------------------------------------------
+
+IMPLEMENT_CUSTOM_SIMPLE_AUTOMATION_TEST(FPCGSourceDataContainerGetMutableTest, FPCGTestBaseClass, "Plugins.PCG.SourceDataContainer.GetMutable", PCGTestsCommon::TestFlags)
+
+bool FPCGSourceDataContainerGetMutableTest::RunTest(const FString& Parameters)
+{
+	const FPCGSourceDataStorageKey Key(TEXT("MutableTest"), 42);
+	const FPCGSourceDataStorageKey InitialPayload(TEXT("Initial"), 1);
+	const FPCGSourceDataStorageKey MutatedPayload(TEXT("Mutated"), 2);
+
+	FPCGSourceDataContainer Container;
+	Container.Store<FPCGSourceDataStorageKey>(Key, InitialPayload);
+
+	// Mutate via GetMutable
+	{
+		FSharedStruct MutableResult = Container.GetMutable<FPCGSourceDataStorageKey>(Key);
+		UTEST_TRUE("Mutable result valid", MutableResult.IsValid());
+
+		FPCGSourceDataStorageKey* MutablePtr = MutableResult.GetPtr<FPCGSourceDataStorageKey>();
+		UTEST_NOT_NULL("Mutable ptr not null", MutablePtr);
+		*MutablePtr = MutatedPayload;
+	}
+
+	// Verify mutation persisted
+	{
+		const FConstSharedStruct ConstResult = Container.Get<FPCGSourceDataStorageKey>(Key);
+		UTEST_TRUE("Const result valid after mutation", ConstResult.IsValid());
+
+		const FPCGSourceDataStorageKey* ConstPtr = ConstResult.GetPtr<FPCGSourceDataStorageKey>();
+		UTEST_NOT_NULL("Const ptr not null", ConstPtr);
+		UTEST_TRUE("Payload reflects mutation", *ConstPtr == MutatedPayload);
+	}
+
+	// GetMutable on non-existent key returns invalid
+	{
+		const FPCGSourceDataStorageKey MissingKey(TEXT("Missing"), 0);
+		const FSharedStruct MissingResult = Container.GetMutable<FPCGSourceDataStorageKey>(MissingKey);
+		UTEST_FALSE("GetMutable on missing key returns invalid", MissingResult.IsValid());
+	}
+
+	return true;
+}
+
+// -------------------------------------------------------------------
+// Remove
+// -------------------------------------------------------------------
+
+IMPLEMENT_CUSTOM_SIMPLE_AUTOMATION_TEST(FPCGSourceDataContainerRemoveTest, FPCGTestBaseClass, "Plugins.PCG.SourceDataContainer.Remove", PCGTestsCommon::TestFlags)
+
+bool FPCGSourceDataContainerRemoveTest::RunTest(const FString& Parameters)
+{
+	const FPCGSourceDataStorageKey Key(TEXT("RemoveMe"), 42);
+	const FPCGSourceDataStorageKey Payload(TEXT("Payload"), 1);
+
+	FPCGSourceDataContainer Container;
+	Container.Store<FPCGSourceDataStorageKey>(Key, Payload);
+
+	UTEST_EQUAL("Num before remove", Container.Num(), 1);
+
+	const bool bRemoved = Container.Remove(Key);
+	UTEST_TRUE("Remove returns true for existing key", bRemoved);
+	UTEST_EQUAL("Num after remove", Container.Num(), 0);
+	UTEST_TRUE("Empty after remove", Container.IsEmpty());
+
+	// Get after removal returns invalid
+	const FConstSharedStruct RemovedResult = Container.Get<FPCGSourceDataStorageKey>(Key);
+	UTEST_FALSE("Get after remove returns invalid", RemovedResult.IsValid());
+
+	// Removing again returns false
+	const bool bRemovedAgain = Container.Remove(Key);
+	UTEST_FALSE("Remove non-existent returns false", bRemovedAgain);
+
+	return true;
+}
+
+// -------------------------------------------------------------------
+// MarkDirty
+// -------------------------------------------------------------------
+
+IMPLEMENT_CUSTOM_SIMPLE_AUTOMATION_TEST(FPCGSourceDataContainerMarkDirtyTest, FPCGTestBaseClass, "Plugins.PCG.SourceDataContainer.MarkDirty", PCGTestsCommon::TestFlags)
+
+bool FPCGSourceDataContainerMarkDirtyTest::RunTest(const FString& Parameters)
+{
+	FPCGSourceDataContainer Container;
+
+	UTEST_EQUAL("Initial generation is 0", Container.GetDirtyGeneration(), 0u);
+
+	Container.MarkDirty();
+	UTEST_EQUAL("Generation after first MarkDirty", Container.GetDirtyGeneration(), 1u);
+
+	Container.MarkDirty();
+	UTEST_EQUAL("Generation after second MarkDirty", Container.GetDirtyGeneration(), 2u);
+
+	return true;
+}
+
+// -------------------------------------------------------------------
+// AutoDirty — Store and Remove increment generation
+// -------------------------------------------------------------------
+
+IMPLEMENT_CUSTOM_SIMPLE_AUTOMATION_TEST(FPCGSourceDataContainerAutoDirtyTest, FPCGTestBaseClass, "Plugins.PCG.SourceDataContainer.AutoDirty", PCGTestsCommon::TestFlags)
+
+bool FPCGSourceDataContainerAutoDirtyTest::RunTest(const FString& Parameters)
+{
+	const FPCGSourceDataStorageKey Key(TEXT("AutoDirty"), 42);
+	const FPCGSourceDataStorageKey Payload(TEXT("Payload"), 1);
+
+	FPCGSourceDataContainer Container;
+
+	// Without auto-dirty: Store does NOT increment generation
+	Container.Store<FPCGSourceDataStorageKey>(Key, Payload);
+	UTEST_EQUAL("Generation unchanged without auto-dirty", Container.GetDirtyGeneration(), 0u);
+
+	Container.Remove(Key);
+	UTEST_EQUAL("Generation unchanged on remove without auto-dirty", Container.GetDirtyGeneration(), 0u);
+
+	// Enable auto-dirty
+	Container.SetShouldAutoDirty(true);
+
+	Container.Store<FPCGSourceDataStorageKey>(Key, Payload);
+	UTEST_EQUAL("Generation incremented on Store with auto-dirty", Container.GetDirtyGeneration(), 1u);
+
+	// Store overwrites — still increments
+	Container.Store<FPCGSourceDataStorageKey>(Key, Payload);
+	UTEST_EQUAL("Generation incremented on overwrite with auto-dirty", Container.GetDirtyGeneration(), 2u);
+
+	// Remove with auto-dirty
+	Container.Remove(Key);
+	UTEST_EQUAL("Generation incremented on Remove with auto-dirty", Container.GetDirtyGeneration(), 3u);
+
+	// Remove non-existent does NOT increment
+	Container.Remove(Key);
+	UTEST_EQUAL("Generation unchanged on remove of non-existent with auto-dirty", Container.GetDirtyGeneration(), 3u);
+
+	return true;
+}
+
+// -------------------------------------------------------------------
+// Empty
+// -------------------------------------------------------------------
+
+IMPLEMENT_CUSTOM_SIMPLE_AUTOMATION_TEST(FPCGSourceDataContainerEmptyTest, FPCGTestBaseClass, "Plugins.PCG.SourceDataContainer.Empty", PCGTestsCommon::TestFlags)
+
+bool FPCGSourceDataContainerEmptyTest::RunTest(const FString& Parameters)
+{
+	const FPCGSourceDataStorageKey KeyA(TEXT("A"), 1);
+	const FPCGSourceDataStorageKey KeyB(TEXT("B"), 2);
+	const FPCGSourceDataStorageKey Payload(TEXT("P"), 0);
+
+	FPCGSourceDataContainer Container;
+	Container.SetShouldAutoDirty(true);
+	Container.Store<FPCGSourceDataStorageKey>(KeyA, Payload);
+	Container.Store<FPCGSourceDataStorageKey>(KeyB, Payload);
+
+	UTEST_EQUAL("Num before Empty", Container.Num(), 2);
+	UTEST_EQUAL("Generation before Empty", Container.GetDirtyGeneration(), 2u);
+
+	Container.Empty();
+
+	UTEST_TRUE("IsEmpty after Empty()", Container.IsEmpty());
+	UTEST_EQUAL("Num after Empty()", Container.Num(), 0);
+	UTEST_EQUAL("Generation reset to 0 after Empty()", Container.GetDirtyGeneration(), 0u);
+
+	return true;
+}
+
+// -------------------------------------------------------------------
+// Type mismatch — GetAs and GetMutable with wrong type
+// -------------------------------------------------------------------
+
+IMPLEMENT_CUSTOM_SIMPLE_AUTOMATION_TEST(FPCGSourceDataContainerTypeMismatchTest, FPCGTestBaseClass, "Plugins.PCG.SourceDataContainer.TypeMismatch", PCGTestsCommon::TestFlags)
+
+bool FPCGSourceDataContainerTypeMismatchTest::RunTest(const FString& Parameters)
+{
+	const FPCGSourceDataStorageKey Key(TEXT("Typed"), 42);
+	const FPCGSourceDataStorageKey Payload(TEXT("Payload"), 1);
+
+	FPCGSourceDataContainer Container;
+	Container.Store<FPCGSourceDataStorageKey>(Key, Payload);
+
+	AddExpectedError(TEXT("type mismatch"), EAutomationExpectedErrorFlags::Contains, 1);
+	const FConstSharedStruct WrongTypeResult = Container.Get<FPCGSourceDataStorageValue>(Key);
+	UTEST_FALSE("Get with wrong type returns invalid", WrongTypeResult.IsValid());
+
+	AddExpectedError(TEXT("type mismatch"), EAutomationExpectedErrorFlags::Contains, 1);
+	const FSharedStruct WrongTypeMutableResult = Container.GetMutable<FPCGSourceDataStorageValue>(Key);
+	UTEST_FALSE("GetMutable with wrong type returns invalid", WrongTypeMutableResult.IsValid());
+
+	return true;
+}
+
+// -------------------------------------------------------------------
+// Copy semantics
+// -------------------------------------------------------------------
+
+IMPLEMENT_CUSTOM_SIMPLE_AUTOMATION_TEST(FPCGSourceDataContainerCopyTest, FPCGTestBaseClass, "Plugins.PCG.SourceDataContainer.CopySemantics", PCGTestsCommon::TestFlags)
+
+bool FPCGSourceDataContainerCopyTest::RunTest(const FString& Parameters)
+{
+	const FPCGSourceDataStorageKey Key(TEXT("CopyTest"), 42);
+	const FPCGSourceDataStorageKey Payload(TEXT("Payload"), 1);
+
+	FPCGSourceDataContainer Original;
+	Original.SetShouldAutoDirty(true);
+	Original.Store<FPCGSourceDataStorageKey>(Key, Payload);
+
+	// Copy construction
+	{
+		FPCGSourceDataContainer Copy(Original);
+
+		UTEST_EQUAL("Copy num matches original", Copy.Num(), Original.Num());
+		UTEST_EQUAL("Copy dirty generation matches original", Copy.GetDirtyGeneration(), Original.GetDirtyGeneration());
+
+		const FConstSharedStruct CopyResult = Copy.Get<FPCGSourceDataStorageKey>(Key);
+		UTEST_TRUE("Copy Get result valid", CopyResult.IsValid());
+
+		const FPCGSourceDataStorageKey* CopyPtr = CopyResult.GetPtr<FPCGSourceDataStorageKey>();
+		UTEST_NOT_NULL("Copy ptr not null", CopyPtr);
+		UTEST_TRUE("Copy payload matches", *CopyPtr == Payload);
+	}
+
+	// Copy assignment
+	{
+		FPCGSourceDataContainer Assigned;
+		Assigned = Original;
+
+		UTEST_EQUAL("Assigned num matches original", Assigned.Num(), Original.Num());
+		UTEST_EQUAL("Assigned dirty generation matches", Assigned.GetDirtyGeneration(), Original.GetDirtyGeneration());
+
+		const FConstSharedStruct AssignedResult = Assigned.Get<FPCGSourceDataStorageKey>(Key);
+		UTEST_TRUE("Assigned Get result valid", AssignedResult.IsValid());
+	}
+
+	return true;
+}
+
+// -------------------------------------------------------------------
+// Move semantics
+// -------------------------------------------------------------------
+
+IMPLEMENT_CUSTOM_SIMPLE_AUTOMATION_TEST(FPCGSourceDataContainerMoveTest, FPCGTestBaseClass, "Plugins.PCG.SourceDataContainer.MoveSemantics", PCGTestsCommon::TestFlags)
+
+bool FPCGSourceDataContainerMoveTest::RunTest(const FString& Parameters)
+{
+	const FPCGSourceDataStorageKey Key(TEXT("MoveTest"), 42);
+	const FPCGSourceDataStorageKey Payload(TEXT("Payload"), 1);
+
+	// Move construction
+	{
+		FPCGSourceDataContainer Source;
+		Source.SetShouldAutoDirty(true);
+		Source.Store<FPCGSourceDataStorageKey>(Key, Payload);
+		const uint32 OriginalGeneration = Source.GetDirtyGeneration();
+
+		FPCGSourceDataContainer Moved(MoveTemp(Source));
+
+		UTEST_EQUAL("Moved-to has correct num", Moved.Num(), 1);
+		UTEST_EQUAL("Moved-to has correct generation", Moved.GetDirtyGeneration(), OriginalGeneration);
+
+		const FConstSharedStruct MovedResult = Moved.Get<FPCGSourceDataStorageKey>(Key);
+		UTEST_TRUE("Moved-to Get result valid", MovedResult.IsValid());
+
+		// Source is reset after move
+		UTEST_TRUE("Source empty after move", Source.IsEmpty());
+		UTEST_EQUAL("Source generation reset after move", Source.GetDirtyGeneration(), 0u);
+	}
+
+	// Move assignment
+	{
+		FPCGSourceDataContainer Source;
+		Source.SetShouldAutoDirty(true);
+		Source.Store<FPCGSourceDataStorageKey>(Key, Payload);
+
+		FPCGSourceDataContainer Assigned;
+		Assigned = MoveTemp(Source);
+
+		UTEST_EQUAL("Move-assigned has correct num", Assigned.Num(), 1);
+		UTEST_TRUE("Source empty after move assignment", Source.IsEmpty());
+		UTEST_EQUAL("Source generation reset after move assignment", Source.GetDirtyGeneration(), 0u);
+	}
+
+	// Self move-assignment is safe
+	{
+		FPCGSourceDataContainer SelfMove;
+		SelfMove.Store<FPCGSourceDataStorageKey>(Key, Payload);
+		SelfMove = MoveTemp(SelfMove);
+
+		UTEST_EQUAL("Self move preserves num", SelfMove.Num(), 1);
+	}
+
+	return true;
+}
+
+#endif // WITH_EDITOR
diff --git a/Engine/Plugins/PCG/Tests/Private/ExecutionSource/PCGSourceDataContainerTests.cpp b/Engine/Plugins/PCG/Tests/Private/ExecutionSource/PCGSourceDataContainerTests.cpp
new file mode 100644
--- /dev/null
+++ b/Engine/Plugins/PCG/Tests/Private/ExecutionSource/PCGSourceDataContainerTests.cpp
@@ -0,0 +1,175 @@
+// Copyright Epic Games, Inc. All Rights Reserved.
+
+#include "CoreMinimal.h"
+#include "TestHarness.h"
+
+#include <catch2/catch_test_macros.hpp>
+
+#include "PCGTestsCommon.h"
+#include "Graph/PCGSourceDataContainer.h"
+
+namespace PCGSourceDataStorageTests
+{
+	static const FName LabelA = TEXT("LabelA");
+	static const FName LabelB = TEXT("LabelB");
+	static constexpr uint64 HashA = 42;
+	static constexpr uint64 HashB = 99;
+
+	// Test payload values — must satisfy CPCGSourceDataCompatibleType (i.e. USTRUCT with StaticStruct()).
+	// FPCGSourceDataStorageKey is a USTRUCT from the header under test.
+	static const FPCGSourceDataStorageKey TestValueA(TEXT("PayloadA"), 100);
+	static const FPCGSourceDataStorageKey TestValueB(TEXT("PayloadB"), 200);
+}
+
+// -------------------------------------------------------------------
+// Key: construction, equality, hashing
+// -------------------------------------------------------------------
+
+TEST_CASE("PCG::SourceDataStorage::Key", "[PCG][SourceDataStorage]")
+{
+	using namespace PCGSourceDataStorageTests;
+
+	SECTION("Default construction")
+	{
+		const FPCGSourceDataStorageKey Key;
+		CHECK(Key.Label == NAME_None);
+		CHECK(Key.Hash == 0);
+	}
+
+	SECTION("Parameterized construction")
+	{
+		const FPCGSourceDataStorageKey Key(LabelA, HashA);
+		CHECK(Key.Label == LabelA);
+		CHECK(Key.Hash == HashA);
+	}
+
+	SECTION("Label-only construction uses default hash")
+	{
+		const FPCGSourceDataStorageKey Key(LabelA, 0);
+		CHECK(Key.Label == LabelA);
+		CHECK(Key.Hash == 0);
+	}
+
+	SECTION("Equal keys")
+	{
+		const FPCGSourceDataStorageKey A(LabelA, HashA);
+		const FPCGSourceDataStorageKey B(LabelA, HashA);
+		CHECK(A == B);
+		CHECK(GetTypeHash(A) == GetTypeHash(B));
+	}
+
+	SECTION("Different label")
+	{
+		const FPCGSourceDataStorageKey A(LabelA, HashA);
+		const FPCGSourceDataStorageKey B(LabelB, HashA);
+		CHECK_FALSE(A == B);
+	}
+
+	SECTION("Different hash")
+	{
+		const FPCGSourceDataStorageKey A(LabelA, HashA);
+		const FPCGSourceDataStorageKey B(LabelA, HashB);
+		CHECK_FALSE(A == B);
+	}
+
+	SECTION("TMap lookup")
+	{
+		TMap<FPCGSourceDataStorageKey, int32> Map;
+		const FPCGSourceDataStorageKey Key(LabelA, HashA);
+		Map.Add(Key, 7);
+
+		const int32* Found = Map.Find(Key);
+		REQUIRE(Found != nullptr);
+		CHECK(*Found == 7);
+
+		const FPCGSourceDataStorageKey Missing(LabelB, HashB);
+		CHECK(Map.Find(Missing) == nullptr);
+	}
+}
+
+// -------------------------------------------------------------------
+// Value: type-safe access via FSharedStruct
+// -------------------------------------------------------------------
+
+TEST_CASE_METHOD(PCGTests::FPCGBaseTest, "PCG::SourceDataStorage::Value", "[PCG][SourceDataStorage]")
+{
+	using namespace PCGSourceDataStorageTests;
+
+	FPCGSourceDataStorageValue Value;
+	Value.Data = FSharedStruct::Make(TestValueA);
+
+	SECTION("GetAs returns const shared struct")
+	{
+		const FConstSharedStruct Result = Value.GetAs<FPCGSourceDataStorageKey>();
+		REQUIRE(Result.IsValid());
+		const FPCGSourceDataStorageKey* ResultPtr = Result.GetPtr<FPCGSourceDataStorageKey>();
+		REQUIRE(ResultPtr != nullptr);
+		CHECK(*ResultPtr == TestValueA);
+	}
+
+	SECTION("GetAsMutable returns mutable shared struct")
+	{
+		FSharedStruct Result = Value.GetAsMutable<FPCGSourceDataStorageKey>();
+		REQUIRE(Result.IsValid());
+		FPCGSourceDataStorageKey* ResultPtr = Result.GetPtr<FPCGSourceDataStorageKey>();
+		REQUIRE(ResultPtr != nullptr);
+		*ResultPtr = TestValueB;
+
+		// Verify mutation persisted
+		const FConstSharedStruct Verify = Value.GetAs<FPCGSourceDataStorageKey>();
+		REQUIRE(Verify.IsValid());
+		const FPCGSourceDataStorageKey* VerifyPtr = Verify.GetPtr<FPCGSourceDataStorageKey>();
+		REQUIRE(VerifyPtr != nullptr);
+		CHECK(*VerifyPtr == TestValueB);
+	}
+
+	// @todo_pcg: These can only be verified in WITH_EDITOR build, which runtime tests are not editor build.
+	// SECTION("GetAs with wrong stored type returns invalid")
+	// {
+	// 	// FPCGSourceDataStorageValue is a different USTRUCT than FPCGSourceDataStorageKey
+	// 	const FConstSharedStruct Result = Value.GetAs<FPCGSourceDataStorageValue>();
+	// 	CHECK_FALSE(Result.IsValid());
+	// }
+
+	// SECTION("GetAsMutable with wrong stored type returns invalid")
+	// {
+	// 	const FSharedStruct Result = Value.GetAsMutable<FPCGSourceDataStorageValue>();
+	// 	CHECK_FALSE(Result.IsValid());
+	// }
+
+	SECTION("GetAs on empty value returns invalid")
+	{
+		FPCGSourceDataStorageValue EmptyValue;
+		CHECK_FALSE(EmptyValue.GetAs<FPCGSourceDataStorageKey>().IsValid());
+	}
+
+	SECTION("GetAsMutable on empty value returns invalid")
+	{
+		FPCGSourceDataStorageValue EmptyValue;
+		CHECK_FALSE(EmptyValue.GetAsMutable<FPCGSourceDataStorageKey>().IsValid());
+	}
+}
+
+// -------------------------------------------------------------------
+// Container: non-editor API (Get, GetCrc, Num, IsEmpty, operator=)
+// -------------------------------------------------------------------
+
+TEST_CASE("PCG::SourceDataStorage::Container", "[PCG][SourceDataStorage]")
+{
+	using namespace PCGSourceDataStorageTests;
+
+	FPCGSourceDataContainer Container;
+
+	SECTION("Default state")
+	{
+		CHECK(Container.IsEmpty());
+		CHECK(Container.Num() == 0);
+		CHECK(Container.GetDirtyGeneration() == 0);
+	}
+
+	SECTION("Get non-existent key returns invalid")
+	{
+		const FPCGSourceDataStorageKey Key(LabelA, HashA);
+		CHECK_FALSE(Container.Get<FPCGSourceDataStorageKey>(Key).IsValid());
+	}
+}
