[PCG] Fixed crash when point operation fails during initialization

#rnx

#Synced-CL 51379665

diff --git a/Engine/Plugins/PCG/Source/PCG/Private/Elements/PCGPointOperationElementBase.cpp b/Engine/Plugins/PCG/Source/PCG/Private/Elements/PCGPointOperationElementBase.cpp
--- a/Engine/Plugins/PCG/Source/PCG/Private/Elements/PCGPointOperationElementBase.cpp
+++ b/Engine/Plugins/PCG/Source/PCG/Private/Elements/PCGPointOperationElementBase.cpp
@@ -157,7 +157,11 @@ PRAGMA_ENABLE_DEPRECATION_WARNINGS
 			// If this input created an error, result in no output
 			if (Context->GetIterationStateResult(IterIndex) != EPCGTimeSliceInitResult::Success)
 			{
-				IterState.OutputData->SetNumPoints(0);
+				if (IterState.OutputData)
+				{
+					IterState.OutputData->SetNumPoints(0);
+				}
+				
 				return true;
 			}
 
