[PCG] Fixed potential edge cases when volumes have no brush component.

#rnx

#Synced-CL 51379665

diff --git a/Engine/Plugins/PCG/Source/PCG/Private/Data/Tool/PCGToolVolumeData.cpp b/Engine/Plugins/PCG/Source/PCG/Private/Data/Tool/PCGToolVolumeData.cpp
--- a/Engine/Plugins/PCG/Source/PCG/Private/Data/Tool/PCGToolVolumeData.cpp
+++ b/Engine/Plugins/PCG/Source/PCG/Private/Data/Tool/PCGToolVolumeData.cpp
@@ -122,7 +122,7 @@ void FPCGInteractiveToolWorkingData_Volume::OnToolCancel(const FPCGInteractiveTo
 
 		Volume->SetActorTransform(OnToolStartActorTransform);
 
-		if (OnToolStartBrushBuilder)
+		if (OnToolStartBrushBuilder && Volume->GetBrushComponent())
 		{
 			// Recreate the brush from the original actor
 			UActorFactory::CreateBrushForVolumeActor(Volume.Get(), OnToolStartBrushBuilder);
diff --git a/Engine/Plugins/PCG/Source/PCGEditor/Private/EditorMode/Tools/Volume/PCGVolumeTool.cpp b/Engine/Plugins/PCG/Source/PCGEditor/Private/EditorMode/Tools/Volume/PCGVolumeTool.cpp
--- a/Engine/Plugins/PCG/Source/PCGEditor/Private/EditorMode/Tools/Volume/PCGVolumeTool.cpp
+++ b/Engine/Plugins/PCG/Source/PCGEditor/Private/EditorMode/Tools/Volume/PCGVolumeTool.cpp
@@ -207,7 +207,7 @@ void UPCGVolumeTool::UpdateVolumeFromBox()
 		Extent.Z += 2.0 * FMath::Abs(SnappedCenter.Z - BoxCenter.Z);
 	}
 
-	if (AVolume* WorkingVolume = ToolSettings ? ToolSettings->GetWorkingVolume(ToolSettings->DataInstance) : nullptr)
+	if (AVolume* WorkingVolume = ToolSettings ? ToolSettings->GetWorkingVolume(ToolSettings->DataInstance) : nullptr; WorkingVolume && WorkingVolume->GetBrushComponent())
 	{
 		// Move actor to bounds center
 		WorkingVolume->SetActorTransform(FTransform(BoxCenter));
