[PCG] Fixes in Generic Attributes, in preparation for the next phase of replacing FPCGMetadataAttribute<T>

* Add support for serialization
* Add IsOfType<T> to factorize code
* Make the constructors for generic attribute protected so they cannot be created outisde of metadata
* Factorize the Copy function using a functor to support copy between different types.
* Fix a few places were we were not cleaning memory correctly, introducing leaks
* Fix issue in the GetValues for continuous values that was not tagging the retrieved values correctly
* Fix potential dangling pointer when copying values to itself
* Fix issue where SoftObjectPath/SoftClassPath were not compressed
* Fix issue with accessor not taking into account the potential offset when re-allocating values.
* Fix traits for soft object/class path that were not using the right function.
* Add tests cases for those issues

#rnx
#Synced-CL 51343786

diff --git a/Engine/Plugins/PCG/Source/PCG/Private/Metadata/PCGMetadataAttributeGeneric.cpp b/Engine/Plugins/PCG/Source/PCG/Private/Metadata/PCGMetadataAttributeGeneric.cpp
--- a/Engine/Plugins/PCG/Source/PCG/Private/Metadata/PCGMetadataAttributeGeneric.cpp
+++ b/Engine/Plugins/PCG/Source/PCG/Private/Metadata/PCGMetadataAttributeGeneric.cpp
@@ -381,8 +381,8 @@ FPCGAttributeProperty::FPCGAttributeProperty(const FPCGMetadataAttributeDesc& In
 		|| InAttributeDesc.ValueType == EPCGMetadataTypes::Text
 		|| InAttributeDesc.ValueType == EPCGMetadataTypes::SoftObjectPath
 		|| InAttributeDesc.ValueType == EPCGMetadataTypes::SoftClassPath
-		|| InAttributeDesc.ValueTypeObject == TBaseStructure<FSoftObjectPath>()
-		|| InAttributeDesc.ValueTypeObject == TBaseStructure<FSoftClassPath>());
+		|| InAttributeDesc.ValueTypeObject == TBaseStructure<FSoftObjectPath>().Get()
+		|| InAttributeDesc.ValueTypeObject == TBaseStructure<FSoftClassPath>().Get());
 
 	InnerElementCachedSize = NewProperty->GetElementSize();
 }
@@ -449,13 +449,15 @@ void FPCGAttributeDefaultValue::Destroy()
 		{
 			UnderlyingProperty->Property->Inner->DestroyValue(Memory);
 		}
+
+		FMemory::Free(Memory);
 	}
 
 	Memory = nullptr;
 }
 
-FPCGMetadataAttributeGeneric::FPCGMetadataAttributeGeneric(const FPCGMetadataAttributeDesc& InAttributeDesc, FPCGMetadataDomain* InMetadata, FName InName, bool bInAllowsInterpolation)
-	: FPCGMetadataAttributeBase(InMetadata, InName, nullptr, bInAllowsInterpolation)
+FPCGMetadataAttributeGeneric::FPCGMetadataAttributeGeneric(const FPCGMetadataAttributeDesc& InAttributeDesc, FPCGMetadataDomain* InMetadata, bool bInAllowsInterpolation)
+	: FPCGMetadataAttributeBase(InMetadata, InAttributeDesc.Name, nullptr, bInAllowsInterpolation)
 	, CachedDesc(InAttributeDesc)
 	, UnderlyingProperty(MakeShared<FPCGAttributeProperty>(CachedDesc))
 {
@@ -470,11 +472,13 @@ FPCGMetadataAttributeGeneric::FPCGMetadataAttributeGeneric(const FPCGMetadataAtt
 	Init();
 }
 
-FPCGMetadataAttributeGeneric::FPCGMetadataAttributeGeneric(TNotNull<const FPCGMetadataAttributeGeneric*> InParent, FPCGMetadataDomain* InMetadata, FName InName, bool bInAllowsInterpolation)
-	: FPCGMetadataAttributeBase(InMetadata, InName, InParent, bInAllowsInterpolation)
+FPCGMetadataAttributeGeneric::FPCGMetadataAttributeGeneric(TNotNull<const FPCGMetadataAttributeGeneric*> InParent, FPCGMetadataDomain* InMetadata, bool bInAllowsInterpolation, TOptional<FName> InName)
+	: FPCGMetadataAttributeBase(InMetadata, InName.Get(InParent->CachedDesc.Name), InParent, bInAllowsInterpolation)
 	, CachedDesc(InParent->CachedDesc)
 	, UnderlyingProperty(InParent->UnderlyingProperty)
 {
+	CachedDesc.Name = InName.Get(InParent->CachedDesc.Name);
+
 	// @todo_pcg: Should we init there so we always have a default value like before?
 	Init();
 
@@ -487,12 +491,85 @@ FPCGMetadataAttributeGeneric::FPCGMetadataAttributeGeneric(TNotNull<const FPCGMe
 	}
 }
 
+FPCGMetadataAttributeGeneric::FPCGMetadataAttributeGeneric(const FPCGMetadataAttributeDesc& InAttributeDesc, const FPCGMetadataAttributeGeneric* InParent, FPCGMetadataDomain* InMetadata, bool bInAllowsInterpolation)
+	: FPCGMetadataAttributeBase(InMetadata, InAttributeDesc.Name, InParent, bInAllowsInterpolation)
+	, CachedDesc(InAttributeDesc)
+	, UnderlyingProperty(InParent ? InParent->UnderlyingProperty : MakeShared<FPCGAttributeProperty>(CachedDesc))
+{
+	check(!InParent || InParent->GetAttributeDesc().IsSameType(CachedDesc));
+
+	// @todo_pcg: Should we init there so we always have a default value like before?
+	Init();
+
+	ValueKeyOffset = InParent ? InParent->GetValueKeyOffsetForChild() : 0;
+
+	if (InParent && ensure(DefaultValue.IsValid()) && ensure(InParent->DefaultValue.IsValid()))
+	{
+		// Copy the default value from the parent.
+		DefaultValue.CopyFrom(InParent->DefaultValue.GetRawPtr());
+	}
+}
+
+FPCGMetadataAttributeGeneric::~FPCGMetadataAttributeGeneric()
+{
+	// We need to make sure to empty the array so the elements are destroyed and the memory is freed.
+	if (TOptional<FScriptArrayHelper> Helper = MakeArrayHelper(); ensure(Values.IsEmpty() || Helper.IsSet()) && Helper)
+	{
+		Helper->EmptyValues();
+	}
+}
+
 PCGMetadataValueKey FPCGMetadataAttributeGeneric::GetValueKeyOffsetForChild() const
 {
 	PCG::TSharedScopeLock ReadLock(ValueLock);
 	return Values.Num() + ValueKeyOffset;
 }
 
+void FPCGMetadataAttributeGeneric::Serialize(FPCGMetadataDomain* InMetadata, FArchive& InArchive)
+{
+	FPCGMetadataAttributeBase::Serialize(InMetadata, InArchive);
+
+	// If we don't have a parent, we need to serialize the CachedDesc
+	if (!Parent)
+	{
+		FPCGMetadataAttributeDesc::StaticStruct()->SerializeItem(FStructuredArchiveFromArchive(InArchive).GetSlot(), &CachedDesc, /*Defaults*/nullptr);
+		if (InArchive.IsLoading() && !UnderlyingProperty)
+		{
+			UnderlyingProperty = MakeShared<FPCGAttributeProperty>(CachedDesc);
+		}
+	}
+	else if (InArchive.IsLoading() && ensure(Parent->IsGeneric()))
+	{
+		const FPCGMetadataAttributeGeneric* GenericParent = static_cast<const FPCGMetadataAttributeGeneric*>(Parent);
+		CachedDesc = GenericParent->CachedDesc;
+		CachedDesc.Name = Name;
+		UnderlyingProperty = GenericParent->UnderlyingProperty;
+	}
+
+	if (InArchive.IsLoading() && !DefaultValue.IsValid())
+	{
+		// Initialize the attribute
+		Init();
+	}
+
+	if (UnderlyingProperty && UnderlyingProperty->IsValid())
+	{
+		// Then serialize the default value
+		if (DefaultValue.IsValid())
+		{
+			UnderlyingProperty->Property->Inner->SerializeItem(FStructuredArchiveFromArchive(InArchive).GetSlot(), DefaultValue.GetRawPtr(), /*Defaults*/nullptr);
+		}
+	
+		// And finally serialize the values
+		UnderlyingProperty->Property->SerializeItem(FStructuredArchiveFromArchive(InArchive).GetSlot(), &Values, /*Defaults*/nullptr);
+	
+		if (InArchive.IsLoading())
+		{
+			ValueKeyOffset = GetParent() ? GetParent()->GetValueKeyOffsetForChild() : 0;
+		}
+	}
+}
+
 void FPCGMetadataAttributeGeneric::AddReferencedObjects(FReferenceCollector& Collector)
 {
 	// @todo_pcg: We should register all the stored objects
@@ -695,17 +772,23 @@ TArray<PCGMetadataValueKey> FPCGMetadataAttributeGeneric::AddCompressedValues(FI
 	return FoundValueKeys;
 }
 
-bool FPCGMetadataAttributeGeneric::FindValues(TArray<PCGMetadataValueKey>& OutValueKeys, const FInValuesByValue& InValues)
+bool FPCGMetadataAttributeGeneric::FindValues(TArray<PCGMetadataValueKey>& OutValueKeys, const FInValuesByValue& InValues) const
 {
 	OutValueKeys.Init(PCGNotFoundValueKey, InValues.Count);
+	
+	return FindValues(MakeArrayView(OutValueKeys), InValues);
+}
 
+bool FPCGMetadataAttributeGeneric::FindValues(TArrayView<PCGMetadataValueKey> OutValueKeys, const FInValuesByValue& InValues) const
+{
+	check(OutValueKeys.Num() == InValues.Count);
 	int ValueKeysSet = 0;
 	FindValuesInternal(InValues, OutValueKeys, ValueKeysSet, /*bIsRoot=*/true);
 
 	return (ValueKeysSet == InValues.Count);
 }
 
-void FPCGMetadataAttributeGeneric::FindValuesInternal(const FInValuesByValue& InValues, TArray<PCGMetadataValueKey>& ValueKeys, int& ValueKeysSet, bool bIsRoot) const
+void FPCGMetadataAttributeGeneric::FindValuesInternal(const FInValuesByValue& InValues, TArrayView<PCGMetadataValueKey> ValueKeys, int& ValueKeysSet, bool bIsRoot) const
 {
 	check(InValues.Count == ValueKeys.Num());
 	
@@ -782,43 +865,7 @@ void FPCGMetadataAttributeGeneric::Init()
 	// Initialize the default value
 	DefaultValue.Init(UnderlyingProperty);
 
-	auto GetType = [](const FProperty* InProperty) -> int16
-	{
-		int16 InType = static_cast<int16>(PCGAttributeAccessorHelpers::GetMetadataTypeForProperty(InProperty));
-		if (InType == static_cast<int16>(EPCGMetadataTypes::Unknown))
-		{
-			// @todo_pcg TO BE FIXED
-			// For now mimic `Unknown + sizeof(T)` of the Traits
-			InType = static_cast<int16>(EPCGMetadataTypes::Unknown) + InProperty->GetElementSize();
-		}
-		
-		return InType;
-	};
-	
-	const FProperty* Property = UnderlyingProperty->Property->Inner;
-	if (!CachedDesc.ContainerTypes.IsEmpty())
-	{
-		if (const FArrayProperty* ArrayProperty = CastField<FArrayProperty>(Property))
-		{
-			TypeId = GetType(ArrayProperty->Inner);
-		}
-		else if (const FSetProperty* SetProperty = CastField<FSetProperty>(Property))
-		{
-			TypeId = GetType(SetProperty->ElementProp);
-		}
-		else if (const FMapProperty* MapProperty = CastField<FMapProperty>(Property))
-		{
-			TypeId = GetType(MapProperty->ValueProp);
-		}
-		else
-		{
-			checkNoEntry();
-		}
-	}
-	else
-	{
-		TypeId = GetType(Property);
-	}
+	TypeId = static_cast<int16>(CachedDesc.ValueType);
 }
 
 void FPCGMetadataAttributeGeneric::Flatten()
@@ -1020,17 +1067,52 @@ FPCGMetadataAttributeBase* FPCGMetadataAttributeGeneric::Copy(FName NewName, FPC
 	FPCGMetadataAttributeGeneric* AttributeCopy = nullptr;
 	if (bKeepParent)
 	{
-		AttributeCopy = new FPCGMetadataAttributeGeneric(this, InMetadata, NewName, bAllowsInterpolation);
+		AttributeCopy = new FPCGMetadataAttributeGeneric(this, InMetadata, bAllowsInterpolation, NewName);
 	}
 	else
 	{
-		AttributeCopy = new FPCGMetadataAttributeGeneric(CachedDesc, InMetadata, NewName, bAllowsInterpolation);
+		FPCGMetadataAttributeDesc NewDesc = CachedDesc;
+		NewDesc.Name = NewName;
+		AttributeCopy = new FPCGMetadataAttributeGeneric(NewDesc, InMetadata, bAllowsInterpolation);
+	}
+	
+	CopyInternal(AttributeCopy, bKeepParent, bCopyEntries, bCopyValues);
+
+	return AttributeCopy;
+}
+
+FPCGMetadataAttributeBase* FPCGMetadataAttributeGeneric::Copy(FName NewName, UPCGMetadata* InMetadata, bool bKeepParent, bool bCopyEntries,
+	bool bCopyValues) const
+{
+	return Copy(NewName, InMetadata ? InMetadata->GetDefaultMetadataDomain() : nullptr, bKeepParent, bCopyEntries, bCopyValues);
+}
+
+FPCGMetadataAttributeBase* FPCGMetadataAttributeGeneric::CopyToAnotherType(int16 Type) const
+{
+	// This was done to support casting from one type to another, where we knew at compile time both types (source and target).
+	// Without it, we would need a way at runtime to know what is compatible and how to convert the values. In the meantime, this is not supported.
+	return nullptr;
+}
+
+void FPCGMetadataAttributeGeneric::CopyInternal(FPCGMetadataAttributeGeneric* NewAttribute, bool bKeepParent, bool bCopyEntries, bool bCopyValues, FCopyFunctor CopyFunctor) const
+{
+	check(NewAttribute);
 
-		if (ensure(DefaultValue.IsValid()) && ensure(AttributeCopy->DefaultValue.IsValid()))
+	if (!CopyFunctor)
+	{
+		// If we do not have an explicit functor, just use the default copy.
+		CopyFunctor = [&UnderlyingProperty = this->UnderlyingProperty](void* DestData, const void* SrcData, const int32 Count)
 		{
-			// Copy this default value to the attribute copy.
-			AttributeCopy->DefaultValue.CopyFrom(DefaultValue.GetRawPtr());
-		}
+			UnderlyingProperty->Copy(DestData, SrcData, Count);
+		};
+	}
+
+	check(CopyFunctor);
+
+	if (ensure(DefaultValue.IsValid()) && ensure(NewAttribute->DefaultValue.IsValid()))
+	{
+		// Copy this default value to the attribute copy.
+		CopyFunctor(NewAttribute->DefaultValue.GetRawPtr(), DefaultValue.GetRawPtr(), 1);
 	}
 
 	// Gather the chain of parents if we don't keep the parent and we want to copy entries/values.
@@ -1062,7 +1144,7 @@ FPCGMetadataAttributeBase* FPCGMetadataAttributeGeneric::Copy(FName NewName, FPC
 
 			// Only need to lock if Current == this, since the parents are assumed to be const.
 			PCG::TSharedScopeLock ReadLock(Current->EntryMapLock, /*bShouldLock=*/Current == this);
-			AttributeCopy->EntryToValueKeyMap.Append(Current->EntryToValueKeyMap);
+			NewAttribute->EntryToValueKeyMap.Append(Current->EntryToValueKeyMap);
 		}
 	}
 
@@ -1070,8 +1152,11 @@ FPCGMetadataAttributeBase* FPCGMetadataAttributeGeneric::Copy(FName NewName, FPC
 	{
 		// We go backwards, since we need to preserve order (root -> this)
 		// Latest entry in our Parents array is the root.
-		if (TOptional<FScriptArrayHelper> Helper = AttributeCopy->MakeArrayHelper(); Helper.IsSet())
+		if (TOptional<FScriptArrayHelper> Helper = NewAttribute->MakeArrayHelper(); Helper.IsSet())
 		{
+			// Make sure to start from a fresh start, as we already have a default value allocated.
+			Helper->EmptyValues();
+
 			for (int32 i = Parents.Num() - 1; i >= 0; --i)
 			{
 				const FPCGMetadataAttributeGeneric* Current = Parents[i];
@@ -1086,35 +1171,22 @@ FPCGMetadataAttributeBase* FPCGMetadataAttributeGeneric::Copy(FName NewName, FPC
 				}
 				
 				const int32 NumValuesToCopy = ParentHelper->Num();
-				if (NumValuesToCopy <= 0 || AttributeCopy->AllocateValues_Unsafe(NumValuesToCopy) == INDEX_NONE)
+				if (NumValuesToCopy <= 0)
 				{
-					continue;	
+					continue;
 				}
 
-				UnderlyingProperty->Copy(Helper->GetRawPtr(0), ParentHelper->GetRawPtr(0), NumValuesToCopy);
-
-				// The expected value key offset is the one for this attribute (i == 0), and only if we
-				// keep the parent. Otherwise we don't have any parent, so offset should be kept at 0.
-				if (i == 0 && bKeepParent)
+				const int32 StartIndex = NewAttribute->AllocateValues_Unsafe(NumValuesToCopy);
+				
+				if (StartIndex == INDEX_NONE)
 				{
-					AttributeCopy->ValueKeyOffset = Current->ValueKeyOffset;
+					continue;
 				}
+
+				CopyFunctor(Helper->GetRawPtr(StartIndex), ParentHelper->GetRawPtr(0), NumValuesToCopy);
 			}
 		}
 	}
-
-	return AttributeCopy;
-}
-
-FPCGMetadataAttributeBase* FPCGMetadataAttributeGeneric::Copy(FName NewName, UPCGMetadata* InMetadata, bool bKeepParent, bool bCopyEntries,
-	bool bCopyValues) const
-{
-	return Copy(NewName, InMetadata ? InMetadata->GetDefaultMetadataDomain() : nullptr, bKeepParent, bCopyEntries, bCopyValues);
-}
-
-FPCGMetadataAttributeBase* FPCGMetadataAttributeGeneric::CopyToAnotherType(int16 Type) const
-{
-	return nullptr;
 }
 
 void FPCGMetadataAttributeGeneric::SetValue(PCGMetadataEntryKey ItemKey, const FPCGMetadataAttributeBase* InAttribute, PCGMetadataEntryKey InEntryKey)
@@ -1131,50 +1203,54 @@ void FPCGMetadataAttributeGeneric::SetValue(PCGMetadataEntryKey ItemKey, const F
 	{
 		return;
 	}
-
-	int32 ValueKey;
 	
+	check(ItemKey != PCGInvalidEntryKey);
+	
+	const void* SrcPtr = InAttributeGeneric->GetReadAddressFromEntryKey_Unsafe(InEntryKey);
+	if (!SrcPtr)
+	{
+		return;	
+	}
+	
+	PCGMetadataValueKey ValueKeyToAdd = PCGNotFoundValueKey;
+	
+	if (UnderlyingProperty->bCompressData)
+	{
+		const bool bValueKeyFound = FindValues(MakeArrayView(&ValueKeyToAdd, 1), FInValuesByValue{SrcPtr, 1});
+		check(!bValueKeyFound || ValueKeyToAdd != PCGNotFoundValueKey);
+	}
+	
+	// If we are not compressed or we haven't find the value, we need to add it.
+	if (ValueKeyToAdd == PCGNotFoundValueKey)
 	{
 		PCG::TUniqueScopeLock WriteLock(ValueLock);
+		
+		// Add the value in the buffer
+		const int32 NewValueIndex = ThisHelper->AddValue();
 
-		const void* SrcPtr = InAttributeGeneric->GetReadAddressFromEntryKey_Unsafe(InEntryKey);
-		if (!SrcPtr)
+		// If InAttributeGeneric == this, SrcPtr might be dangling now because of possible reallocation, so fix it there
+		if (InAttributeGeneric == this)
 		{
-			return;	
+			SrcPtr = InAttributeGeneric->GetReadAddressFromEntryKey_Unsafe(InEntryKey);
 		}
 		
-		const int32 NewValueIndex = ThisHelper->AddValue();
-		ValueKey = ValueKeyOffset + NewValueIndex;
+		check(SrcPtr);
+
+		ValueKeyToAdd = ValueKeyOffset + NewValueIndex;
 		void* DestPtr = ThisHelper->GetRawPtr(NewValueIndex);
 		UnderlyingProperty->Copy(DestPtr, SrcPtr);
 	}
 
+	if (ensure(ValueKeyToAdd != PCGNotFoundValueKey))
 	{
 		PCG::TUniqueScopeLock WriteLock(EntryMapLock);
-		EntryToValueKeyMap.Emplace(ItemKey, ValueKey);
+		EntryToValueKeyMap.Emplace(ItemKey, ValueKeyToAdd);
 	}
 }
 
 void FPCGMetadataAttributeGeneric::SetZeroValue(PCGMetadataEntryKey ItemKey)
 {
-	TOptional<FScriptArrayHelper> ThisHelper = MakeArrayHelper();
-	if (!ThisHelper.IsSet())
-	{
-		return;
-	}
-
-	int32 ValueKey;
-	
-	{
-		PCG::TUniqueScopeLock WriteLock(ValueLock);
-		const int32 NewValueIndex = ThisHelper->AddValue();
-		ValueKey = ValueKeyOffset + NewValueIndex;
-	}
-
-	{
-		PCG::TUniqueScopeLock WriteLock(EntryMapLock);
-		EntryToValueKeyMap.Emplace(ItemKey, ValueKey);
-	}
+	// Not supported, will need to be specialized in the templated versions.
 }
 
 void FPCGMetadataAttributeGeneric::AccumulateValue(PCGMetadataEntryKey ItemKey, const FPCGMetadataAttributeBase* InAttributeA,
@@ -1445,10 +1521,11 @@ void FPCGMetadataAttributeGeneric::GetValues_Internal(const TArrayView<const PCG
 			}
 			else
 			{
-				// Check if we are still in the same continuous range. If so, increment and continue.
+				// Check if we are still in the same continuous range. If so, increment and continue. We also need to mark the value as retrieved.
 				if (ContinuousRange.StartIndex != INDEX_NONE && ContinuousRange.StartIndex + ContinuousRange.Count == Index && ContinuousRange.StartValueKey + ContinuousRange.Count == ValueKeyIndex)
 				{
 					ContinuousRange.Count++;
+					UnretrievedValues[Index] = false;
 					continue;
 				}
 				else if (ContinuousRange.StartIndex != INDEX_NONE)
@@ -1568,8 +1645,23 @@ void FPCGMetadataAttributeGeneric::SetValueFromProperty(PCGMetadataEntryKey Item
 		}
 		
 		const int32 ExistingIndex = ItemKey == PCGInvalidEntryKey && ensure(DefaultValue.IsValid()) ? PCG::Private::ExistingIndexForDefaultValue : INDEX_NONE;
+		
+		PCGMetadataValueKey ValueKeyToAdd = PCGNotFoundValueKey;
+	
+		if (UnderlyingProperty->bCompressData && ItemKey != PCGInvalidEntryKey)
+		{
+			const bool bValueKeyFound = FindValues(MakeArrayView(&ValueKeyToAdd, 1), FInValuesByValue{SrcPtr, 1});
+			check(!bValueKeyFound || ValueKeyToAdd != PCGNotFoundValueKey);
+		}
 
-		SetValuesFromValueKeys({ItemKey}, AddValues_Internal(FInValues{TInPlaceType<FInValuesByValue>{}, SrcPtr, 1}, 1, ExistingIndex), /*bResetDefaultValueOnDefaultValueKey=*/ItemKey==PCGInvalidEntryKey);
+		if (ValueKeyToAdd == PCGNotFoundValueKey)
+		{
+			SetValuesFromValueKeys({ItemKey}, AddValues_Internal(FInValues{TInPlaceType<FInValuesByValue>{}, SrcPtr, 1}, 1, ExistingIndex), /*bResetDefaultValueOnDefaultValueKey=*/ItemKey==PCGInvalidEntryKey);
+		}
+		else
+		{
+			SetValueFromValueKey(ItemKey, ValueKeyToAdd);
+		}
 	}
 }
 
@@ -1607,11 +1699,8 @@ int32 FPCGMetadataAttributeGeneric::PreallocateValues(TArrayView<PCGMetadataEntr
 
 void FPCGMetadataAttributeGeneric::Prepare(int32 Count)
 {
+	PCG::TUniqueScopeLock Lock(EntryMapLock);
 	EntryToValueKeyMap.Reserve(EntryToValueKeyMap.Num() + Count);
-	if (UnderlyingProperty && !UnderlyingProperty->bCompressData)
-	{
-		AllocateValues_Unsafe(Count);
-	}
 }
 
 int32 FPCGMetadataAttributeGeneric::AllocateValues_Unsafe(int32 Count)
diff --git a/Engine/Plugins/PCG/Source/PCG/Private/Metadata/PCGMetadataCommon.cpp b/Engine/Plugins/PCG/Source/PCG/Private/Metadata/PCGMetadataCommon.cpp
--- a/Engine/Plugins/PCG/Source/PCG/Private/Metadata/PCGMetadataCommon.cpp
+++ b/Engine/Plugins/PCG/Source/PCG/Private/Metadata/PCGMetadataCommon.cpp
@@ -12,16 +12,20 @@ namespace PCGMetadataDomainID
 	const FPCGMetadataDomainID Invalid{EPCGMetadataDomainFlag::Invalid, -1, TEXT("Invalid")};
 }
 
-bool FPCGMetadataAttributeDesc::operator==(const FPCGMetadataAttributeDesc& Other) const
+bool FPCGMetadataAttributeDesc::IsSameType(const FPCGMetadataAttributeDesc& Other) const
 {
-	return Name == Other.Name &&
-		ValueType == Other.ValueType &&
+	return ValueType == Other.ValueType &&
 		ContainerTypes == Other.ContainerTypes &&
 		ValueTypeObject == Other.ValueTypeObject &&
 		KeyType == Other.KeyType &&
 		KeyTypeObject == Other.KeyTypeObject;
 }
 
+bool FPCGMetadataAttributeDesc::operator==(const FPCGMetadataAttributeDesc& Other) const
+{
+	return Name == Other.Name && IsSameType(Other);
+}
+
 int16 FPCGMetadataAttributeDesc::AsLegacyTypeId() const
 {
 	if (static_cast<uint8>(ValueType) < static_cast<uint8>(EPCGMetadataTypes::EndLegacyTypes))
diff --git a/Engine/Plugins/PCG/Source/PCG/Private/Metadata/PCGMetadataDomain.cpp b/Engine/Plugins/PCG/Source/PCG/Private/Metadata/PCGMetadataDomain.cpp
--- a/Engine/Plugins/PCG/Source/PCG/Private/Metadata/PCGMetadataDomain.cpp
+++ b/Engine/Plugins/PCG/Source/PCG/Private/Metadata/PCGMetadataDomain.cpp
@@ -166,13 +166,17 @@ void FPCGMetadataDomain::Serialize(FArchive& InArchive)
 			InArchive << AttributeTypeId;
 
 			FPCGMetadataAttributeBase* SerializedAttribute = PCGMetadataAttribute::AllocateEmptyAttributeFromType(static_cast<int16>(AttributeTypeId));
-			if (ensure(SerializedAttribute))
+			if (!SerializedAttribute)
 			{
-				SerializedAttribute->Name = AttributeName;
-				SerializedAttribute->Serialize(this, InArchive);
-				Attributes.Add(AttributeName, SerializedAttribute);
-				MaxAttributeId = FMath::Max(SerializedAttribute->AttributeId, MaxAttributeId);
+				SerializedAttribute = new FPCGMetadataAttributeGeneric();
 			}
+
+			check(SerializedAttribute);
+
+			SerializedAttribute->Name = AttributeName;
+			SerializedAttribute->Serialize(this, InArchive);
+			Attributes.Add(AttributeName, SerializedAttribute);
+			MaxAttributeId = FMath::Max(SerializedAttribute->AttributeId, MaxAttributeId);
 		}
 	}
 	else
@@ -227,11 +231,11 @@ FPCGMetadataAttributeGeneric* FPCGMetadataDomain::CreateGenericAttribute(const F
 	if (!ParentAttribute || (ParentAttribute->GetTypeId() != TypeId || !ParentAttribute->IsGeneric()))
 	{
 		// Can't parent if the types do not match, or is not generic
-		NewAttribute = new FPCGMetadataAttributeGeneric(Desc, this, AttributeName, bAllowsInterpolation);
+		NewAttribute = new FPCGMetadataAttributeGeneric(Desc, this, bAllowsInterpolation);
 	}
 	else
 	{
-		NewAttribute = new FPCGMetadataAttributeGeneric(static_cast<const FPCGMetadataAttributeGeneric*>(ParentAttribute), this, AttributeName, bAllowsInterpolation);
+		NewAttribute = new FPCGMetadataAttributeGeneric(static_cast<const FPCGMetadataAttributeGeneric*>(ParentAttribute), this, bAllowsInterpolation, AttributeName);
 	}
 	
 	{
diff --git a/Engine/Plugins/PCG/Source/PCG/Public/Metadata/Accessors/PCGAttributeAccessor.h b/Engine/Plugins/PCG/Source/PCG/Public/Metadata/Accessors/PCGAttributeAccessor.h
--- a/Engine/Plugins/PCG/Source/PCG/Public/Metadata/Accessors/PCGAttributeAccessor.h
+++ b/Engine/Plugins/PCG/Source/PCG/Public/Metadata/Accessors/PCGAttributeAccessor.h
@@ -96,8 +96,9 @@ public:
 		}
 		else
 		{
+			check(StartIndex != INDEX_NONE);
 			//TRACE_CPUPROFILER_EVENT_SCOPE(FPCGAttributeAccessor::SetRangeImpl::SetValues);
-			static_cast<UnderlyingAccessor&>(*this).SetValues_TryLockless(EntryKeys, InValues, Start);
+			static_cast<UnderlyingAccessor&>(*this).SetValues_TryLockless(EntryKeys, InValues, StartIndex + Start);
 		}
 
 		return true;
@@ -169,10 +170,10 @@ public:
 
 		if (bPreallocateValues && ensure(!bWasPrepared))
 		{
-			const int32 StartIndex = static_cast<UnderlyingAccessor&>(*this).PreallocateValues(EntryKeysView, /*bLockless=*/false);
+			const int32 TempStartIndex = static_cast<UnderlyingAccessor&>(*this).PreallocateValues(EntryKeysView, /*bLockless=*/false);
 			if (OutStartIndex)
 			{
-				*OutStartIndex = StartIndex;
+				*OutStartIndex = TempStartIndex;
 			}
 		}
 
@@ -187,13 +188,14 @@ public:
 		}
 
 		static_cast<UnderlyingAccessor&>(*this).PrepareAttribute(Count);
-		bWasPrepared = Prepare(Keys, Count, bCanReuseEntryKeys, /*Index=*/0, /*bPreallocateValues=*/true);
+		bWasPrepared = Prepare(Keys, Count, bCanReuseEntryKeys, /*Index=*/0, /*bPreallocateValues=*/true, /*OutEntryKeys=*/nullptr, &StartIndex);
 	}
 
 protected:
 	FPCGMetadataAttributeBase* Attribute = nullptr;
 	FPCGMetadataDomain* MetadataDomain = nullptr;
 	bool bWasPrepared = false;
+	int32 StartIndex = INDEX_NONE;
 };
 
 template <typename T>
diff --git a/Engine/Plugins/PCG/Source/PCG/Public/Metadata/PCGMetadataAttributeGeneric.h b/Engine/Plugins/PCG/Source/PCG/Public/Metadata/PCGMetadataAttributeGeneric.h
--- a/Engine/Plugins/PCG/Source/PCG/Public/Metadata/PCGMetadataAttributeGeneric.h
+++ b/Engine/Plugins/PCG/Source/PCG/Public/Metadata/PCGMetadataAttributeGeneric.h
@@ -33,6 +33,12 @@ namespace PCG::Private
 		|| TIsTArrayView_V<Container>
 		|| TIsStaticArray<Container>::value;
 
+	template <typename LhsType, typename RhsType>
+	concept CIsComparable = requires(const LhsType& Lhs, const RhsType& Rhs)
+	{
+		{ Lhs == Rhs } -> UE::CSameAs<bool>;
+	};
+
 	/** Magic number for AddValues to know if we need to write to the DefaultValue or not. */
 	constexpr int32 ExistingIndexForDefaultValue = -2;
 }
@@ -85,12 +91,8 @@ private:
 class FPCGMetadataAttributeGeneric : public FPCGMetadataAttributeBase
 {
 public:
-	// New attribute with no parent
-	PCG_API FPCGMetadataAttributeGeneric(const FPCGMetadataAttributeDesc& InAttributeDesc, FPCGMetadataDomain* InMetadata, FName InName, bool bInAllowsInterpolation);
+	PCG_API ~FPCGMetadataAttributeGeneric();
 
-	// New attribute with an existing parent
-	PCG_API FPCGMetadataAttributeGeneric(TNotNull<const FPCGMetadataAttributeGeneric*> InParent, FPCGMetadataDomain* InMetadata, FName InName, bool bInAllowsInterpolation);
-	
 	PCG_API void AddReferencedObjects(FReferenceCollector& Collector);
 
 	virtual bool IsGeneric() const override { return true; }
@@ -109,8 +111,16 @@ public:
 
 	const FPCGMetadataAttributeDesc& GetAttributeDesc() const { return CachedDesc; }
 
+	template <typename T>
+	bool IsOfType() const
+	{
+		return PCG::Private::IsEquivalentDesc<T>(CachedDesc);
+	}
+
 	PCG_API virtual PCGMetadataValueKey GetValueKeyOffsetForChild() const override;
-	
+
+	PCG_API virtual void Serialize(FPCGMetadataDomain* InMetadata, FArchive& InArchive) override;
+
 	/**
 	 * GetValues section
 	 * Templated code to support all container types.
@@ -127,7 +137,12 @@ public:
 	template <typename T> requires (!std::is_pointer_v<T>)
 	void GetValues(TArrayView<const PCGMetadataValueKey> ValueKeys, TArrayView<T> OutValues) const
 	{
-		if (!PCG::Private::IsEquivalentDesc<T>(CachedDesc))
+		if (ValueKeys.IsEmpty())
+		{
+			return;
+		}
+
+		if (!IsOfType<T>())
 		{
 			// @todo_pcg Better logging
 			UE_LOG(LogPCG, Error, TEXT("Type mismatch on GetValues"));
@@ -145,7 +160,12 @@ public:
 	template <typename T>
 	void GetValues(TArrayView<const PCGMetadataValueKey> ValueKeys, TArrayView<const T*> OutValues) const
 	{
-		if (!PCG::Private::IsEquivalentDesc<T>(CachedDesc))
+		if (ValueKeys.IsEmpty())
+		{
+			return;
+		}
+
+		if (!IsOfType<T>())
 		{
 			// @todo_pcg Better logging
 			UE_LOG(LogPCG, Error, TEXT("Type mismatch on GetValues"));
@@ -170,7 +190,12 @@ public:
 	template <typename T> 
 	void GetValues(TArrayView<const PCGMetadataValueKey> ValueKeys, TArrayView<TConstArrayView<T>> OutValues) const
 	{
-		if (!PCG::Private::IsEquivalentDesc<TArray<T>>(CachedDesc))
+		if (ValueKeys.IsEmpty())
+		{
+			return;
+		}
+
+		if (!IsOfType<TArray<T>>())
 		{
 			// @todo_pcg Better logging
 			UE_LOG(LogPCG, Error, TEXT("Type mismatch on GetValues"));
@@ -195,7 +220,12 @@ public:
 	template <typename T> 
 	void GetValues(TArrayView<const PCGMetadataValueKey> ValueKeys, TArrayView<PCG::TScriptSetWrapper<T>> OutValues) const
 	{
-		if (!PCG::Private::IsEquivalentDesc<TSet<T>>(CachedDesc))
+		if (ValueKeys.IsEmpty())
+		{
+			return;
+		}
+
+		if (!IsOfType<TSet<T>>())
 		{
 			// @todo_pcg Better logging
 			UE_LOG(LogPCG, Error, TEXT("Type mismatch on GetValues"));
@@ -220,7 +250,12 @@ public:
 	template <typename KeyType, typename ValueType> 
 	void GetValues(TArrayView<const PCGMetadataValueKey> ValueKeys, TArrayView<PCG::TScriptMapWrapper<KeyType, ValueType>> OutValues) const
 	{
-		if (!PCG::Private::IsEquivalentDesc<TMap<KeyType, ValueType>>(CachedDesc))
+		if (ValueKeys.IsEmpty())
+		{
+			return;
+		}
+
+		if (!IsOfType<TMap<KeyType, ValueType>>())
 		{
 			// @todo_pcg Better logging
 			UE_LOG(LogPCG, Error, TEXT("Type mismatch on GetValues"));
@@ -256,7 +291,7 @@ public:
 	template <typename T>
 	void GetValuesFromItemKeys(TConstPCGValueRange<PCGMetadataEntryKey> EntryKeys, TArrayView<T> OutValues) const
 	{
-		if (!ensure(EntryKeys.Num() == OutValues.Num()))
+		if (EntryKeys.IsEmpty() || !ensure(EntryKeys.Num() == OutValues.Num()))
 		{
 			return;
 		}
@@ -282,7 +317,7 @@ public:
 		if (TOptional<FScriptArrayHelper> Helper = MakeArrayHelper(); Helper.IsSet())
 		{
 			check(DefaultValue.IsValid());
-			AddValues(MakeArrayView(&Value, 1), /*ExistingIndex=*/PCG::Private::ExistingIndexForDefaultValue, /*Count=*/1);
+			AddValues(MakeArrayView(&Value, 1), /*ExistingIndex=*/PCG::Private::ExistingIndexForDefaultValue);
 		}
 	}
 
@@ -356,14 +391,29 @@ public:
 	PCG_API const void* GetReadAddressFromEntryKey_Unsafe(const PCGMetadataEntryKey InEntryKey) const;
 
 protected:
+	/** Attributes can't be created individually, they need to be created through a Metadata Domain. */
+	friend FPCGMetadataDomain;
+
+	// New attribute with no parent, name is contained in the InAttributeDesc.
+	PCG_API FPCGMetadataAttributeGeneric(const FPCGMetadataAttributeDesc& InAttributeDesc, FPCGMetadataDomain* InMetadata, bool bInAllowsInterpolation);
+
+	// New attribute with an existing parent, can override the name.
+	PCG_API FPCGMetadataAttributeGeneric(TNotNull<const FPCGMetadataAttributeGeneric*> InParent, FPCGMetadataDomain* InMetadata, bool bInAllowsInterpolation, TOptional<FName> InName = {});
+
+	// For deprecation purposes of the typed attributes
+	PCG_API FPCGMetadataAttributeGeneric(const FPCGMetadataAttributeDesc& InAttributeDesc, const FPCGMetadataAttributeGeneric* InParent, FPCGMetadataDomain* InMetadata, bool bInAllowsInterpolation);
+
+	// For deserialization purposes
+	FPCGMetadataAttributeGeneric() = default;
+
 	// Allocate the values (without initialization if it can) and return the first index of that new memory space.
 	// Return INDEX_NONE if it didn't work.
-	int32 AllocateValues_Unsafe(int32 Count);
+	PCG_API int32 AllocateValues_Unsafe(int32 Count);
 
 	// Used to make an array helper of the local FScriptArray.
 	// Optional, since the attribute can be ill-formed, so the user should always check if it is set before
 	// using it.
-	TOptional<FScriptArrayHelper> MakeArrayHelper() const;
+	PCG_API TOptional<FScriptArrayHelper> MakeArrayHelper() const;
 	
 	/**
 	 * GetValues implementation section
@@ -450,7 +500,12 @@ protected:
 	template <typename T>
 	TArray<PCGMetadataValueKey> AddValues(TArrayView<const T> InValues, int32 ExistingIndex = INDEX_NONE, bool bLockless = false)
 	{
-		if (!PCG::Private::IsEquivalentDesc<T>(CachedDesc))
+		if (InValues.IsEmpty())
+		{
+			return {};
+		}
+
+		if (!IsOfType<T>())
 		{
 			// @todo_pcg Better logging
 			UE_LOG(LogPCG, Error, TEXT("Type mismatch on AddValues"));
@@ -459,7 +514,7 @@ protected:
 
 		// For compressed data, we need to analyze which values actually need to be added. If existing index is not none, ignore it.
 		// Need the concept for Equals, since `AddUnique` require the operator '==` to be defined.
-		if constexpr (UE::Core::Private::CWithConstUEOpEquals<T, T>)
+		if constexpr (PCG::Private::CIsComparable<T, T>)
 		{
 			if (UnderlyingProperty->bCompressData && ExistingIndex == INDEX_NONE)
 			{
@@ -489,8 +544,13 @@ protected:
 	template <typename ArrayContainer> requires (PCG::Private::CIsArrayLike<ArrayContainer>)
 	TArray<PCGMetadataValueKey> AddValues(TConstArrayView<ArrayContainer> InValues, int32 ExistingIndex = INDEX_NONE, bool bLockless = false)
 	{
+		if (InValues.IsEmpty())
+		{
+			return {};
+		}
+
 		using T = typename ArrayContainer::ElementType;
-		if (!PCG::Private::IsEquivalentDesc<TArray<T>>(CachedDesc))
+		if (!IsOfType<TArray<T>>())
 		{
 			// @todo_pcg Better logging
 			UE_LOG(LogPCG, Error, TEXT("Type mismatch on AddValues"));
@@ -509,7 +569,12 @@ protected:
 	template <typename T, typename ...Extra>
 	TArray<PCGMetadataValueKey> AddValues(TConstArrayView<TSet<T, Extra...>> InValues, int32 ExistingIndex = INDEX_NONE, bool bLockless = false)
 	{
-		if (!PCG::Private::IsEquivalentDesc<TSet<T>>(CachedDesc))
+		if (InValues.IsEmpty())
+		{
+			return {};
+		}
+
+		if (!IsOfType<TSet<T>>())
 		{
 			// @todo_pcg Better logging
 			UE_LOG(LogPCG, Error, TEXT("Type mismatch on AddValues"));
@@ -534,7 +599,12 @@ protected:
 	template <typename KeyType, typename ValueType, typename ...Extra>
 	TArray<PCGMetadataValueKey> AddValues(TConstArrayView<TMap<KeyType, ValueType, Extra...>> InValues, int32 ExistingIndex = INDEX_NONE, bool bLockless = false)
 	{
-		if (!PCG::Private::IsEquivalentDesc<TMap<KeyType, ValueType>>(CachedDesc))
+		if (InValues.IsEmpty())
+		{
+			return {};
+		}
+
+		if (!IsOfType<TMap<KeyType, ValueType>>())
 		{
 			// @todo_pcg Better logging
 			UE_LOG(LogPCG, Error, TEXT("Type mismatch on AddValues"));
@@ -554,47 +624,69 @@ protected:
 	}
 	
 	// Internal implementation that will do the dispatch using the typed erased variant.
-	TArray<PCGMetadataValueKey> AddValues_Internal(FInValues InValues, int32 Count, int32 ExistingIndex = INDEX_NONE, bool bLockless = false);
+	PCG_API TArray<PCGMetadataValueKey> AddValues_Internal(FInValues InValues, int32 Count, int32 ExistingIndex = INDEX_NONE, bool bLockless = false);
 
 	// Internal implementation for compressed data
-	TArray<PCGMetadataValueKey> AddCompressedValues(FInValuesByValue InUniqueValues, TArray<PCGMetadataValueKey>& FoundValueKeys);
+	PCG_API TArray<PCGMetadataValueKey> AddCompressedValues(FInValuesByValue InUniqueValues, TArray<PCGMetadataValueKey>& FoundValueKeys);
 
 	/** 
 	 * Find all the value keys associated to InValues for compressed data. Return true if all the values were found.
 	 */
-	bool FindValues(TArray<PCGMetadataValueKey>& OutValueKeys, const FInValuesByValue& InValues);
+	PCG_API bool FindValues(TArray<PCGMetadataValueKey>& OutValueKeys, const FInValuesByValue& InValues) const;
+
+	/** 
+	 * Find all the value keys associated to InValues for compressed data. Return true if all the values were found. Preallocated version.
+	 */
+	PCG_API bool FindValues(TArrayView<PCGMetadataValueKey> OutValueKeys, const FInValuesByValue& InValues) const;
+
+	template <typename T>
+	void InitForDeprecation(TConstArrayView<T> InValues, const T& InDefaultValue)
+	{
+		CachedDesc.Name = Name;
+
+		if (TOptional<FScriptArrayHelper> Helper = MakeArrayHelper(); ensure(Helper))
+		{
+			Helper->EmptyValues(InValues.Num());
+			SetDefaultValue<T>(InDefaultValue);
+
+			if (!InValues.IsEmpty())
+			{
+				AddValues<T>(InValues);
+			}
+			
+			ValueKeyOffset = GetParent() ? GetParent()->GetValueKeyOffsetForChild() : 0;
+		}
+	}
+
+	using FCopyFunctor = TFunction<void(void* /*DestData*/, const void* /*SrcData*/, const int32 /*Count*/)>;
+	PCG_API void CopyInternal(FPCGMetadataAttributeGeneric* NewAttribute, bool bKeepParent, bool bCopyEntries = true, bool bCopyValues = true, FCopyFunctor CopyFunctor = {}) const;
 
 private:
 	void Init();
 	
-	void GetValues_Internal(const TArrayView<const PCGMetadataValueKey> ValueKeys, FOutValues OutValues, TBitArray<>& UnretrievedValues) const;
-	void FindValuesInternal(const FInValuesByValue& InValues, TArray<PCGMetadataValueKey>& ValueKeys, int& ValueKeysSet, bool bIsRoot) const;
+	PCG_API void GetValues_Internal(const TArrayView<const PCGMetadataValueKey> ValueKeys, FOutValues OutValues, TBitArray<>& UnretrievedValues) const;
+	void FindValuesInternal(const FInValuesByValue& InValues, TArrayView<PCGMetadataValueKey> ValueKeys, int& ValueKeysSet, bool bIsRoot) const;
 
 public:
-	virtual void Flatten() override;
-	virtual void FlattenAndCompress(const TArrayView<const PCGMetadataEntryKey>& InEntryKeysToKeep) override;
-	virtual void Reset() override;
-	virtual FPCGMetadataAttributeBase* Copy(FName NewName, FPCGMetadataDomain* InMetadata, bool bKeepParent, bool bCopyEntries = true,
-		bool bCopyValues = true) const override;
+	PCG_API virtual void Flatten() override;
+	PCG_API virtual void FlattenAndCompress(const TArrayView<const PCGMetadataEntryKey>& InEntryKeysToKeep) override;
+	PCG_API virtual void Reset() override;
+	PCG_API virtual FPCGMetadataAttributeBase* Copy(FName NewName, FPCGMetadataDomain* InMetadata, bool bKeepParent, bool bCopyEntries = true, bool bCopyValues = true) const override;
 
 	UE_DEPRECATED(5.6, "Use the version with FPCGMetadataDomain")
-	virtual FPCGMetadataAttributeBase* Copy(FName NewName, UPCGMetadata* InMetadata, bool bKeepParent, bool bCopyEntries = true,
-		bool bCopyValues = true) const override;
-
-	virtual FPCGMetadataAttributeBase* CopyToAnotherType(int16 Type) const override;
-	virtual void SetValue(PCGMetadataEntryKey ItemKey, const FPCGMetadataAttributeBase* InAttribute, PCGMetadataEntryKey InEntryKey) override;
-	virtual void SetZeroValue(PCGMetadataEntryKey ItemKey) override;
-	virtual void AccumulateValue(PCGMetadataEntryKey ItemKey, const FPCGMetadataAttributeBase* InAttributeA, PCGMetadataEntryKey InEntryKeyA,
-		float Weight) override;
-	virtual void SetWeightedValue(PCGMetadataEntryKey ItemKey, const FPCGMetadataAttributeBase* InAttribute,
-		const TArrayView<const TPair<PCGMetadataEntryKey, float>>& InWeightedKeys) override;
-	virtual void SetValue(PCGMetadataEntryKey ItemKey, const FPCGMetadataAttributeBase* InAttributeA, PCGMetadataEntryKey InEntryKeyA,
-		const FPCGMetadataAttributeBase* InAttributeB, PCGMetadataEntryKey InEntryKeyB, EPCGMetadataOp Op) override;
-	virtual bool IsEqualToDefaultValue(PCGMetadataValueKey ValueKey) const override;
-	virtual void SetDefaultValueToFirstEntry() override;
-	virtual bool UsesValueKeys() const override;
-	virtual bool AreValuesEqualForEntryKeys(PCGMetadataEntryKey EntryKey1, PCGMetadataEntryKey EntryKey2) const override;
-	virtual bool AreValuesEqual(PCGMetadataValueKey ValueKey1, PCGMetadataValueKey ValueKey2) const override;
+	PCG_API virtual FPCGMetadataAttributeBase* Copy(FName NewName, UPCGMetadata* InMetadata, bool bKeepParent, bool bCopyEntries = true, bool bCopyValues = true) const override;
+
+	PCG_API virtual FPCGMetadataAttributeBase* CopyToAnotherType(int16 Type) const override;
+	PCG_API virtual void SetValue(PCGMetadataEntryKey ItemKey, const FPCGMetadataAttributeBase* InAttribute, PCGMetadataEntryKey InEntryKey) override;
+	PCG_API virtual void SetZeroValue(PCGMetadataEntryKey ItemKey) override;
+	PCG_API virtual void AccumulateValue(PCGMetadataEntryKey ItemKey, const FPCGMetadataAttributeBase* InAttributeA, PCGMetadataEntryKey InEntryKeyA, float Weight) override;
+	PCG_API virtual void SetWeightedValue(PCGMetadataEntryKey ItemKey, const FPCGMetadataAttributeBase* InAttribute, const TArrayView<const TPair<PCGMetadataEntryKey, float>>& InWeightedKeys) override;
+	PCG_API virtual void SetValue(PCGMetadataEntryKey ItemKey, const FPCGMetadataAttributeBase* InAttributeA, PCGMetadataEntryKey InEntryKeyA, const FPCGMetadataAttributeBase* InAttributeB, PCGMetadataEntryKey InEntryKeyB, EPCGMetadataOp Op) override;
+	PCG_API virtual bool IsEqualToDefaultValue(PCGMetadataValueKey ValueKey) const override;
+	PCG_API virtual void SetDefaultValueToFirstEntry() override;
+	PCG_API virtual bool UsesValueKeys() const override;
+	PCG_API virtual bool AreValuesEqualForEntryKeys(PCGMetadataEntryKey EntryKey1, PCGMetadataEntryKey EntryKey2) const override;
+	PCG_API virtual bool AreValuesEqual(PCGMetadataValueKey ValueKey1, PCGMetadataValueKey ValueKey2) const override;
 
 private:
 	FPCGMetadataAttributeDesc CachedDesc;
diff --git a/Engine/Plugins/PCG/Source/PCG/Public/Metadata/PCGMetadataAttributeTraits.h b/Engine/Plugins/PCG/Source/PCG/Public/Metadata/PCGMetadataAttributeTraits.h
--- a/Engine/Plugins/PCG/Source/PCG/Public/Metadata/PCGMetadataAttributeTraits.h
+++ b/Engine/Plugins/PCG/Source/PCG/Public/Metadata/PCGMetadataAttributeTraits.h
@@ -144,7 +144,7 @@ namespace PCG
 		 *   - TSet<T> if T is a supported type. Does not support nested containers for now.
 		 *   - TMap<T, U> if T and U are supported types. Does not support nested containers for now.
 		 */
-		template<typename T> requires (TModels_V<CStaticStructProvider, T>)
+		template<typename T> requires (TModels_V<CStaticStructProvider, T> && !std::is_base_of_v<FSoftObjectPath, T>)
 		FPCGMetadataAttributeDesc MakeAttributeDesc(FName AttributeName)
 		{
 			return FPCGMetadataAttributeDesc{
@@ -199,7 +199,7 @@ namespace PCG
 		
 		//////////////////////////////
 		
-		template<typename T> requires (TModels_V<CStaticStructProvider, T>)
+		template<typename T> requires (TModels_V<CStaticStructProvider, T>  && !std::is_base_of_v<FSoftObjectPath, T>)
 		bool IsSameType(const EPCGMetadataTypes InType, const UObject* InTypeObject)
 		{
 			return InType == EPCGMetadataTypes::Struct &&
@@ -215,7 +215,7 @@ namespace PCG
 		
 		//////////////////////////////
 
-		template<typename T> requires (TModels_V<CStaticStructProvider, T> || TIsUEnumClass<T>::Value == true)
+		template<typename T> requires ((TModels_V<CStaticStructProvider, T>  && !std::is_base_of_v<FSoftObjectPath, T>) || TIsUEnumClass<T>::Value == true)
 		bool IsEquivalentDesc(const FPCGMetadataAttributeDesc& Desc)
 		{
 			return Desc.ContainerTypes.IsEmpty() 
@@ -475,6 +475,15 @@ namespace PCG
 
 			return ConstructibleTypes[FirstType][SecondType];
 		}
+		
+		// Convenience function for accessor operation, to know if a Get will succeed with AllowBroadcast and AllowConstructible
+		// where the SecondType is the target type, and FirstType is the accessor type.
+		// ie. Accessor is an int, Accessor->Get<float>(...). First type is int, second type is float.
+		template <typename FirstType, typename SecondType>
+		constexpr inline bool IsBroadcastableOrConstructible()
+		{
+			return IsBroadcastable<FirstType, SecondType>() || std::is_constructible_v<SecondType, FirstType>;
+		}
 
 		// Convenience function for accessor operation, to know if a Get will succeed with AllowBroadcast and AllowConstructible
 		// where the SecondType is the target type, and FirstType is the accessor type.
diff --git a/Engine/Plugins/PCG/Source/PCG/Public/Metadata/PCGMetadataCommon.h b/Engine/Plugins/PCG/Source/PCG/Public/Metadata/PCGMetadataCommon.h
--- a/Engine/Plugins/PCG/Source/PCG/Public/Metadata/PCGMetadataCommon.h
+++ b/Engine/Plugins/PCG/Source/PCG/Public/Metadata/PCGMetadataCommon.h
@@ -286,7 +286,11 @@ struct FPCGMetadataAttributeDesc
 	UPROPERTY()
 	TObjectPtr<const UObject> KeyTypeObject = nullptr;
 	
+	PCG_API bool IsSameType(const FPCGMetadataAttributeDesc& Other) const;
+
+	/** Same type and same name */
 	PCG_API bool operator==(const FPCGMetadataAttributeDesc& Other) const;
+
 	PCG_API friend int32 GetTypeHash(const FPCGMetadataAttributeDesc& Desc);
 	
 	UE_DEPRECATED(all, "Should not be used anymore as some information is lost in the process.")
diff --git a/Engine/Plugins/PCG/Tests/Private/Metadata/PCGMetadataAttributeGenericTests.cpp b/Engine/Plugins/PCG/Tests/Private/Metadata/PCGMetadataAttributeGenericTests.cpp
--- a/Engine/Plugins/PCG/Tests/Private/Metadata/PCGMetadataAttributeGenericTests.cpp
+++ b/Engine/Plugins/PCG/Tests/Private/Metadata/PCGMetadataAttributeGenericTests.cpp
@@ -45,6 +45,8 @@ namespace PCGAttributeGenericTests
 	};
 	
 	static const FName AttributeName = "Attr";
+	static const FName AttributeName2 = "Attr2";
+	static const FName AttributeName3 = "Attr3";
 	
 	/** Getter to create a tester for different types. */
 	TypedAttributeTester<float> FloatTester()
@@ -847,27 +849,56 @@ TEST_CASE_METHOD(PCGTests::FPCGBaseTestWithBasicContext, "PCG::AttributeGeneric:
 	TPCGValueRange<int64> PointArrayDataMetadataEntryRange = PointArrayData->GetMetadataEntryValueRange(/*bAllocate=*/false);
 
 	FPCGMetadataDomain* Domain = PointArrayData->MutableMetadata()->GetDefaultMetadataDomain();
+
+	// Have base values for SoftPath and SoftClass
+	const FString BasePathFormat = TEXT("/PCGTests/Item_{0}.Item_{0}");
+	TStaticArray<const UClass*, 10> UClassesPool =
+	{
+		UPCGData::StaticClass(),
+		UPCGParamData::StaticClass(),
+		UPCGBasePointData::StaticClass(),
+		UPCGPointData::StaticClass(),
+		UPCGPointArrayData::StaticClass(),
+		UObject::StaticClass(),
+		UClass::StaticClass(),
+		UPCGSettings::StaticClass(),
+		UPCGSettingsInterface::StaticClass(),
+		UPCGSettingsInstance::StaticClass(),
+	};
 	
 	// Initialize the default value to a value we will add later on.
 	FPCGMetadataAttributeGeneric* GenericAttribute = Domain->CreateGenericAttribute<FString>(PCGAttributeGenericTests::AttributeName, TEXT("Value: 0"));
+	FPCGMetadataAttributeGeneric* GenericAttributePath = Domain->CreateGenericAttribute<FSoftObjectPath>(PCGAttributeGenericTests::AttributeName2, FSoftObjectPath(FString::Format(*BasePathFormat, { 0 })));
+	FPCGMetadataAttributeGeneric* GenericAttributeClass = Domain->CreateGenericAttribute<FSoftClassPath>(PCGAttributeGenericTests::AttributeName3, UClassesPool[0]);
 	REQUIRE_NOT_EQUAL(GenericAttribute, nullptr);
+	REQUIRE_NOT_EQUAL(GenericAttributePath, nullptr);
+	REQUIRE_NOT_EQUAL(GenericAttributeClass, nullptr);
 	
 	// First set half the points with different values
 	TStaticArray<FString, NumPoints> Values{};
+	TStaticArray<FSoftObjectPath, NumPoints> PathValues{};
+	TStaticArray<FSoftClassPath, NumPoints> ClassValues{};
 	TStaticArray<PCGMetadataEntryKey, NumPoints/2> EntryKeys{};
-	
+
+
 	for (int32 i = 0; i < NumPoints/2; ++i)
 	{
 		Domain->InitializeOnSet(PointArrayDataMetadataEntryRange[i]);
 		EntryKeys[i] = PointArrayDataMetadataEntryRange[i];
-		Values[i] = FString::Printf(TEXT("Value: %d"), i % 10);
+		Values[i] = FString::Printf(TEXT("Value: %d"), i % 5);
+		PathValues[i] = FSoftObjectPath(FString::Format(*BasePathFormat, { i % 5 }));
+		ClassValues[i] = UClassesPool[i % 5];
 	}
 	
 	// Set the values
 	GenericAttribute->SetValues<FString>(EntryKeys, MakeArrayView(Values.GetData(), NumPoints / 2));
+	GenericAttributePath->SetValues<FSoftObjectPath>(EntryKeys, MakeArrayView(PathValues.GetData(), NumPoints / 2));
+	GenericAttributeClass->SetValues<FSoftClassPath>(EntryKeys, MakeArrayView(ClassValues.GetData(), NumPoints / 2));
 	
-	// Verify that we have only 9 Value keys (the tenth one is the default value).
-	REQUIRE_EQUAL(GenericAttribute->GetValueKeyOffsetForChild(), 9);
+	// Verify that we have only 4 Value keys (the fifth one is the default value).
+	REQUIRE_EQUAL(GenericAttribute->GetValueKeyOffsetForChild(), 4);
+	REQUIRE_EQUAL(GenericAttributePath->GetValueKeyOffsetForChild(), 4);
+	REQUIRE_EQUAL(GenericAttributeClass->GetValueKeyOffsetForChild(), 4);
 	
 	// Then add a mix of existing strings and new strings, to test that we only add what does not exist
 	for (int32 i = NumPoints/2; i < NumPoints; ++i)
@@ -875,27 +906,42 @@ TEST_CASE_METHOD(PCGTests::FPCGBaseTestWithBasicContext, "PCG::AttributeGeneric:
 		const int32 EntryKeyIndex = i - (NumPoints/2);
 		Domain->InitializeOnSet(PointArrayDataMetadataEntryRange[i]);
 		EntryKeys[EntryKeyIndex] = PointArrayDataMetadataEntryRange[i];
-		if (EntryKeyIndex < 10)
+		if (EntryKeyIndex < 5)
 		{
-			Values[i] = FString::Printf(TEXT("Value: %d"), EntryKeyIndex % 10);
+			Values[i] = FString::Printf(TEXT("Value: %d"), EntryKeyIndex % 5);
+			PathValues[i] = FSoftObjectPath(FString::Format(*BasePathFormat, { i % 5 }));
+			ClassValues[i] = UClassesPool[i % 5];
 		}
 		else
 		{
-			Values[i] = FString::Printf(TEXT("New Value: %d"), EntryKeyIndex % 10);
+			Values[i] = FString::Printf(TEXT("New Value: %d"), EntryKeyIndex % 5);
+			PathValues[i] = FSoftObjectPath(FString::Format(*BasePathFormat, { 5 + i % 5 }));
+			ClassValues[i] = UClassesPool[5 + i % 5];
 		}
 	}
 	
 	// Set the values
 	GenericAttribute->SetValues<FString>(EntryKeys, MakeArrayView(Values.GetData() + (NumPoints / 2), NumPoints / 2));
+	GenericAttributePath->SetValues<FSoftObjectPath>(EntryKeys, MakeArrayView(PathValues.GetData() + (NumPoints / 2), NumPoints / 2));
+	GenericAttributeClass->SetValues<FSoftClassPath>(EntryKeys, MakeArrayView(ClassValues.GetData() + (NumPoints / 2), NumPoints / 2));
 	
-	// Verify that we have only 19 Value keys (the twentieth one is the default value)
-	REQUIRE_EQUAL(GenericAttribute->GetValueKeyOffsetForChild(), 19);
+	// Verify that we have only 9 Value keys (the tenth one is the default value)
+	REQUIRE_EQUAL(GenericAttribute->GetValueKeyOffsetForChild(), 9);
+	REQUIRE_EQUAL(GenericAttributePath->GetValueKeyOffsetForChild(), 9);
+	REQUIRE_EQUAL(GenericAttributeClass->GetValueKeyOffsetForChild(), 9);
 	
 	// Verify that all the values can be retrieved.
 	TStaticArray<const FString*, NumPoints> ReadValues{};
+	TStaticArray<const FSoftObjectPath*, NumPoints> ReadValuesPath{};
+	TStaticArray<const FSoftClassPath*, NumPoints> ReadValuesClass{};
+
 	GenericAttribute->GetValuesFromItemKeys<const FString*>(PCGValueRangeHelpers::MakeConstValueRange(PointArrayDataMetadataEntryRange), ReadValues);
+	GenericAttributePath->GetValuesFromItemKeys<const FSoftObjectPath*>(PCGValueRangeHelpers::MakeConstValueRange(PointArrayDataMetadataEntryRange), ReadValuesPath);
+	GenericAttributeClass->GetValuesFromItemKeys<const FSoftClassPath*>(PCGValueRangeHelpers::MakeConstValueRange(PointArrayDataMetadataEntryRange), ReadValuesClass);
 	
 	CHECK_THAT(Values, Catch::Matchers::RangeEquals(ReadValues, [](const FString& a, const FString* b){ return a == *b; } ));
+	CHECK_THAT(PathValues, Catch::Matchers::RangeEquals(ReadValuesPath, [](const FSoftObjectPath& a, const FSoftObjectPath* b){ return a == *b; } ));
+	CHECK_THAT(ClassValues, Catch::Matchers::RangeEquals(ReadValuesClass, [](const FSoftClassPath& a, const FSoftClassPath* b){ return a == *b; } ));
 }
 
 TEST_CASE_METHOD(PCGTests::FPCGBaseTestWithBasicContext, "PCG::AttributeGeneric::MultiThreadedWrite", "[PCG][AttributeGeneric]")
@@ -1134,4 +1180,210 @@ TEST_CASE_METHOD(PCGTests::FPCGBaseTestWithBasicContext, "PCG::AttributeGeneric:
 			Tester(PCGAttributeGenericTests::EnumTester(), /*bLockless=*/true);
 		}
 	}
+}
+
+/**
+ * Test that SetValue(ItemKey, this, EntryKey) works correctly when the source attribute is the same
+ * as the destination. This exercises the fix where SrcPtr obtained from GetReadAddressFromEntryKey_Unsafe
+ * could become dangling after AddValue() triggers a reallocation of the internal Values buffer.
+ * We use enough self-copies to force multiple reallocations of the underlying FScriptArray.
+ */
+TEST_CASE_METHOD(PCGTests::FPCGBaseTestWithBasicContext, "PCG::AttributeGeneric::SetValueSelfReference", "[PCG][AttributeGeneric]")
+{
+	auto Tester = [this]<typename T, typename ReadType>(PCGAttributeGenericTests::TypedAttributeTester<T, ReadType> AttributeTester)
+	{
+		constexpr int32 NumInitialEntries = 10;
+		// Enough self-copies to trigger multiple reallocations of the underlying FScriptArray.
+		constexpr int32 NumFullSelfCopies = 20;
+
+		UPCGParamData* ParamData = CreateData<UPCGParamData>();
+		FPCGMetadataDomain* Domain = ParamData->MutableMetadata()->GetDefaultMetadataDomain();
+
+		FRandomStream RandomStream(42);
+
+		FPCGMetadataAttributeGeneric* GenericAttribute = Domain->CreateGenericAttribute<T>(PCGAttributeGenericTests::AttributeName, AttributeTester.DefaultValue);
+		REQUIRE_NOT_EQUAL(GenericAttribute, nullptr);
+
+		// Populate initial entries with random values.
+		TStaticArray<T, NumInitialEntries> InitialValues{};
+		TStaticArray<PCGMetadataEntryKey, NumInitialEntries * (NumFullSelfCopies + 1)> AllEntryKeys;
+
+		for (int32 i = 0; i < NumInitialEntries; ++i)
+		{
+			AllEntryKeys[i] = Domain->AddEntry();
+			AttributeTester.GenerateRandom(RandomStream, InitialValues[i]);
+		}
+
+		GenericAttribute->SetValues<T>(MakeArrayView(AllEntryKeys.GetData(), NumInitialEntries), InitialValues);
+		REQUIRE_EQUAL(GenericAttribute->GetNumberOfEntries(), NumInitialEntries);
+
+		// Copying into itself
+		for (int32 i = 0; i < NumInitialEntries * NumFullSelfCopies; ++i)
+		{
+			AllEntryKeys[NumInitialEntries + i] = Domain->AddEntry();
+			GenericAttribute->SetValue(AllEntryKeys[NumInitialEntries + i], GenericAttribute, AllEntryKeys[i % NumInitialEntries]);
+		}
+		
+		REQUIRE_EQUAL(GenericAttribute->GetNumberOfEntries(), NumInitialEntries * (NumFullSelfCopies + 1));
+		TStaticArray<ReadType, NumInitialEntries> ReadValues{};
+		for (int32 i = 0; i < NumFullSelfCopies + 1; ++i)
+		{
+			const int32 StartIndex = i * NumInitialEntries;
+			
+			GenericAttribute->GetValuesFromItemKeys<ReadType>(MakeArrayView(AllEntryKeys.GetData() + StartIndex, NumInitialEntries), ReadValues);
+			CHECK_THAT(InitialValues, Catch::Matchers::RangeEquals(ReadValues, [&AttributeTester](const auto& a, const auto& b){ return AttributeTester.Verify(a, b); } ));
+		}
+	};
+
+	SECTION("Double (POD type)")
+	{
+		Tester(PCGAttributeGenericTests::DoubleTester());
+	}
+
+	SECTION("String (non-POD, compressed)")
+	{
+		Tester(PCGAttributeGenericTests::StringTester());
+	}
+
+	SECTION("Vector")
+	{
+		Tester(PCGAttributeGenericTests::VectorTester_Basic());
+	}
+
+	SECTION("FPCGPoint (Struct)")
+	{
+		Tester(PCGAttributeGenericTests::FPCGPointTester());
+	}
+}
+
+/**
+ * Test that reading continuous values on GetValues is working as intended, even with parenting.
+ */
+TEST_CASE_METHOD(PCGTests::FPCGBaseTestWithBasicContext, "PCG::AttributeGeneric::GetValuesContinuous", "[PCG][AttributeGeneric]")
+{
+	auto Tester = [this]<typename T, typename ReadType>(PCGAttributeGenericTests::TypedAttributeTester<T, ReadType> AttributeTester)
+	{
+		// Preparing the data
+		constexpr int32 NumPoints = 100;
+		UPCGPointArrayData* PointArrayData = CreateData<UPCGPointArrayData>();
+
+		PointArrayData->SetNumPoints(NumPoints);
+
+		PointArrayData->AllocateProperties(EPCGPointNativeProperties::MetadataEntry);
+		TPCGValueRange<int64> PointArrayDataMetadataEntryRange = PointArrayData->GetMetadataEntryValueRange(/*bAllocate=*/false);
+
+		FPCGMetadataDomain* Domain = PointArrayData->MutableMetadata()->GetDefaultMetadataDomain();
+
+		FRandomStream RandomStream(42);
+
+		FPCGMetadataAttributeGeneric* GenericAttribute = nullptr;
+
+		GenericAttribute = Domain->CreateGenericAttribute<T>(PCGAttributeGenericTests::AttributeName, AttributeTester.DefaultValue);
+		REQUIRE_NOT_EQUAL(GenericAttribute, nullptr);
+
+		TStaticArray<T, NumPoints> Values{};
+		TStaticArray<ReadType, NumPoints> ReadValues;
+
+		for (int32 i = 0; i < NumPoints; ++i)
+		{
+			PointArrayDataMetadataEntryRange[i] = Domain->AddEntry(PointArrayDataMetadataEntryRange[i]);
+			AttributeTester.GenerateRandom(RandomStream, Values[i]);
+		}
+
+		GenericAttribute->SetValues<T>(PointArrayDataMetadataEntryRange, Values);
+		REQUIRE_EQUAL(GenericAttribute->GetNumberOfEntries(), Values.Num());
+
+		SECTION("Read back and verify values")
+		{
+			GenericAttribute->GetValuesFromItemKeys<ReadType>(PCGValueRangeHelpers::MakeConstValueRange(PointArrayDataMetadataEntryRange), ReadValues);
+			CHECK_THAT(Values, Catch::Matchers::RangeEquals(ReadValues, [&AttributeTester](const auto& a, const auto& b) { return AttributeTester.Verify(a, b); }));
+		}
+
+		SECTION("Parenting with more points")
+		{
+			UPCGPointArrayData* ChildPointArrayData = CreateData<UPCGPointArrayData>();
+			ChildPointArrayData->InitializeFromDataWithParams(FPCGInitializeFromDataParams{ PointArrayData });
+
+			FPCGMetadataDomain* ChildDomain = ChildPointArrayData->MutableMetadata()->GetDefaultMetadataDomain();
+
+			ChildPointArrayData->SetNumPoints(2 * NumPoints);
+
+			TPCGValueRange<int64> ChildPointArrayDataMetadataEntryRange = ChildPointArrayData->GetMetadataEntryValueRange(/*bAllocate=*/false);
+
+			FPCGMetadataAttributeGeneric* ChildGenericAttribute = ChildDomain->GetMutableGenericAttribute(PCGAttributeGenericTests::AttributeName);
+			REQUIRE_NOT_EQUAL(ChildGenericAttribute, nullptr);
+
+			// Add more values
+			TStaticArray<T, NumPoints> ChildValues{};
+			TStaticArray<PCGMetadataEntryKey, NumPoints> ChildEntryKeys{};
+			for (int32 i = 0; i < NumPoints; ++i)
+			{
+				ChildEntryKeys[i] = ChildDomain->AddEntry();
+				ChildPointArrayDataMetadataEntryRange[i + NumPoints] = ChildEntryKeys[i];
+				AttributeTester.GenerateRandom(RandomStream, ChildValues[i]);
+			}
+
+			ChildGenericAttribute->SetValues<T>(ChildEntryKeys, ChildValues);
+			REQUIRE_EQUAL(ChildGenericAttribute->GetNumberOfEntries(), NumPoints);
+			REQUIRE_EQUAL(ChildGenericAttribute->GetNumberOfEntriesWithParents(), 2 * NumPoints);
+
+			SECTION("Read back and verify values")
+			{
+				TStaticArray<ReadType, 2*NumPoints> ChildReadValues;
+				ChildGenericAttribute->GetValuesFromItemKeys<ReadType>(PCGValueRangeHelpers::MakeConstValueRange(ChildPointArrayDataMetadataEntryRange), ChildReadValues);
+				CHECK_THAT(Values, Catch::Matchers::RangeEquals(MakeArrayView(ChildReadValues.GetData(), NumPoints), [&AttributeTester](const auto& a, const auto& b) { return AttributeTester.Verify(a, b); }));
+				CHECK_THAT(ChildValues, Catch::Matchers::RangeEquals(MakeArrayView(ChildReadValues.GetData() + NumPoints, NumPoints), [&AttributeTester](const auto& a, const auto& b) { return AttributeTester.Verify(a, b); }));
+			}
+		}
+	};
+
+	SECTION("Float Attribute")
+	{
+		Tester(PCGAttributeGenericTests::FloatTester());
+	}
+
+	SECTION("Double Attribute")
+	{
+		Tester(PCGAttributeGenericTests::DoubleTester());
+	}
+
+	SECTION("As Basic Type")
+	{
+		Tester(PCGAttributeGenericTests::VectorTester_Basic());
+	}
+
+	SECTION("String Attribute")
+	{
+		Tester(PCGAttributeGenericTests::StringTester());
+	}
+
+	SECTION("FPCGPoint Attribute")
+	{
+		Tester(PCGAttributeGenericTests::FPCGPointTester());
+	}
+
+	SECTION("Array of Double Attribute")
+	{
+		Tester(PCGAttributeGenericTests::ArrayDoubleTester());
+	}
+
+	SECTION("Array of FPCGPoint Attribute")
+	{
+		Tester(PCGAttributeGenericTests::ArrayFPCGPointTester());
+	}
+
+	SECTION("Set of Double Attribute")
+	{
+		Tester(PCGAttributeGenericTests::SetDoubleTester());
+	}
+
+	SECTION("Map of String -> Double Attribute")
+	{
+		Tester(PCGAttributeGenericTests::MapStringDoubleTester());
+	}
+
+	SECTION("Enum attribute")
+	{
+		Tester(PCGAttributeGenericTests::EnumTester());
+	}
 }
\ No newline at end of file
