[PCG] Revert change when Metadata was added to spline

* We added the Output metadata to the SamplePoint call of the bounding shape
* We should not sample the metadata for the bounding shape, we should just test the intersection.
* This was adding entries to the metadata for no reasons.

#rnx

#Synced-CL 51431209
#Synced-CL 51431209

diff --git a/Engine/Plugins/PCG/Source/PCG/Private/Elements/PCGSplineSampler.cpp b/Engine/Plugins/PCG/Source/PCG/Private/Elements/PCGSplineSampler.cpp
--- a/Engine/Plugins/PCG/Source/PCG/Private/Elements/PCGSplineSampler.cpp
+++ b/Engine/Plugins/PCG/Source/PCG/Private/Elements/PCGSplineSampler.cpp
@@ -625,7 +625,7 @@ namespace PCGSplineSamplerHelpers
 
 			// Only add valid points within bounds
 			FPCGPoint BoundsTestPoint;
-			if (bValid && (!BoundingShapeData || BoundingShapeData->SamplePoint(Transform, InResult.Box, BoundsTestPoint, OutPointData->Metadata)))
+			if (bValid && (!BoundingShapeData || BoundingShapeData->SamplePoint(Transform, InResult.Box, BoundsTestPoint, nullptr)))
 			{
 				SetSeed(OutPoint.Seed, OutPoint.Transform, InResult.LocalTransform.GetLocation(), Params, Seed, InResult.SampleIndex);
 				SetMetadata(InResult, OutPoint, OutPointData->Metadata);
@@ -742,7 +742,7 @@ namespace PCGSplineSamplerHelpers
 
 					// Prune points outside of bounds
 					FPCGPoint BoundsTestPoint;
-					if (BoundingShapeData && !BoundingShapeData->SamplePoint(OutPoint.Transform, OutPoint.GetLocalBounds(), BoundsTestPoint, OutPointData->Metadata))
+					if (BoundingShapeData && !BoundingShapeData->SamplePoint(OutPoint.Transform, OutPoint.GetLocalBounds(), BoundsTestPoint, nullptr))
 					{
 						continue;
 					}
@@ -1237,7 +1237,7 @@ namespace PCGSplineSamplerHelpers
 
 						// Prune points outside of bounds
 						FPCGPoint BoundsTestPoint;
-						if (InBoundingShape && !InBoundingShape->SamplePoint(TransformWS, GeneratedPointBounds, BoundsTestPoint, OutPointData->Metadata))
+						if (InBoundingShape && !InBoundingShape->SamplePoint(TransformWS, GeneratedPointBounds, BoundsTestPoint, nullptr))
 						{
 							continue;
 						}
