[PCG] Added some robustness around the graph inspected stack to prevent potential crashes

#rnx
#Synced-CL 51379665

diff --git a/Engine/Plugins/PCG/Source/PCG/Private/PCGContext.cpp b/Engine/Plugins/PCG/Source/PCG/Private/PCGContext.cpp
--- a/Engine/Plugins/PCG/Source/PCG/Private/PCGContext.cpp
+++ b/Engine/Plugins/PCG/Source/PCG/Private/PCGContext.cpp
@@ -621,7 +621,7 @@ bool FPCGContext::IsExecutingGraphInspected() const
 	const FPCGStack* StackPtr = GetStack();
 	const UPCGGraph* Graph = StackPtr ? StackPtr->GetGraphForCurrentFrame() : nullptr;
 
-	return Graph && Graph->IsInspecting() && Graph->GetInspectedStack() == *StackPtr;
+	return Graph && Graph->IsSameInspectedStack(*StackPtr);
 }
 
 AActor* FPCGContext::GetDebugDisplayActor() const
diff --git a/Engine/Plugins/PCG/Source/PCG/Private/PCGElement.cpp b/Engine/Plugins/PCG/Source/PCG/Private/PCGElement.cpp
--- a/Engine/Plugins/PCG/Source/PCG/Private/PCGElement.cpp
+++ b/Engine/Plugins/PCG/Source/PCG/Private/PCGElement.cpp
@@ -54,7 +54,7 @@ static TAutoConsoleVariable<bool> CVarPCGEnablePointArrayDataToPointDataConversi
 
 #if WITH_EDITOR
 #define PCG_ELEMENT_EXECUTION_BREAKPOINT() \
-	if (Context && Context->IsExecutingGraphInspected() && Context->GetInputSettingsInterface() && Context->GetInputSettingsInterface()->bBreakDebugger) \
+	if (Context && Context->GetInputSettingsInterface() && Context->GetInputSettingsInterface()->bBreakDebugger && Context->IsExecutingGraphInspected()) \
 	{ \
 		UE_DEBUG_BREAK(); \
 	}
diff --git a/Engine/Plugins/PCG/Source/PCG/Private/PCGGraph.cpp b/Engine/Plugins/PCG/Source/PCG/Private/PCGGraph.cpp
--- a/Engine/Plugins/PCG/Source/PCG/Private/PCGGraph.cpp
+++ b/Engine/Plugins/PCG/Source/PCG/Private/PCGGraph.cpp
@@ -1855,6 +1855,32 @@ void UPCGGraph::RemoveExtraEditorNode(const UObject* InNode)
 	ExtraEditorNodes.Remove(const_cast<UObject*>(InNode));
 }
 
+void UPCGGraph::EnableInspection(const FPCGStack& InInspectedStack)
+{
+	PCG::TUniqueScopeLock WriteLock(InspectionStateLock);
+	bIsInspecting = true;
+	InspectedStack = InInspectedStack;
+}
+
+void UPCGGraph::DisableInspection()
+{
+	PCG::TUniqueScopeLock WriteLock(InspectionStateLock);
+	bIsInspecting = false;
+	InspectedStack = FPCGStack();
+}
+
+FPCGStack UPCGGraph::GetInspectedStack() const
+{
+	PCG::TSharedScopeLock ReadLock(InspectionStateLock);
+	return InspectedStack;
+}
+
+bool UPCGGraph::IsSameInspectedStack(const FPCGStack& InOtherStack) const
+{
+	PCG::TSharedScopeLock ReadLock(InspectionStateLock);
+	return bIsInspecting && InspectedStack == InOtherStack;
+}
+
 bool UPCGGraph::PrimeGraphCompilationCache()
 {
 	UPCGSubsystem* Subsystem = UPCGSubsystem::GetActiveEditorInstance();
diff --git a/Engine/Plugins/PCG/Source/PCG/Public/PCGGraph.h b/Engine/Plugins/PCG/Source/PCG/Public/PCGGraph.h
--- a/Engine/Plugins/PCG/Source/PCG/Public/PCGGraph.h
+++ b/Engine/Plugins/PCG/Source/PCG/Public/PCGGraph.h
@@ -539,10 +539,13 @@ public:
 	UE_API void RemoveExtraEditorNode(const UObject* InNode);
 
 	bool IsInspecting() const { return bIsInspecting; }
-	void EnableInspection(const FPCGStack& InInspectedStack) { bIsInspecting = true; InspectedStack = InInspectedStack; }
-	void DisableInspection() { bIsInspecting = false; InspectedStack = FPCGStack(); }
+	UE_API void EnableInspection(const FPCGStack& InInspectedStack);
+	UE_API void DisableInspection();
 	bool DebugFlagAppliesToIndividualComponents() const { return bDebugFlagAppliesToIndividualComponents; }
-	const FPCGStack& GetInspectedStack() const { return InspectedStack; }
+	UE_API bool IsSameInspectedStack(const FPCGStack& InOtherStack) const;
+
+	UE_DEPRECATED(5.8, "Compare using the IsSameInspectedStack method.")
+	UE_API FPCGStack GetInspectedStack() const;
 	
 	/** Instruct the graph compiler to cache the relevant permutations of this graph. */
 	UE_API bool PrimeGraphCompilationCache();
@@ -778,6 +781,7 @@ private:
 	bool bUserPausedNotificationsInGraphEditor = false;
 	bool bIsInspecting = false;
 	FPCGStack InspectedStack;
+	mutable PCG::FSharedLock InspectionStateLock;
 #endif // WITH_EDITOR
 
 	// Deprecated section
