[PCG] Attribute Partition - Mixed Radix Composite Key Hash Map upgrade

#Synced-CL 51463632

diff --git a/Engine/Plugins/PCG/Source/PCG/Private/Metadata/PCGMetadataPartitionCommon.cpp b/Engine/Plugins/PCG/Source/PCG/Private/Metadata/PCGMetadataPartitionCommon.cpp
--- a/Engine/Plugins/PCG/Source/PCG/Private/Metadata/PCGMetadataPartitionCommon.cpp
+++ b/Engine/Plugins/PCG/Source/PCG/Private/Metadata/PCGMetadataPartitionCommon.cpp
@@ -23,9 +23,25 @@
 
 namespace PCGMetadataPartitionCommon
 {
-	template <typename T> 
+	using IndexPartition = TArray<TArray<int32>>;
+
+	// @todo_pcg: Evaluate the bit array for removal as a contingency slow path
+	template <typename T>
 	constexpr bool IsBitArray = std::is_same_v<T, TBitArray<>>;
 
+	// @todo_pcg: To be removed or replaced by internal procedural selection if/when the algorithm is validated
+	static TAutoConsoleVariable<bool> CVarPCGMetadataPartitionMixedRadixHash(
+		TEXT("pcg.MetadataPartition.UseMixedRadixHash"),
+		true,
+		TEXT("Enables the mixed radix hash method for partitioning attributes."));
+
+	static TAutoConsoleVariable<bool> CVarPCGMetadataPartitionValidateOverflow(
+		TEXT("pcg.MetadataPartition.ValidateOverflow"),
+		true,
+		TEXT("Validates that the composite key space does not overflow uint64 before using the mixed radix algorithm. "
+			"Falls back to bitwise intersection on overflow. Disabling may cause silent partition collisions if the product "
+			"of unique-value counts across all attributes exceeds uint64 (2^64)."));
+
 	/**
 	* Partition a given attribute, by first partitioning all value keys that point to the same value
 	* and then for each unique value key, list of index in the keys that match for this value.
@@ -314,61 +330,94 @@ namespace PCGMetadataPartitionCommon
 		return AttributeGenericPartition<TArray<int32>>(InData, InSelector, InOptionalContext, bSilenceMissingAttributeErrors);
 	}
 
-	/**
-	 * Partition on multiple attributes by first partitioning on the attributes independently. Then take the resultant
-	 * partition and convert them to a BitArray representation of each element's partition. Once in BitArray form,
-	 * combine the results with a logical AND operation to filter them into final partition groupings.
-	 *
-	 * Multi-Partition Example:
-	 * Pt  A  B  C                         Partition on A->[0,1],[2,3,4]
-	 *  0  a  a  a                         Partition on B->[0],[1,2],[3,4]
-	 *  1  a  b  a                         Partition on C->[0,1],[2],[3,4]
-	 *  2  b  b  b                         Partition on A&B->[0],[1],[2],[3,4]
-	 *  3  b  c  c                         Final Partition (A&B&C)->[0],[1],[2],[3,4]
-	 *  4  b  c  c
-	 */
-	TArray<TArray<int32>> AttributeGenericPartition(const UPCGData* InData, const TArrayView<const FPCGAttributePropertySelector>& InSelectorArrayView, FPCGContext* InOptionalContext, bool bSilenceMissingAttributeErrors)
+	/** Combine per-attribute partitions into a final partition using mixed radix composite key encoding. */
+	IndexPartition MixedRadixCompositeKeyPartition(const TArray<IndexPartition>& PerAttributePartitions, const int32 NumElements, const int32 NumAttributes)
 	{
-		TRACE_CPUPROFILER_EVENT_SCOPE(PCGMetadataPartitionCommon::AttributeGenericPartition::MultiSelector);
+		TRACE_CPUPROFILER_EVENT_SCOPE(PCGMetadataPartitionCommon::AttributeGenericPartition::MultiSelector::CompositeKeyGrouping);
+		IndexPartition FinalPartition;
 
-		// Small optimization to partition on a single attribute
-		if (InSelectorArrayView.Num() == 1)
-		{
-			return AttributeGenericPartition<TArray<int32>>(InData, InSelectorArrayView[0], InOptionalContext, bSilenceMissingAttributeErrors);
-		}
+		TArray<int32> PartitionCounts;
+		PartitionCounts.SetNum(NumAttributes);
 
-		if (!InData || InSelectorArrayView.IsEmpty() || !InData->ConstMetadata())
+		TArray<TArray<int32>> PointToPartitionIndex;
+		PointToPartitionIndex.SetNum(NumAttributes);
+
+		// Build reverse mapping (point -> partition index) and collect partition counts for composite key encoding
 		{
-			return {};
+			TRACE_CPUPROFILER_EVENT_SCOPE(PCGMetadataPartitionCommon::AttributeGenericPartition::MultiSelector::ReverseMapping);
+			for (int32 AttributeIndex = 0; AttributeIndex < NumAttributes; ++AttributeIndex)
+			{
+				const IndexPartition& AttributePartition = PerAttributePartitions[AttributeIndex];
+				PartitionCounts[AttributeIndex] = AttributePartition.Num();
+
+				TArray<int32>& PointToPartitionMap = PointToPartitionIndex[AttributeIndex];
+				// By default, this is set to -1, because 0 is a valid partition index.
+				PointToPartitionMap.Init(INDEX_NONE, NumElements);
+
+				for (int32 PartitionIndex = 0; PartitionIndex < AttributePartition.Num(); ++PartitionIndex)
+				{
+					for (const int32 PointIndex : AttributePartition[PartitionIndex])
+					{
+						PointToPartitionMap[PointIndex] = PartitionIndex;
+					}
+				}
+			}
 		}
 
-		// Get the element count from the number of keys which should work for spatial points and attribute sets
-		const TUniquePtr<const IPCGAttributeAccessorKeys> Keys = PCGAttributeAccessorHelpers::CreateConstKeys(InData, InSelectorArrayView[0]);
-		if (!Keys.IsValid())
+		// Maps the composite key (see below) to the index in FinalPartition with one entry per unique attribute combination.
+		TMap<uint64, int32> CompositeKeyToPartition;
+
+		for (int32 PointIndex = 0; PointIndex < NumElements; ++PointIndex)
 		{
-			PCGLog::LogErrorOnGraph(FText::Format(LOCTEXT("InvalidKeys", "Could not create keys for the input data with selector {0}"), InSelectorArrayView[0].GetDisplayText()), InOptionalContext);
-			return {};
+			// https://en.wikipedia.org/wiki/Composite_key
+			uint64 CompositeKey = 0;
+
+			// Use mixed radix encoding to create unique keys, scaling multiplicatively on the partition counts.
+			// https://en.wikipedia.org/wiki/Mixed_radix
+			uint64 Multiplier = 1;
+			for (int32 AttributeIndex = 0; AttributeIndex < NumAttributes; ++AttributeIndex)
+			{
+				CompositeKey += Multiplier * PointToPartitionIndex[AttributeIndex][PointIndex];
+				Multiplier *= PartitionCounts[AttributeIndex];
+			}
+
+			// Look up or create the final partition for this key.
+			const int32* FinalIndex = CompositeKeyToPartition.Find(CompositeKey);
+			if (!FinalIndex)
+			{
+				FinalIndex = &CompositeKeyToPartition.Add(CompositeKey, FinalPartition.Num());
+				FinalPartition.Emplace();
+			}
+
+			FinalPartition[*FinalIndex].Add(PointIndex);
 		}
 
-		const int64 NumElements = Keys->GetNum();
-		const int64 NumAttributes = InSelectorArrayView.Num();
+		return FinalPartition;
+	}
 
-		using IndexPartition = TArray<TArray<int32>>;
+	/** Combine per-attribute partitions into a final partition using bitwise intersection. */
+	IndexPartition BitwiseIntersectionPartition(const TArray<IndexPartition>& PerAttributePartitions, int32 NumElements)
+	{
 		using BitPartition = TArray<TBitArray<>>;
 
 		TArray<BitPartition> BitPartitions;
+		const int32 NumAttributes = PerAttributePartitions.Num();
 		BitPartitions.SetNum(NumAttributes);
 
 		{
 			TRACE_CPUPROFILER_EVENT_SCOPE(PCGMetadataPartitionCommon::AttributeGenericPartition::MultiSelector::PartitionOnBitArray);
-			/* TODO: Can be executed in parallel, threadsafe. There is a follow up task to evaluate between
-			* option A.) partitioning on all attributes, and then merging and B.) Partitioning on each attribute
-			* in succession, further partitioning the grouping results of the previous iteration.
-			*/
-			// Calculate each partition into a bitfield for simple/efficient intersection processing
-			for (int32 I = 0; I < InSelectorArrayView.Num(); ++I)
+			for (int32 AttributeIndex = 0; AttributeIndex < NumAttributes; ++AttributeIndex)
 			{
-				BitPartitions[I] = AttributeGenericPartition<TBitArray<>>(InData, InSelectorArrayView[I], InOptionalContext, bSilenceMissingAttributeErrors);
+				BitPartition& Partition = BitPartitions[AttributeIndex];
+				Partition.SetNum(PerAttributePartitions[AttributeIndex].Num());
+				for (int32 PartitionIndex = 0; PartitionIndex < PerAttributePartitions[AttributeIndex].Num(); ++PartitionIndex)
+				{
+					Partition[PartitionIndex].Init(false, NumElements);
+					for (const int32 ElementIndex : PerAttributePartitions[AttributeIndex][PartitionIndex])
+					{
+						Partition[PartitionIndex][ElementIndex] = true;
+					}
+				}
 			}
 		}
 
@@ -417,6 +466,101 @@ namespace PCGMetadataPartitionCommon
 		return FinalPartition;
 	}
 
+	/**
+	 * Partition on multiple attributes by first partitioning on the attributes independently. Then build a reverse
+	 * mapping from each point to its per-attribute partition index. Finally, group points by their composite key
+	 * (a mixed radix computation of partition indices) in a single pass. O(N*M).
+	 *
+	 * Implementer's Note: A previous slow path algorithm for bitwise intersection remains as a fallback in the event of
+	 * an overflow (uint64) or if bypassed explicitly via CVar.
+	 *
+	 * Multi-Partition Example:
+	 * Pt  A  B  C                         Partition on A->[0,1],[2,3,4]
+	 *  0  a  a  a                         Partition on B->[0],[1,2],[3,4]
+	 *  1  a  b  a                         Partition on C->[0,1],[2],[3,4]
+	 *  2  b  b  b                         Composite keys: 0=(0,0,0), 1=(0,1,0), 2=(1,1,1), 3=(1,2,1), 4=(1,2,1)
+	 *  3  b  c  c                         Final Partition (A&B&C)->[0],[1],[2],[3,4]
+	 *  4  b  c  c
+	 */
+	TArray<TArray<int32>> AttributeGenericPartition(const UPCGData* InData, const TArrayView<const FPCGAttributePropertySelector>& InSelectorArrayView, FPCGContext* InOptionalContext, bool bSilenceMissingAttributeErrors)
+	{
+		TRACE_CPUPROFILER_EVENT_SCOPE(PCGMetadataPartitionCommon::AttributeGenericPartition::MultiSelector);
+
+		// Small optimization to partition on a single attribute
+		if (InSelectorArrayView.Num() == 1)
+		{
+			return AttributeGenericPartition<TArray<int32>>(InData, InSelectorArrayView[0], InOptionalContext, bSilenceMissingAttributeErrors);
+		}
+
+		if (!InData || InSelectorArrayView.IsEmpty() || !InData->ConstMetadata())
+		{
+			return {};
+		}
+
+		// Get the element count from the number of keys which should work for spatial points and attribute sets
+		const TUniquePtr<const IPCGAttributeAccessorKeys> Keys = PCGAttributeAccessorHelpers::CreateConstKeys(InData, InSelectorArrayView[0]);
+		if (!Keys.IsValid())
+		{
+			PCGLog::LogErrorOnGraph(FText::Format(LOCTEXT("InvalidKeys", "Could not create keys for the input data with selector {0}"), InSelectorArrayView[0].GetDisplayText()), InOptionalContext);
+			return {};
+		}
+
+		const int32 NumElements = Keys->GetNum();
+		const int32 NumAttributes = InSelectorArrayView.Num();
+
+		TArray<IndexPartition> PerAttributePartitions;
+		PerAttributePartitions.SetNum(NumAttributes);
+
+		{
+			TRACE_CPUPROFILER_EVENT_SCOPE(PCGMetadataPartitionCommon::AttributeGenericPartition::MultiSelector::PerAttributePartition);
+
+			// Partition per attribute
+			for (int32 I = 0; I < NumAttributes; ++I)
+			{
+				PerAttributePartitions[I] = AttributeGenericPartition<TArray<int32>>(InData, InSelectorArrayView[I], InOptionalContext, bSilenceMissingAttributeErrors);
+				// A partition could be empty if the selector is bad, which would result in a 0 value multiplier later.
+				if (PerAttributePartitions[I].IsEmpty())
+				{
+					PCGLog::LogErrorOnGraph(FText::Format(LOCTEXT("AttributeResultsInZeroPartitions", "Attribute resulted in zero partitions. Check selector {0}"), InSelectorArrayView[I].GetDisplayText()), InOptionalContext);
+					return {};
+				}
+			}
+		}
+
+		bool bUseMixedRadixHash = CVarPCGMetadataPartitionMixedRadixHash.GetValueOnAnyThread();
+
+		if (CVarPCGMetadataPartitionValidateOverflow.GetValueOnAnyThread())
+		{
+			// Overflow check: the composite key space is the product of all partition counts. With uint64, this supports
+			// a very large range, i.e. 6 attributes with 1000 unique values each or 20 attributes with up to 9 unique
+			// values each. But, as a nuclear option, we can check and bail out early.
+			// Validated with 1M points across 1000 unique partitions with no issues.
+			uint64 Multiplier = 1;
+			for (int32 AttributeIndex = 0; AttributeIndex < NumAttributes; ++AttributeIndex)
+			{
+				const int32 NumPartitions = PerAttributePartitions[AttributeIndex].Num();
+				if (Multiplier > (std::numeric_limits<uint64>::max() / FMath::Max(NumPartitions, 1)))
+				{
+					PCGLog::LogErrorOnGraph(LOCTEXT("AttributePartitionOverflow", "Overflow in partitioning too many unique attribute partitions. Falling back to slower algorithm."), InOptionalContext);
+					bUseMixedRadixHash = false;
+					break;
+				}
+
+				Multiplier *= NumPartitions;
+			}
+		}
+
+		// Select algorithm: mixed radix composite key (fast path) or bitwise intersection (slow path)
+		if (bUseMixedRadixHash)
+		{
+			return MixedRadixCompositeKeyPartition(PerAttributePartitions, NumElements, NumAttributes);
+		}
+		else
+		{
+			return BitwiseIntersectionPartition(PerAttributePartitions, NumElements);
+		}
+	}
+
 	/**
 	* Do a partition on the given point data for the selector
 	*/
@@ -494,7 +638,12 @@ namespace PCGMetadataPartitionCommon
 
 	TArray<UPCGData*> AttributeParamSpatialPartition(const UPCGData* InData, const TArrayView<const FPCGAttributePropertySelector>& InSelectorArray, FPCGContext* InOptionalContext, bool bSilenceMissingAttributeErrors)
 	{
-		if (!InData || (!InData->IsA<UPCGSpatialData>() && !InData->IsA<UPCGParamData>()))
+		if (!InData)
+		{
+			return {};
+		}
+
+		if (!InData->IsA<UPCGSpatialData>() && !InData->IsA<UPCGParamData>())
 		{
 			PCGLog::LogErrorOnGraph(LOCTEXT("InvalidDataType", "Input data is not an attribute set nor a spatial data. Operation not supported."), InOptionalContext);
 			return {};
@@ -566,7 +715,12 @@ namespace PCGMetadataPartitionCommon
 
 	UPCGData* RemoveDuplicatesParamSpatial(const UPCGData* InData, const TArrayView<const FPCGAttributePropertySelector>& InSelectorArray, FPCGContext* InOptionalContext, bool bSilenceMissingAttributeErrors)
 	{
-		if (!InData || (!InData->IsA<UPCGSpatialData>() && !InData->IsA<UPCGParamData>()))
+		if (!InData)
+		{
+			return nullptr;
+		}
+
+		if (!InData->IsA<UPCGSpatialData>() && !InData->IsA<UPCGParamData>())
 		{
 			PCGLog::LogErrorOnGraph(LOCTEXT("InvalidDataType", "Input data is not an attribute set nor a spatial data. Operation not supported."), InOptionalContext);
 			return nullptr;
