[PCG] Source Data Storage Serialization

- Removed ensure from GetAs and GetAsMutable, as there is already logging

#Synced-CL 51463632

diff --git a/Engine/Plugins/PCG/Source/PCG/Private/Graph/PCGSourceDataContainer.cpp b/Engine/Plugins/PCG/Source/PCG/Private/Graph/PCGSourceDataContainer.cpp
--- a/Engine/Plugins/PCG/Source/PCG/Private/Graph/PCGSourceDataContainer.cpp
+++ b/Engine/Plugins/PCG/Source/PCG/Private/Graph/PCGSourceDataContainer.cpp
@@ -2,6 +2,10 @@
 
 #include "Graph/PCGSourceDataContainer.h"
 
+#include "Serialization/Archive.h"
+#include "StructUtils/InstancedStruct.h"
+#include "UObject/FortniteMainBranchObjectVersion.h"
+
 #include UE_INLINE_GENERATED_CPP_BY_NAME(PCGSourceDataContainer)
 
 #define LOCTEXT_NAMESPACE "PCGSourceDataContainer"
@@ -37,6 +41,71 @@ FPCGSourceDataContainer& FPCGSourceDataContainer::operator=(FPCGSourceDataContai
 	return *this;
 }
 
+bool FPCGSourceDataContainer::Serialize(FArchive& Ar)
+{
+	Ar.UsingCustomVersion(FFortniteMainBranchObjectVersion::GUID);
+
+	Ar << DirtyGeneration;
+
+	int32 NumEntries = DataStorage.Num();
+	Ar << NumEntries;
+
+	// Load container
+	if (Ar.IsLoading())
+	{
+		if (!ensure(NumEntries >= 0))
+		{
+			return false;
+		}
+
+		DataStorage.Empty(NumEntries);
+
+		for (int32 Index = 0; Index < NumEntries; ++Index)
+		{
+			FPCGSourceDataStorageKey Key;
+
+			Ar << Key.Label;
+			Ar << Key.Hash;
+
+			// @todo_pcg: Write a custom helper 'SerializeSharedStruct' in a future pass to avoid the copy overhead.
+			FInstancedStruct InstancedValue;
+			InstancedValue.Serialize(Ar);
+
+			if (Ar.IsError())
+			{
+				DataStorage.Empty();
+				return false;
+			}
+
+			if (InstancedValue.IsValid())
+			{
+				FPCGSourceDataStorageValue Value;
+				Value.Data = FSharedStruct::Make(InstancedValue.GetScriptStruct(), InstancedValue.GetMemory());
+				DataStorage.Add(Key, MoveTemp(Value));
+			}
+		}
+	}
+	// Save container
+	else if (Ar.IsSaving())
+	{
+		for (TPair<FPCGSourceDataStorageKey, FPCGSourceDataStorageValue>& Pair : DataStorage)
+		{
+			FPCGSourceDataStorageKey& Key = Pair.Key;
+			FPCGSourceDataStorageValue& Value = Pair.Value;
+
+			Ar << Key.Label;
+			Ar << Key.Hash;
+
+			// @todo_pcg: Much like loading above, could be replaced with a custom serializer for the FSharedStruct.
+			FInstancedStruct InstancedValue;
+			InstancedValue.InitializeAs(Value.Data.GetScriptStruct(), Value.Data.GetMemory());
+			InstancedValue.Serialize(Ar);
+		}
+	}
+
+	return true;
+}
+
 FPCGSourceDataContainer& FPCGSourceDataContainer::operator=(const FPCGSourceDataContainer& Other)
 {
 	if (this != &Other)
diff --git a/Engine/Plugins/PCG/Source/PCG/Public/Graph/PCGSourceDataContainer.h b/Engine/Plugins/PCG/Source/PCG/Public/Graph/PCGSourceDataContainer.h
--- a/Engine/Plugins/PCG/Source/PCG/Public/Graph/PCGSourceDataContainer.h
+++ b/Engine/Plugins/PCG/Source/PCG/Public/Graph/PCGSourceDataContainer.h
@@ -71,6 +71,9 @@ struct FPCGSourceDataContainer
 	PCG_API FPCGSourceDataContainer& operator=(const FPCGSourceDataContainer& Other);
 	PCG_API FPCGSourceDataContainer& operator=(FPCGSourceDataContainer&& Other);
 
+	/** Custom serialization for the container. Serializes type-erased FSharedStruct payloads via FInstancedStruct. */
+	bool Serialize(FArchive& Ar);
+
 #if WITH_EDITOR
 	/** Store data of type T within the storage container. */
 	template <CPCGSourceDataCompatibleType T>
@@ -122,11 +125,20 @@ private:
 #endif // WITH_EDITOR
 };
 
+template<>
+struct TStructOpsTypeTraits<FPCGSourceDataContainer> : public TStructOpsTypeTraitsBase2<FPCGSourceDataContainer>
+{
+	enum
+	{
+		WithSerializer = true
+	};
+};
+
 template <CPCGSourceDataCompatibleType T>
 FConstSharedStruct FPCGSourceDataStorageValue::GetAs() const
 {
 #if WITH_EDITOR
-	if (!ensure(Data.GetPtr<T>() || !Data.IsValid()))
+	if (!(Data.GetPtr<T>() || !Data.IsValid()))
 	{
 		const FString RequestedType = T::StaticStruct()->GetName();
 		const FString StoredType = Data.GetScriptStruct()->GetName();
@@ -141,7 +153,7 @@ template <CPCGSourceDataCompatibleType T>
 FSharedStruct FPCGSourceDataStorageValue::GetAsMutable()
 {
 #if WITH_EDITOR
-	if (!ensure(Data.GetPtr<T>() || !Data.IsValid()))
+	if (!(Data.GetPtr<T>() || !Data.IsValid()))
 	{
 		const FString RequestedType = T::StaticStruct()->GetName();
 		const FString StoredType = Data.GetScriptStruct()->GetName();
