[PCG] Add generic method in PCGBlueprintHelper to expose the ExtractStructureAsAttributeSet function in BP

#rnx

#Synced-CL 51431209

diff --git a/Engine/Plugins/PCG/Source/PCG/Private/Helpers/PCGBlueprintHelpers.cpp b/Engine/Plugins/PCG/Source/PCG/Private/Helpers/PCGBlueprintHelpers.cpp
--- a/Engine/Plugins/PCG/Source/PCG/Private/Helpers/PCGBlueprintHelpers.cpp
+++ b/Engine/Plugins/PCG/Source/PCG/Private/Helpers/PCGBlueprintHelpers.cpp
@@ -5,10 +5,12 @@
 #include "PCGComponent.h"
 #include "PCGContext.h"
 #include "PCGModule.h"
+#include "PCGParamData.h"
 #include "PCGPoint.h"
 #include "Grid/PCGPartitionActor.h"
 #include "Grid/PCGLandscapeCache.h"
 #include "Helpers/PCGHelpers.h"
+#include "Helpers/PCGPropertyHelpers.h"
 #include "Subsystems/PCGEngineSubsystem.h"
 #include "Subsystems/PCGSubsystem.h"
 
@@ -58,15 +60,13 @@ void UPCGGenerateGraphAsync::OnGenerationDone(IPCGGraphExecutionSource* InExecut
 
 void UPCGBlueprintHelpers::ThrowBlueprintException(const FText& ErrorMessage)
 {
+	PCGLog::LogErrorOnGraph(ErrorMessage);
+
 	if (FFrame::GetThreadLocalTopStackFrame() && FFrame::GetThreadLocalTopStackFrame()->Object)
 	{
 		const FBlueprintExceptionInfo ExceptionInfo(EBlueprintExceptionType::FatalError, ErrorMessage);
 		FBlueprintCoreDelegates::ThrowScriptException(FFrame::GetThreadLocalTopStackFrame()->Object, *FFrame::GetThreadLocalTopStackFrame(), ExceptionInfo);
 	}
-	else
-	{
-		UE_LOG(LogPCG, Error, TEXT("%s"), *ErrorMessage.ToString());
-	}
 }
 
 int UPCGBlueprintHelpers::ComputeSeedFromPosition(const FVector& InPosition)
@@ -397,4 +397,43 @@ UPCGData* UPCGBlueprintHelpers::DuplicateDataWithContext(const UPCGData* InData,
 
 	UE_LOG(LogPCG, Error, TEXT("UPCGBlueprintHelpers::DuplicateDataWithContext : Invalid context handle"));
 	return nullptr;
+}
+
+DEFINE_FUNCTION(UPCGBlueprintHelpers::execExtractStructureAsAttributeSet)
+{
+	Stack.StepCompiledIn<FStructProperty>(nullptr);
+	void* StructPtr = Stack.MostRecentPropertyAddress;
+	FStructProperty* StructProp = CastField<FStructProperty>(Stack.MostRecentProperty);
+
+	P_GET_PROPERTY(FNameProperty, PropertyPath);
+	P_GET_PROPERTY(FNameProperty, OutputAttributeName);
+	P_GET_PROPERTY(FBoolProperty, bExtract);
+
+	P_FINISH;
+
+	P_NATIVE_BEGIN;
+
+	UPCGParamData* Result = nullptr;
+
+	const bool bHasValidStructProp = StructProp && StructProp->Struct && StructPtr;
+
+	FPCGAttributePropertySelector InputSelector = FPCGAttributePropertySelector::CreateSelectorFromString<FPCGAttributePropertySelector>(PropertyPath.ToString());
+
+	if (bHasValidStructProp)
+	{
+		PCGPropertyHelpers::FExtractorParameters Params(StructPtr, StructProp->Struct, InputSelector, OutputAttributeName, bExtract, /*bInPropertyNeedsToBeVisible=*/true);
+		Result = PCGPropertyHelpers::ExtractPropertyAsAttributeSet(Params);
+		if (!Result)
+		{
+			ThrowBlueprintException(NSLOCTEXT("PCGPropertyHelpers", "FailedToExtract", "Failed to extract the struct as a ParamData. Check the log for more information."));
+		}
+	}
+	else
+	{
+		ThrowBlueprintException(NSLOCTEXT("PCGPropertyHelpers", "InvalidStruct", "Parameter passed into \"Structure To Convert\" is not a Structure."));
+	}
+
+	*static_cast<UPCGParamData**>(RESULT_PARAM) = Result;
+
+	P_NATIVE_END;
 }
\ No newline at end of file
diff --git a/Engine/Plugins/PCG/Source/PCG/Public/Helpers/PCGBlueprintHelpers.h b/Engine/Plugins/PCG/Source/PCG/Public/Helpers/PCGBlueprintHelpers.h
--- a/Engine/Plugins/PCG/Source/PCG/Public/Helpers/PCGBlueprintHelpers.h
+++ b/Engine/Plugins/PCG/Source/PCG/Public/Helpers/PCGBlueprintHelpers.h
@@ -14,6 +14,7 @@ class IPCGGraphExecutionSource;
 class UPCGComponent;
 class UPCGData;
 class UPCGGraphInterface;
+class UPCGParamData;
 class UPCGSettings;
 
 struct FPCGBlueprintContextHandle;
@@ -171,4 +172,18 @@ public:
 	/** Return a copy of the data, with Metadata inheritance for spatial data. */
 	UFUNCTION(BlueprintCallable, Category = "PCG|Data", meta = (ScriptMethod))
 	static PCG_API UPCGData* DuplicateDataWithContext(const UPCGData* InData, const FPCGBlueprintContextHandle& ContextHandle, bool bInitializeMetadata = true);
+
+private:
+	// ~ This method cannot be called outside of BP, use PCGPropertyHelpers::ExtractPropertyAsAttributeSet.
+	/**
+	 * Extract a given property from a structure into an Attribute Set. If the Property is None, will extract the full structure. Only works with a structure as input.
+	 * @param Struct Structure to extract
+	 * @param PropertyPath Path to the property in the structure to extract. Must be a valid property. If None, will extract the structure as input.
+	 * @param OutputAttributeName If the property to extract is not a struct/object, can provide an override to the output attribute name. Use @Source to use the PropertyName.
+	 * @param bExtract If the property path is not None and reference an object/struct, give the possibility to extract all its members in one go.
+	 * @return Extracted param data. Can be null if something went wrong.
+	 */
+	UFUNCTION(BlueprintCallable, CustomThunk, Category = "PCG|Data", meta = (CustomStructureParam = "Struct", AllowAbstract = "false"))
+	static PCG_API UPCGParamData* ExtractStructureAsAttributeSet(UPARAM(DisplayName = "Structure To Convert") const int32& Struct, const FName PropertyPath = NAME_None, const FName OutputAttributeName = "@Source", const bool bExtract = true);
+	DECLARE_FUNCTION(execExtractStructureAsAttributeSet);
 };
