[PCG] Fix GPU-instanced primitives flushed when game viewport resized.

#Synced-CL 51538157

diff --git a/Engine/Plugins/PCG/Source/PCG/Private/Subsystems/PCGSubsystem.cpp b/Engine/Plugins/PCG/Source/PCG/Private/Subsystems/PCGSubsystem.cpp
--- a/Engine/Plugins/PCG/Source/PCG/Private/Subsystems/PCGSubsystem.cpp
+++ b/Engine/Plugins/PCG/Source/PCG/Private/Subsystems/PCGSubsystem.cpp
@@ -20,7 +20,9 @@
 #include "RuntimeGen/PCGRuntimeGenScheduler.h"
 
 #include "ComponentReregisterContext.h"
+#include "UnrealClient.h"
 #include "Algo/Transform.h"
+#include "Engine/GameViewportClient.h"
 #include "Engine/Level.h"
 
 #include UE_INLINE_GENERATED_CPP_BY_NAME(PCGSubsystem)
@@ -214,6 +216,11 @@ void UPCGSubsystem::Deinitialize()
 	Scalability::OnScalabilitySettingsChanged.RemoveAll(this);
 	FGlobalComponentReregisterContext::OnGlobalComponentReregisterContextDestroyed().RemoveAll(this);
 
+	if (GEngine && GEngine->GameViewport && GEngine->GameViewport->Viewport)
+	{
+		GEngine->GameViewport->Viewport->ViewportResizedEvent.RemoveAll(this);
+	}
+
 #if WITH_EDITOR
 	FCoreUObjectDelegates::OnObjectsReplaced.RemoveAll(this);
 
@@ -286,6 +293,11 @@ void UPCGSubsystem::PostInitialize()
 	FGlobalComponentReregisterContext::OnGlobalComponentReregisterContextDestroyed().AddUObject(this, &UPCGSubsystem::OnGlobalComponentReregisterContextDestroyed);
 
 	Scalability::OnScalabilitySettingsChanged.AddUObject(this, &UPCGSubsystem::OnScalabilitySettingsChanged);
+
+	if (GEngine && GEngine->GameViewport && GEngine->GameViewport->Viewport)
+	{
+		GEngine->GameViewport->Viewport->ViewportResizedEvent.AddUObject(this, &UPCGSubsystem::OnViewportResized);
+	}
 }
 
 UPCGSubsystem* UPCGSubsystem::GetInstance(UWorld* World)
@@ -367,7 +379,6 @@ void UPCGSubsystem::Tick(float DeltaSeconds)
 		return;
 	}
 
-#if WITH_EDITOR
 	// Refresh dirty runtime gen sources before ticking the graph executor (in IPCGBaseSubsystem::Tick()) so that we're not processing a bunch of graphs/execution sources that will be cancelled immediately.
 	if (bDirtyRuntimeGenExecutionSources)
 	{
@@ -376,7 +387,6 @@ void UPCGSubsystem::Tick(float DeltaSeconds)
 		bDirtyRuntimeGenExecutionSources = false;
 		DirtyRuntimeGenChangeType = EPCGChangeType::None;
 	}
-#endif
 
 	Super::Tick(DeltaSeconds);
 
@@ -1081,13 +1091,13 @@ void UPCGSubsystem::RefreshAllRuntimeGenExecutionSources(EPCGChangeType ChangeTy
 	}
 }
 
-#if WITH_EDITOR
 void UPCGSubsystem::DirtyRuntimeGenExecutionSources(EPCGChangeType ChangeType)
 {
 	bDirtyRuntimeGenExecutionSources = true;
 	DirtyRuntimeGenChangeType |= ChangeType;
 }
 
+#if WITH_EDITOR
 void UPCGSubsystem::RefreshAllComponentsFiltered(const TFunction<bool(UPCGComponent*)>& ComponentFilter, EPCGChangeType ChangeType)
 {
 	// @todo_pcg: support execution source
@@ -1449,6 +1459,11 @@ FPCGGenSourceManager* UPCGSubsystem::GetGenSourceManager() const
 	return RuntimeGenScheduler ? RuntimeGenScheduler->GenSourceManager : nullptr;
 }
 
+void UPCGSubsystem::OnViewportResized(FViewport* Viewport, uint32 val)
+{
+	DirtyRuntimeGenExecutionSources();
+}
+
 void UPCGSubsystem::ResetCacheStats()
 {
 	if (GraphExecutor)
diff --git a/Engine/Plugins/PCG/Source/PCG/Public/Subsystems/PCGSubsystem.h b/Engine/Plugins/PCG/Source/PCG/Public/Subsystems/PCGSubsystem.h
--- a/Engine/Plugins/PCG/Source/PCG/Public/Subsystems/PCGSubsystem.h
+++ b/Engine/Plugins/PCG/Source/PCG/Public/Subsystems/PCGSubsystem.h
@@ -20,6 +20,7 @@ class APCGWorldActor;
 class FPCGTrackingManager;
 class FPCGGenSourceManager;
 class FPCGRuntimeGenScheduler;
+class FViewport;
 class UPCGComponent;
 class UPCGComputeGraph;
 class UPCGData;
@@ -217,10 +218,10 @@ public:
 	/** Asks the runtime generation scheduler to refresh all GenerateAtRuntime execution sources. ChangeType should be 'GenerationGrid' to perform a full cleanup of PAs and local sources. */
 	PCG_API void RefreshAllRuntimeGenExecutionSources(EPCGChangeType ChangeType = EPCGChangeType::None);
 
-#if WITH_EDITOR
 	/** Asks the runtime generation scheduler to refresh all GenerateAtRuntime execution sources. ChangeType should be 'GenerationGrid' to perform a full cleanup of PAs and local sources. */
 	PCG_API void DirtyRuntimeGenExecutionSources(EPCGChangeType ChangeType = EPCGChangeType::None);
 
+#if WITH_EDITOR
 	/** Refresh all components selected by the filter (runtime generated or otherwise). */
 	PCG_API void RefreshAllComponentsFiltered(const TFunction<bool(UPCGComponent*)>& ComponentFilter, EPCGChangeType ChangeType = EPCGChangeType::None);
 #endif
@@ -323,6 +324,8 @@ public:
 
 	PCG_API FPCGGenSourceManager* GetGenSourceManager() const;
 
+	void OnViewportResized(FViewport* Viewport, uint32 val);
+
 	void ResetCacheStats();
 	void LogCacheStats();
 
@@ -455,6 +458,9 @@ private:
 
 	PCG::FLock PCGWorldActorLock;
 
+	bool bDirtyRuntimeGenExecutionSources = false;
+	EPCGChangeType DirtyRuntimeGenChangeType = EPCGChangeType::None;
+
 #if WITH_EDITOR
 	/** Debug transient actor spawned to display in world PCG debug visualization data */
 	TWeakObjectPtr<AActor> DebugActor;
@@ -464,9 +470,6 @@ private:
 	using FConstructionScriptSourceComponents = TMap<FName, TObjectKey<UPCGComponent>>;
 	TMap<TObjectKey<AActor>, FConstructionScriptSourceComponents> PerActorConstructionScriptSourceComponents;
 
-	bool bDirtyRuntimeGenExecutionSources = false;
-	EPCGChangeType DirtyRuntimeGenChangeType = EPCGChangeType::None;
-
 	PCG_API static TSet<UWorld*> DisablePartitionActorCreationForWorld;
 
 	// Used by UPCGWorldPartitonBuilder to disable PA creation while outside of a certain scope
