[PCG] GPU graph compiler fixes for kernels with multiple outputs.

Also improve kernel wiring validation to guide user better.

#Synced-CL 51537890

diff --git a/Engine/Plugins/PCG/Source/PCG/Private/Compute/PCGDataBinding.cpp b/Engine/Plugins/PCG/Source/PCG/Private/Compute/PCGDataBinding.cpp
--- a/Engine/Plugins/PCG/Source/PCG/Private/Compute/PCGDataBinding.cpp
+++ b/Engine/Plugins/PCG/Source/PCG/Private/Compute/PCGDataBinding.cpp
@@ -178,7 +178,21 @@ TSharedPtr<const FPCGDataCollectionDesc> UPCGDataBinding::GetCachedKernelPinData
 		return nullptr;
 	}
 
-	return GetCachedKernelPinDataDesc(Graph->GetBindingIndex(FPCGKernelPin(InKernel->GetKernelIndex(), InPinLabel, bIsInput)));
+	const int32 GraphBindingIndex = Graph->GetBindingIndex(FPCGKernelPin(InKernel->GetKernelIndex(), InPinLabel, bIsInput));
+
+	if (GraphBindingIndex != INDEX_NONE)
+	{
+		return GetCachedKernelPinDataDesc(GraphBindingIndex);
+	}
+	else
+	{
+		UE_LOGF(LogPCG, Error, "No binding index found for %s pin '%ls' of kernel '%ls'.",
+			bIsInput ? "input" : "output",
+			*InPinLabel.ToString(),
+			*InKernel->GetName());
+
+		return nullptr;
+	}
 }
 
 TSharedPtr<const FPCGDataCollectionDesc> UPCGDataBinding::GetCachedKernelPinDataDesc(int32 InGraphBindingIndex) const
@@ -196,7 +210,7 @@ TSharedPtr<const FPCGDataCollectionDesc> UPCGDataBinding::GetCachedKernelPinData
 	if (!FoundDescription)
 	{
 		// Should not happen - all data descriptions should have been computed at this point.
-		UE_LOG(LogPCG, Warning, TEXT("Cache miss while retrieving kernel pin data description, binding index %d."), InGraphBindingIndex);
+		UE_LOG(LogPCG, Error, TEXT("Cache miss while retrieving kernel pin data description for binding index (%d)."), InGraphBindingIndex);
 		return nullptr;
 	}
 
diff --git a/Engine/Plugins/PCG/Source/PCG/Private/Graph/PCGGraphCompilerGPU.cpp b/Engine/Plugins/PCG/Source/PCG/Private/Graph/PCGGraphCompilerGPU.cpp
--- a/Engine/Plugins/PCG/Source/PCG/Private/Graph/PCGGraphCompilerGPU.cpp
+++ b/Engine/Plugins/PCG/Source/PCG/Private/Graph/PCGGraphCompilerGPU.cpp
@@ -35,6 +35,61 @@ namespace PCGGraphCompilerGPU
 		TEXT("pcg.GraphExecution.GPU.EnableDebugging"),
 		false,
 		TEXT("Enable verbose logging of GPU compilation and execution."));
+
+	static void ValidateNodeKernels(const UPCGSettings* InNodeSettings, const TArray<UPCGComputeKernel*>& InNodeKernels, const TArray<FPCGKernelEdge>& InNodeKernelEdges)
+	{
+		UE_CLOGF(InNodeKernels.IsEmpty(), LogPCG, Error, "PCG GPU compiler: Settings '%ls' did not emit any kernels, check implementation of UPCGSettings::CreateKernels().", *InNodeSettings->GetName());
+		UE_CLOGF(InNodeKernels.Contains(nullptr), LogPCG, Error, "PCG GPU compiler: Settings '%ls' emitted one or more null kernels, check implementation of UPCGSettings::CreateKernels().", *InNodeSettings->GetName());
+
+		auto ValidateKernelPin = [InNodeSettings](const FPCGPinReference& Pin, bool bIsInputPin)
+		{
+			if (Pin.Kernel)
+			{
+				bool bPinFound = false;
+
+				if (bIsInputPin)
+				{
+					TArray<FPCGPinProperties> Pins;
+					Pin.Kernel->GetInputPins(Pins);
+					bPinFound = Pins.ContainsByPredicate([InNodeSettings, &Pin](const FPCGPinProperties& Other) { return Other.Label == Pin.Label; });
+				}
+				else
+				{
+					TArray<FPCGPinPropertiesGPU> Pins;
+					Pin.Kernel->GetOutputPins(Pins);
+					bPinFound = Pins.ContainsByPredicate([InNodeSettings, &Pin](const FPCGPinPropertiesGPU& Other) { return Other.Label == Pin.Label; });
+				}
+
+				if (!bPinFound)
+				{
+					UE_LOGF(LogPCG, Error, "PCG GPU compiler: Settings '%ls' emitted a kernel edge connecting to %s pin '%ls' of kernel '%ls' which does not exist, check the kernel is returning this pin from Get%sPins().",
+						*InNodeSettings->GetName(),
+						bIsInputPin ? "input" : "output",
+						*Pin.Label.ToString(),
+						*Pin.Kernel->GetName(),
+						bIsInputPin ? "Input" : "Output");
+				}
+			}
+			else
+			{
+				TArray<FPCGPinProperties> Pins = bIsInputPin ? InNodeSettings->InputPinProperties() : InNodeSettings->OutputPinProperties();
+
+				if (!Pins.ContainsByPredicate([InNodeSettings, &Pin](const FPCGPinProperties& Other) { return Other.Label == Pin.Label; }))
+				{
+					UE_LOGF(LogPCG, Error, "PCG GPU compiler: Settings '%ls' emitted a kernel edge connecting to external %s pin '%ls' of the node which does not exist, check the node really has this pin.",
+						*InNodeSettings->GetName(),
+						bIsInputPin ? "input" : "output",
+						*Pin.Label.ToString());
+				}
+			}
+		};
+
+		for (const FPCGKernelEdge& Edge : InNodeKernelEdges)
+		{
+			ValidateKernelPin(Edge.UpstreamPin, /*bIsInputPin=*/!Edge.GetUpstreamKernel()); // Either wired to an output pin of an upstream kernel, or a node input pin.
+			ValidateKernelPin(Edge.DownstreamPin, /*bIsInputPin=*/!!Edge.GetDownstreamKernel()); // Either wired to an input pin of a downstream kernel, or a node output pin.
+		}
+	}
 }
 
 #define PCG_DATA_DESCRIPTION_LOGGING 0
@@ -305,21 +360,22 @@ void FPCGGraphCompilerGPU::ExpandGPUNodeKernelsToTasks(
 		NodeKernelEdges.Reset();
 		Settings->CreateKernels(InOutContext, GetTransientPackage(), NodeKernels, NodeKernelEdges);
 
-#ifdef PCG_GPU_KERNEL_PROFILING
-		if (NodeKernels.IsValidIndex(Settings->ProfileKernelIndex))
-		{
-			NodeKernels[Settings->ProfileKernelIndex]->EnableRepeatDispatch();
-		}
-#endif // PCG_GPU_KERNEL_PROFILING
+		PCGGraphCompilerGPU::ValidateNodeKernels(Settings, NodeKernels, NodeKernelEdges);
 
-		ensureMsgf(NodeKernels.Remove(nullptr) == 0, TEXT("Settings '%s' returned one or more null kernels."), *Settings->GetName());
+		NodeKernels.Remove(nullptr);
 
 		if (NodeKernels.IsEmpty())
 		{
-			ensureMsgf(false, TEXT("PCG GPU compiler: Settings '%s' did not emit any kernels, check implementation of UPCGSettings::CreateKernels()."), *Settings->GetName());
 			continue;
 		}
 
+#ifdef PCG_GPU_KERNEL_PROFILING
+		if (NodeKernels.IsValidIndex(Settings->ProfileKernelIndex))
+		{
+			NodeKernels[Settings->ProfileKernelIndex]->EnableRepeatDispatch();
+		}
+#endif // PCG_GPU_KERNEL_PROFILING
+
 		// Wire up overridden parameters for any kernels that require them.
 		for (UPCGComputeKernel* NodeKernel : NodeKernels)
 		{
@@ -671,6 +727,7 @@ void FPCGGraphCompilerGPU::SetupVirtualPins(
 	const TSet<FPCGTaskId>& InCollapsedTasks,
 	const TArray<FPCGGraphTask>& InCompiledTasks,
 	const FTaskToSuccessors& InTaskSuccessors,
+	const TMap<FPCGTaskId, UPCGComputeKernel*>& InTaskIdToKernel,
 	FOriginalToVirtualPin& OutOriginalToVirtualPin)
 {
 	TRACE_CPUPROFILER_EVENT_SCOPE(FPCGGraphCompilerGPU::SetupVirtualPins);
@@ -701,10 +758,12 @@ void FPCGGraphCompilerGPU::SetupVirtualPins(
 
 		// Create virtual pin labels for all output pins of GPU nodes. We could create them only for output pins that internal->external edges, however
 		// it is helpful for special cases like inspection-only to always have virtual labels for all external output pins.
-		// TODO update to use kernels once the GPU compiler flow is refactored.
-		if (const UPCGSettings* Settings = InCompiledTasks[GPUTaskId].Node ? InCompiledTasks[GPUTaskId].Node->GetSettings() : nullptr)
+		if (const UPCGComputeKernel* const* FoundKernel = InTaskIdToKernel.Find(GPUTaskId); ensure(FoundKernel && *FoundKernel))
 		{
-			for (const FPCGPinProperties& PinProps : Settings->OutputPinProperties())
+			TArray<FPCGPinPropertiesGPU> OutputPins;
+			(*FoundKernel)->GetOutputPins(OutputPins);
+
+			for (const FPCGPinPropertiesGPU& PinProps : OutputPins)
 			{
 				const FNodePin PinKey = { GPUTaskId, PinProps.Label, /*Pin is input*/false };
 				if (!OutOriginalToVirtualPin.Contains(PinKey))
@@ -792,7 +851,7 @@ void FPCGGraphCompilerGPU::WireComputeGraphTask(
 					// Get the compute graph virtual pin corresponding to this pin, if one was created.
 					const FName* VirtualPinLabel = InOriginalToVirtualPin.Find({ GPUTaskId, InputCopy.UpstreamPin->Label, /*Pin is input*/false });
 
-					if (ensure(VirtualPinLabel))
+					if (ensure(VirtualPinLabel)) // Ensure should never fire - compute graph outputs are currently created for all kernel output pins.
 					{
 						// Wire to the existing virtual output pin.
 						InputCopy.UpstreamPin->Label = *VirtualPinLabel;
@@ -1824,7 +1883,7 @@ void FPCGGraphCompilerGPU::CreateGPUNodes(FPCGGraphCompiler& InOutCompiler, UPCG
 	// Setup mappings from existing pins to compute graph element virtual pins as a prestep before wiring in the compute graph tasks.
 	for (uint32 ComputeGraphIndex = 0; ComputeGraphIndex < NumComputeGraphs; ++ComputeGraphIndex)
 	{
-		SetupVirtualPins(NodeSubsetsToConvertToCFGraph[ComputeGraphIndex], InOutCompiledTasks, TaskSuccessors, OriginalToVirtualPin[ComputeGraphIndex]);
+		SetupVirtualPins(NodeSubsetsToConvertToCFGraph[ComputeGraphIndex], InOutCompiledTasks, TaskSuccessors, TaskIdToKernel, OriginalToVirtualPin[ComputeGraphIndex]);
 	}
 
 	TArray<FPCGTaskId> ComputeElementTasks;
diff --git a/Engine/Plugins/PCG/Source/PCG/Private/Graph/PCGGraphCompilerGPU.h b/Engine/Plugins/PCG/Source/PCG/Private/Graph/PCGGraphCompilerGPU.h
--- a/Engine/Plugins/PCG/Source/PCG/Private/Graph/PCGGraphCompilerGPU.h
+++ b/Engine/Plugins/PCG/Source/PCG/Private/Graph/PCGGraphCompilerGPU.h
@@ -60,6 +60,7 @@ private:
 		const TSet<FPCGTaskId>& InCollapsedTasks,
 		const TArray<FPCGGraphTask>& InCompiledTasks,
 		const TMap<FPCGTaskId, TArray<FPCGTaskId>>& InTaskSuccessors,
+		const TMap<FPCGTaskId, UPCGComputeKernel*>& InTaskIdToKernel,
 		FOriginalToVirtualPin& OutOriginalToVirtualPin);
 
 	/** Wires in a compute graph task alongside each set of GPU compatible nodes. The tasks for each node will be culled later. */
