[PCG] Avoid calling ULevelInstanceSubsystem::GetParentLevelInstance() with null owner

#Synced-CL 51541007

diff --git a/Engine/Plugins/PCG/Source/PCGEditor/Private/ChangeTracking/PCGComponentChangeHandler.cpp b/Engine/Plugins/PCG/Source/PCGEditor/Private/ChangeTracking/PCGComponentChangeHandler.cpp
--- a/Engine/Plugins/PCG/Source/PCGEditor/Private/ChangeTracking/PCGComponentChangeHandler.cpp
+++ b/Engine/Plugins/PCG/Source/PCGEditor/Private/ChangeTracking/PCGComponentChangeHandler.cpp
@@ -274,7 +274,8 @@ void FPCGComponentChangeHandler::EndChangeHandling(bool bSkipRefresh)
 			continue;
 		}
 
-		const bool bOwnerHasChanged = Component->GetOwner() == Change->ChangedObject;
+		const AActor* ComponentOwner = Component->GetOwner();
+		const bool bOwnerHasChanged = ComponentOwner == Change->ChangedObject;
 
 		if ((!Change->bNoRefreshOwner || !bOwnerHasChanged) && (!Component->bOnlyTrackItself || bOwnerHasChanged))
 		{
@@ -290,7 +291,7 @@ void FPCGComponentChangeHandler::EndChangeHandling(bool bSkipRefresh)
 				const ILevelInstanceInterface* ActorLevelInstance = LevelInstanceSubsystem->GetParentLevelInstance(Actor);
 
 				// Immediate Level Instance Parent (PCG Component)
-				const ILevelInstanceInterface* ComponentLevelInstance = LevelInstanceSubsystem->GetParentLevelInstance(Component->GetOwner());
+				const ILevelInstanceInterface* ComponentLevelInstance = ComponentOwner ? LevelInstanceSubsystem->GetParentLevelInstance(ComponentOwner) : nullptr;
 
 				// Is the Actor under a Level Instance that is the same or under the Components owning Level Instance
 				bool bActorLevelInstanceInComponentLevelInstanceHierarchy = false;
