[PCG] Add metadata memory counting

#Synced-CL 51473272

diff --git a/Engine/Plugins/PCG/Source/PCG/Private/Data/PCGSpatialData.cpp b/Engine/Plugins/PCG/Source/PCG/Private/Data/PCGSpatialData.cpp
--- a/Engine/Plugins/PCG/Source/PCG/Private/Data/PCGSpatialData.cpp
+++ b/Engine/Plugins/PCG/Source/PCG/Private/Data/PCGSpatialData.cpp
@@ -216,16 +216,6 @@ const UPCGPointArrayData* UPCGSpatialDataWithPointCache::CreatePointArrayData(FP
 	return CreatePointData(Context, InBounds)->ToPointArrayData(Context);
 }
 
-void UPCGSpatialData::GetResourceSizeEx(FResourceSizeEx& CumulativeResourceSize)
-{
-	Super::GetResourceSizeEx(CumulativeResourceSize);
-
-	if (Metadata)
-	{
-		Metadata->GetResourceSizeEx(CumulativeResourceSize);
-	}
-}
-
 float UPCGSpatialData::GetDensityAtPosition(const FVector& InPosition) const
 {
 	FPCGPoint TemporaryPoint;
diff --git a/Engine/Plugins/PCG/Source/PCG/Private/Metadata/PCGMetadata.cpp b/Engine/Plugins/PCG/Source/PCG/Private/Metadata/PCGMetadata.cpp
--- a/Engine/Plugins/PCG/Source/PCG/Private/Metadata/PCGMetadata.cpp
+++ b/Engine/Plugins/PCG/Source/PCG/Private/Metadata/PCGMetadata.cpp
@@ -406,6 +406,21 @@ void UPCGMetadata::Serialize(FArchive& InArchive)
 	}
 }
 
+void UPCGMetadata::GetResourceSizeEx(FResourceSizeEx& CumulativeResourceSize)
+{
+	TRACE_CPUPROFILER_EVENT_SCOPE(UPCGMetadata::GetResourceSizeEx);
+	Super::GetResourceSizeEx(CumulativeResourceSize);
+
+	// Note: we will not gather memory from the parent, but only "local" memory here
+	// This might report as not quite all the memory we need (in the context of the cache, for example)
+	// but it will not also grossly over-evalute the memory needed.
+	// Gather from domains
+	ForEachValidUniqueConstMetadataDomain([&CumulativeResourceSize](FPCGMetadataDomainID DomainID, const FPCGMetadataDomain& Domain)
+	{
+		Domain.GetResourceSizeEx(CumulativeResourceSize);
+	});
+}
+
 void UPCGMetadata::K2_InitializeAsCopy(const UPCGMetadata* InMetadataToCopy, const TArray<int64>& InOptionalEntriesToCopy)
 {
 	if (!IsValid(InMetadataToCopy))
diff --git a/Engine/Plugins/PCG/Source/PCG/Private/Metadata/PCGMetadataAttribute.cpp b/Engine/Plugins/PCG/Source/PCG/Private/Metadata/PCGMetadataAttribute.cpp
--- a/Engine/Plugins/PCG/Source/PCG/Private/Metadata/PCGMetadataAttribute.cpp
+++ b/Engine/Plugins/PCG/Source/PCG/Private/Metadata/PCGMetadataAttribute.cpp
@@ -82,6 +82,12 @@ void FPCGMetadataAttributeBase::Serialize(FPCGMetadataDomain* InMetadata, FArchi
 	InArchive << AttributeId;
 }
 
+void FPCGMetadataAttributeBase::GetResourceSizeEx(FResourceSizeEx& CumulativeResourceSize) const
+{
+	PCG::TSharedScopeLock ReadLock(EntryMapLock);
+	CumulativeResourceSize.AddDedicatedSystemMemoryBytes(EntryToValueKeyMap.GetAllocatedSize());
+}
+
 void FPCGMetadataAttributeBase::SetValueFromValueKey(PCGMetadataEntryKey EntryKey, PCGMetadataValueKey ValueKey, bool bResetValueOnDefaultValueKey)
 {
 	PCG::TUniqueScopeLock WriteLock(EntryMapLock);
diff --git a/Engine/Plugins/PCG/Source/PCG/Private/Metadata/PCGMetadataAttributeGeneric.cpp b/Engine/Plugins/PCG/Source/PCG/Private/Metadata/PCGMetadataAttributeGeneric.cpp
--- a/Engine/Plugins/PCG/Source/PCG/Private/Metadata/PCGMetadataAttributeGeneric.cpp
+++ b/Engine/Plugins/PCG/Source/PCG/Private/Metadata/PCGMetadataAttributeGeneric.cpp
@@ -6,6 +6,7 @@
 #include "Metadata/PCGMetadataDomain.h"
 #include "Metadata/Accessors/PCGAttributeAccessorHelpers.h"
 
+#include "Serialization/ArchiveCountMem.h"
 #include "UObject/TextProperty.h"
 
 namespace PCG::Private
@@ -519,6 +520,31 @@ FPCGMetadataAttributeGeneric::~FPCGMetadataAttributeGeneric()
 	}
 }
 
+void FPCGMetadataAttributeGeneric::GetResourceSizeEx(FResourceSizeEx& CumulativeResourceSize) const
+{
+	FPCGMetadataAttributeBase::GetResourceSizeEx(CumulativeResourceSize);
+
+	PCG::TSharedScopeLock ReadLock(ValueLock);
+	// Note - can replace this with serialize call once it's there.
+	FArchiveCountMem Archive(nullptr);
+	// Serialize desc
+	FPCGMetadataAttributeDesc::StaticStruct()->SerializeItem(FStructuredArchiveFromArchive(Archive).GetSlot(), const_cast<FPCGMetadataAttributeDesc*>(&CachedDesc), /*Defaults*/nullptr);
+	
+	if (UnderlyingProperty && UnderlyingProperty->Property)
+	{
+		if (UnderlyingProperty->Property->Inner && DefaultValue.GetRawPtr())
+		{
+			// Serialize default value
+			UnderlyingProperty->Property->Inner->SerializeItem(FStructuredArchiveFromArchive(Archive).GetSlot(), const_cast<void*>(DefaultValue.GetRawPtr()), /*Defaults*/nullptr);
+		}
+
+		// Serialize values
+		UnderlyingProperty->Property->SerializeItem(FStructuredArchiveFromArchive(Archive).GetSlot(), const_cast<FScriptArray*>(&Values), /*Defaults*/nullptr);
+	}
+
+	CumulativeResourceSize.AddDedicatedSystemMemoryBytes(Archive.GetNum());
+}
+
 PCGMetadataValueKey FPCGMetadataAttributeGeneric::GetValueKeyOffsetForChild() const
 {
 	PCG::TSharedScopeLock ReadLock(ValueLock);
diff --git a/Engine/Plugins/PCG/Source/PCG/Private/Metadata/PCGMetadataDomain.cpp b/Engine/Plugins/PCG/Source/PCG/Private/Metadata/PCGMetadataDomain.cpp
--- a/Engine/Plugins/PCG/Source/PCG/Private/Metadata/PCGMetadataDomain.cpp
+++ b/Engine/Plugins/PCG/Source/PCG/Private/Metadata/PCGMetadataDomain.cpp
@@ -205,6 +205,27 @@ void FPCGMetadataDomain::Serialize(FArchive& InArchive)
 	}
 }
 
+void FPCGMetadataDomain::GetResourceSizeEx(FResourceSizeEx& CumulativeResourceSize) const
+{
+	// Iterate over all attributes
+	{
+		PCG::TSharedScopeLock ReadLock(AttributeLock);
+		for (const auto& [AttributeName, AttributePtr] : Attributes)
+		{
+			if (AttributePtr)
+			{
+				AttributePtr->GetResourceSizeEx(CumulativeResourceSize);
+			}
+		}
+	}
+
+	// Count entry keys/parent keys
+	{
+		PCG::TSharedScopeLock ReadLock(ItemLock);
+		CumulativeResourceSize.AddDedicatedSystemMemoryBytes(ParentKeys.GetAllocatedSize());
+	}
+}
+
 FPCGMetadataAttributeGeneric* FPCGMetadataDomain::CreateGenericAttribute(const FPCGMetadataAttributeDesc& Desc, bool bAllowsInterpolation, bool bOverrideParent)
 {
 	const FName AttributeName = Desc.Name;
diff --git a/Engine/Plugins/PCG/Source/PCG/Private/PCGData.cpp b/Engine/Plugins/PCG/Source/PCG/Private/PCGData.cpp
--- a/Engine/Plugins/PCG/Source/PCG/Private/PCGData.cpp
+++ b/Engine/Plugins/PCG/Source/PCG/Private/PCGData.cpp
@@ -41,6 +41,16 @@ UPCGData::UPCGData(const FObjectInitializer& ObjectInitializer)
 	}
 }
 
+void UPCGData::GetResourceSizeEx(FResourceSizeEx& CumulativeResourceSize)
+{
+	Super::GetResourceSizeEx(CumulativeResourceSize);
+	
+	if (Metadata)
+	{
+		Metadata->GetResourceSizeEx(CumulativeResourceSize);
+	}
+}
+
 FPCGCrc UPCGData::GetOrComputeCrc(bool bFullDataCrc) const
 {
 #if WITH_EDITOR
diff --git a/Engine/Plugins/PCG/Source/PCG/Public/Data/PCGSpatialData.h b/Engine/Plugins/PCG/Source/PCG/Public/Data/PCGSpatialData.h
--- a/Engine/Plugins/PCG/Source/PCG/Public/Data/PCGSpatialData.h
+++ b/Engine/Plugins/PCG/Source/PCG/Public/Data/PCGSpatialData.h
@@ -112,10 +112,6 @@ class UPCGSpatialData : public UPCGData
 public:
 	UE_API UPCGSpatialData(const FObjectInitializer& ObjectInitializer);
 
-	//~Begin UObject Interface
-	UE_API virtual void GetResourceSizeEx(FResourceSizeEx& CumulativeResourceSize) override;
-	//~End UObject Interface
-
 	// ~Begin UPCGData interface
 	PCG_ASSIGN_DEFAULT_TYPE_INFO(FPCGDataTypeInfoSpatial)
 
diff --git a/Engine/Plugins/PCG/Source/PCG/Public/Metadata/PCGMetadata.h b/Engine/Plugins/PCG/Source/PCG/Public/Metadata/PCGMetadata.h
--- a/Engine/Plugins/PCG/Source/PCG/Public/Metadata/PCGMetadata.h
+++ b/Engine/Plugins/PCG/Source/PCG/Public/Metadata/PCGMetadata.h
@@ -62,6 +62,7 @@ public:
 	
 	//~ Begin UObject interface
 	UE_API virtual void Serialize(FArchive& InArchive) override;
+	UE_API virtual void GetResourceSizeEx(FResourceSizeEx& CumulativeResourceSize) override;
 	//~ End UObject interface
 
 	/** To be called by the outer pcg data to initialize the domains supported for this data. Can mark this one as default, will override previously set default. */
diff --git a/Engine/Plugins/PCG/Source/PCG/Public/Metadata/PCGMetadataAttribute.h b/Engine/Plugins/PCG/Source/PCG/Public/Metadata/PCGMetadataAttribute.h
--- a/Engine/Plugins/PCG/Source/PCG/Public/Metadata/PCGMetadataAttribute.h
+++ b/Engine/Plugins/PCG/Source/PCG/Public/Metadata/PCGMetadataAttribute.h
@@ -32,6 +32,7 @@ public:
 	virtual ~FPCGMetadataAttributeBase() = default;
 	
 	UE_API virtual void Serialize(FPCGMetadataDomain* InMetadata, FArchive& InArchive);
+	UE_API virtual void GetResourceSizeEx(FResourceSizeEx& CumulativeResourceSize) const;
 
 	UE_DEPRECATED(5.6, "Use the version with FPCGMetadataDomain")
 	UE_API virtual void Serialize(UPCGMetadata* InMetadata, FArchive& InArchive);
diff --git a/Engine/Plugins/PCG/Source/PCG/Public/Metadata/PCGMetadataAttributeGeneric.h b/Engine/Plugins/PCG/Source/PCG/Public/Metadata/PCGMetadataAttributeGeneric.h
--- a/Engine/Plugins/PCG/Source/PCG/Public/Metadata/PCGMetadataAttributeGeneric.h
+++ b/Engine/Plugins/PCG/Source/PCG/Public/Metadata/PCGMetadataAttributeGeneric.h
@@ -93,6 +93,7 @@ class FPCGMetadataAttributeGeneric : public FPCGMetadataAttributeBase
 public:
 	PCG_API ~FPCGMetadataAttributeGeneric();
 
+	PCG_API virtual void GetResourceSizeEx(FResourceSizeEx& CumulativeResourceSize) const override;
 	PCG_API void AddReferencedObjects(FReferenceCollector& Collector);
 
 	virtual bool IsGeneric() const override { return true; }
diff --git a/Engine/Plugins/PCG/Source/PCG/Public/Metadata/PCGMetadataAttributeTpl.h b/Engine/Plugins/PCG/Source/PCG/Public/Metadata/PCGMetadataAttributeTpl.h
--- a/Engine/Plugins/PCG/Source/PCG/Public/Metadata/PCGMetadataAttributeTpl.h
+++ b/Engine/Plugins/PCG/Source/PCG/Public/Metadata/PCGMetadataAttributeTpl.h
@@ -11,6 +11,8 @@
 #include "Metadata/PCGMetadataAttributeTraits.h"
 #include "Misc/ScopeRWLock.h"
 
+#include "Serialization/ArchiveCountMem.h"
+
 class FPCGMetadataDomain;
 class UPCGMetadata;
 
@@ -62,6 +64,15 @@ public:
 		}
 	}
 
+	virtual void GetResourceSizeEx(FResourceSizeEx& CumulativeResourceSize) const override
+	{
+		PCG::TUniqueScopeLock ReadLock(ValueLock);
+		FArchiveCountMem Archive(nullptr);
+		const_cast<FPCGMetadataAttribute<T>*>(this)->Serialize(Metadata, Archive);
+
+		CumulativeResourceSize.AddDedicatedSystemMemoryBytes(Archive.GetNum());
+	}
+
 	virtual void Flatten() override
 	{
 		// Implementation notes:
diff --git a/Engine/Plugins/PCG/Source/PCG/Public/Metadata/PCGMetadataDomain.h b/Engine/Plugins/PCG/Source/PCG/Public/Metadata/PCGMetadataDomain.h
--- a/Engine/Plugins/PCG/Source/PCG/Public/Metadata/PCGMetadataDomain.h
+++ b/Engine/Plugins/PCG/Source/PCG/Public/Metadata/PCGMetadataDomain.h
@@ -49,6 +49,7 @@ public:
 	UE_API virtual ~FPCGMetadataDomain();
 	
 	UE_API virtual void Serialize(FArchive& InArchive);
+	UE_API virtual void GetResourceSizeEx(FResourceSizeEx& CumulativeResourceSize) const;
 	
 	/** Initializes the sub metadata from a parent sub metadata, if any (can be null). Copies attributes and values. */
 	UE_API void Initialize(const FPCGMetadataDomain* InParent);
diff --git a/Engine/Plugins/PCG/Source/PCG/Public/PCGData.h b/Engine/Plugins/PCG/Source/PCG/Public/PCGData.h
--- a/Engine/Plugins/PCG/Source/PCG/Public/PCGData.h
+++ b/Engine/Plugins/PCG/Source/PCG/Public/PCGData.h
@@ -55,6 +55,10 @@ class UPCGData : public UObject
 public:
 	UE_API UPCGData(const FObjectInitializer& ObjectInitializer);
 
+	//~Begin UObject Interface
+	UE_API virtual void GetResourceSizeEx(FResourceSizeEx& CumulativeResourceSize) override;
+	//~End UObject Interface
+
 	// @todo_pcg To be deprecated when we switch to the new type
 	//UE_DEPRECATED(5.7, "Use the new GetDataTypeId")
 	virtual EPCGDataType GetDataType() const { return EPCGDataType::Any; }
