[PCG] Add spline & polygon 2d memory size reporting.

#Synced-CL 51473272

diff --git a/Engine/Plugins/PCG/Source/PCG/Private/Data/PCGPolygon2DData.cpp b/Engine/Plugins/PCG/Source/PCG/Private/Data/PCGPolygon2DData.cpp
--- a/Engine/Plugins/PCG/Source/PCG/Private/Data/PCGPolygon2DData.cpp
+++ b/Engine/Plugins/PCG/Source/PCG/Private/Data/PCGPolygon2DData.cpp
@@ -278,6 +278,20 @@ void UPCGPolygon2DData::Serialize(FArchive& InArchive)
 	}
 }
 
+void UPCGPolygon2DData::GetResourceSizeEx(FResourceSizeEx& CumulativeResourceSize)
+{
+	Super::GetResourceSizeEx(CumulativeResourceSize);
+
+	CumulativeResourceSize.AddDedicatedSystemMemoryBytes(
+		sizeof(UE::Geometry::FGeneralPolygon2d) +
+		Polygon.GetAllocatedSize() +
+		sizeof(Transform) +
+		VertexEntryKeys.GetAllocatedSize() +
+		SegmentIndexToSegmentAndHoleIndices.GetAllocatedSize() +
+		SegmentAndHoleIndicesToSegmentIndex.GetAllocatedSize() +
+		sizeof(int));
+}
+
 void UPCGPolygon2DData::AddToCrc(FArchiveCrc32& Ar, bool bFullDataCrc) const
 {
 	Super::AddToCrc(Ar, bFullDataCrc);
diff --git a/Engine/Plugins/PCG/Source/PCG/Private/Data/PCGSplineData.cpp b/Engine/Plugins/PCG/Source/PCG/Private/Data/PCGSplineData.cpp
--- a/Engine/Plugins/PCG/Source/PCG/Private/Data/PCGSplineData.cpp
+++ b/Engine/Plugins/PCG/Source/PCG/Private/Data/PCGSplineData.cpp
@@ -15,6 +15,7 @@
 
 #include "Components/SplineComponent.h"
 #include "Serialization/ArchiveCrc32.h"
+#include "Serialization/ArchiveCountMem.h"
 
 #include UE_INLINE_GENERATED_CPP_BY_NAME(PCGSplineData)
 
@@ -46,6 +47,15 @@ UPCGSplineData::UPCGSplineData()
 	Metadata->SetupDomain(PCGMetadataDomainID::Elements, /*bIsDefault=*/true);
 }
 
+void UPCGSplineData::GetResourceSizeEx(FResourceSizeEx& CumulativeResourceSize)
+{
+	Super::GetResourceSizeEx(CumulativeResourceSize);
+
+	FArchiveCountMem Archive(this);
+	FPCGSplineStruct::StaticStruct()->SerializeItem(FStructuredArchiveFromArchive(Archive).GetSlot(), &SplineStruct, /*Defaults*/nullptr);
+	CumulativeResourceSize.AddDedicatedSystemMemoryBytes(Archive.GetNum() + sizeof(FBox));
+}
+
 void UPCGSplineData::Initialize(const USplineComponent* InSpline)
 {
 	check(InSpline);
diff --git a/Engine/Plugins/PCG/Source/PCG/Public/Data/PCGPolygon2DData.h b/Engine/Plugins/PCG/Source/PCG/Public/Data/PCGPolygon2DData.h
--- a/Engine/Plugins/PCG/Source/PCG/Public/Data/PCGPolygon2DData.h
+++ b/Engine/Plugins/PCG/Source/PCG/Public/Data/PCGPolygon2DData.h
@@ -76,6 +76,7 @@ public:
 
 	// ~Begin UObject interface
 	PCG_API virtual void Serialize(FArchive& InArchive) override;
+	PCG_API virtual void GetResourceSizeEx(FResourceSizeEx& CumulativeResourceSize) override;
 	// ~End UObject interface
 
 	// ~Begin UPCGData interface
diff --git a/Engine/Plugins/PCG/Source/PCG/Public/Data/PCGSplineData.h b/Engine/Plugins/PCG/Source/PCG/Public/Data/PCGSplineData.h
--- a/Engine/Plugins/PCG/Source/PCG/Public/Data/PCGSplineData.h
+++ b/Engine/Plugins/PCG/Source/PCG/Public/Data/PCGSplineData.h
@@ -40,6 +40,8 @@ class UPCGSplineData : public UPCGPolyLineData
 public:
 	UE_API UPCGSplineData();
 
+	UE_API virtual void GetResourceSizeEx(FResourceSizeEx& CumulativeResourceSize) override;
+
 	/**
 	 * Initialize a spline data from a Spline Component.
 	 */
