[PCG] First GPU graph unit tests.

Common helper code to execute a PCG GPU graph, and ability to specify which nodes to collect output data from which can then be validated. This is a first pass and we can extend and refactor as needed.

[RN] [PCG] Added unit test support code for PCG GPU graphs.
#Synced-CL 51538157

diff --git a/Engine/Plugins/PCG/Source/PCG/Private/Compute/Elements/PCGComputeGraphElement.cpp b/Engine/Plugins/PCG/Source/PCG/Private/Compute/Elements/PCGComputeGraphElement.cpp
--- a/Engine/Plugins/PCG/Source/PCG/Private/Compute/Elements/PCGComputeGraphElement.cpp
+++ b/Engine/Plugins/PCG/Source/PCG/Private/Compute/Elements/PCGComputeGraphElement.cpp
@@ -182,28 +182,20 @@ bool FPCGComputeGraphElement::ExecuteInternal(FPCGContext* InContext) const
 			Context->GenerationGridSize = Source->GetExecutionState().GetGenerationGridSize();
 		}
 		
-		if (IPCGBaseSubsystem* Subsystem = Source->GetExecutionState().GetSubsystem())
-		{
-			Context->ComputeGraph = Subsystem->GetComputeGraph(DynamicSubgraph ? DynamicSubgraph : TopGraph, Context->GenerationGridSize, ComputeGraphIndex);
+		Context->ComputeGraph = Context->GetComputeGraph(DynamicSubgraph ? DynamicSubgraph : TopGraph, Context->GenerationGridSize, ComputeGraphIndex);
 
-			if (!Context->ComputeGraph)
-			{
+		if (!Context->ComputeGraph)
+		{
 #if WITH_EDITOR
-				// Normally emitted when a graph compilation is flushed while generation in progress.
-				UE_LOG(LogPCG, Warning, TEXT("Failed to obtain compute graph for %s '%s', grid size %u, compute graph index %d. Expected if graph compilation cache flushed mid-execution. Compute graph execution aborted."),
-					DynamicSubgraph ? TEXT("dynamic subgraph") : TEXT("top graph"), DynamicSubgraph ? *DynamicSubgraph->GetName() : *TopGraph->GetName(), Context->GenerationGridSize, ComputeGraphIndex);
+			// Normally emitted when a graph compilation is flushed while generation in progress.
+			UE_LOG(LogPCG, Warning, TEXT("Failed to obtain compute graph for %s '%s', grid size %u, compute graph index %d. Expected if graph compilation cache flushed mid-execution. Compute graph execution aborted."),
+				DynamicSubgraph ? TEXT("dynamic subgraph") : TEXT("top graph"), DynamicSubgraph ? *DynamicSubgraph->GetName() : *TopGraph->GetName(), Context->GenerationGridSize, ComputeGraphIndex);
 #else
-				// Flush unlikely in standalone, so may be that cook failed to produce the compute graph.
-				UE_LOG(LogPCG, Error, TEXT("Failed to obtain compute graph for %s '%s', grid size %u, compute graph index %d. Either the graph compilation cache was flushed mid-execution, or the cook may have failed to produce the required compute graph. Compute graph execution aborted."),
-					DynamicSubgraph ? TEXT("dynamic subgraph") : TEXT("top graph"), DynamicSubgraph ? *DynamicSubgraph->GetName() : *TopGraph->GetName(), Context->GenerationGridSize, ComputeGraphIndex);
+			// Flush unlikely in standalone, so may be that cook failed to produce the compute graph.
+			UE_LOG(LogPCG, Error, TEXT("Failed to obtain compute graph for %s '%s', grid size %u, compute graph index %d. Either the graph compilation cache was flushed mid-execution, or the cook may have failed to produce the required compute graph. Compute graph execution aborted."),
+				DynamicSubgraph ? TEXT("dynamic subgraph") : TEXT("top graph"), DynamicSubgraph ? *DynamicSubgraph->GetName() : *TopGraph->GetName(), Context->GenerationGridSize, ComputeGraphIndex);
 #endif
 
-				return true;
-			}
-		}
-		else
-		{
-			UE_LOG(LogPCG, Error, TEXT("Failed to obtain compute graph, no subsystem present. Compute graph execution aborted."));
 			return true;
 		}
 
diff --git a/Engine/Plugins/PCG/Source/PCG/Private/Compute/Elements/PCGCustomHLSL.cpp b/Engine/Plugins/PCG/Source/PCG/Private/Compute/Elements/PCGCustomHLSL.cpp
--- a/Engine/Plugins/PCG/Source/PCG/Private/Compute/Elements/PCGCustomHLSL.cpp
+++ b/Engine/Plugins/PCG/Source/PCG/Private/Compute/Elements/PCGCustomHLSL.cpp
@@ -351,6 +351,20 @@ bool UPCGCustomHLSLSettings::IsReadOnly() const
 	return KernelSourceOverride != nullptr;
 }
 
+void UPCGCustomHLSLSettings::SetKernelType(EPCGKernelType InKernelType)
+{
+	KernelType = InKernelType;
+
+	// Generators don't use the default input pin; clear so UpdatePinSettings recreates it with the correct type.
+	if (IsGeneratorKernel())
+	{
+		InputPins.Empty();
+	}
+
+	UpdatePinSettings();
+	UpdateDeclarations();
+}
+
 void UPCGCustomHLSLSettings::CreateKernels(FPCGGPUCompilationContext& InOutContext, UObject* InObjectOuter, TArray<UPCGComputeKernel*>& OutKernels, TArray<FPCGKernelEdge>& OutEdges) const
 {
 	PCGKernelHelpers::FCreateKernelParams CreateParams(InObjectOuter, this);
diff --git a/Engine/Plugins/PCG/Source/PCG/Private/Compute/Elements/PCGCustomHLSL.h b/Engine/Plugins/PCG/Source/PCG/Private/Compute/Elements/PCGCustomHLSL.h
--- a/Engine/Plugins/PCG/Source/PCG/Private/Compute/Elements/PCGCustomHLSL.h
+++ b/Engine/Plugins/PCG/Source/PCG/Private/Compute/Elements/PCGCustomHLSL.h
@@ -94,6 +94,9 @@ public:
 	void SetSourceText(const FString& InText) override;
 	bool IsReadOnly() const override;
 	//~End IPCGCodeEditorTextProvider interface
+
+	/** Sets the kernel type and updates pin settings and declarations. */
+	void SetKernelType(EPCGKernelType InKernelType);
 #endif
 
 #if WITH_EDITOR
diff --git a/Engine/Plugins/PCG/Source/PCG/Private/PCGContext.cpp b/Engine/Plugins/PCG/Source/PCG/Private/PCGContext.cpp
--- a/Engine/Plugins/PCG/Source/PCG/Private/PCGContext.cpp
+++ b/Engine/Plugins/PCG/Source/PCG/Private/PCGContext.cpp
@@ -730,6 +730,16 @@ void FPCGContext::ClearOutputData(FPCGTaskId InTaskId)
 	}
 }
 
+UPCGComputeGraph* FPCGContext::GetComputeGraph(const UPCGGraph* InGraph, uint32 GridSize, uint32 ComputeGraphIndex) const
+{
+	if (TSharedPtr<FPCGGraphExecutor> Executor = GraphExecutor.Pin())
+	{
+		return Executor->GetCompiler()->GetComputeGraph(InGraph, GridSize, ComputeGraphIndex);
+	}
+
+	return nullptr;
+}
+
 void FPCGContext::InitializeGraphExecutor(FPCGContext* InContext)
 {
 	check(InContext);
diff --git a/Engine/Plugins/PCG/Source/PCG/Private/Subsystems/IPCGBaseSubsystem.cpp b/Engine/Plugins/PCG/Source/PCG/Private/Subsystems/IPCGBaseSubsystem.cpp
--- a/Engine/Plugins/PCG/Source/PCG/Private/Subsystems/IPCGBaseSubsystem.cpp
+++ b/Engine/Plugins/PCG/Source/PCG/Private/Subsystems/IPCGBaseSubsystem.cpp
@@ -176,6 +176,7 @@ FPCGGraphCompiler* IPCGBaseSubsystem::GetGraphCompiler()
 	return GraphExecutor ? GraphExecutor->GetCompiler() : nullptr;
 }
 
+// deprecated (5.8)
 UPCGComputeGraph* IPCGBaseSubsystem::GetComputeGraph(const UPCGGraph* InGraph, uint32 GridSize, uint32 ComputeGraphIndex)
 {
 	if (FPCGGraphCompiler* GraphCompiler = GetGraphCompiler())
diff --git a/Engine/Plugins/PCG/Source/PCG/Private/Tests/Compute/PCGCustomHLSLTest.cpp b/Engine/Plugins/PCG/Source/PCG/Private/Tests/Compute/PCGCustomHLSLTest.cpp
new file mode 100644
--- /dev/null
+++ b/Engine/Plugins/PCG/Source/PCG/Private/Tests/Compute/PCGCustomHLSLTest.cpp
@@ -0,0 +1,343 @@
+// Copyright Epic Games, Inc. All Rights Reserved.
+
+#if WITH_EDITOR
+
+#include "Tests/Compute/PCGGPUTestCommon.h"
+
+#include "PCGCommon.h"
+#include "PCGGraph.h"
+#include "PCGParamData.h"
+#include "Compute/Elements/PCGCustomHLSL.h"
+#include "Data/PCGBasePointData.h"
+#include "Elements/PCGCreateAttribute.h"
+#include "Elements/PCGCreatePoints.h"
+#include "Elements/PCGGather.h"
+#include "Metadata/PCGMetadataAttributeTpl.h"
+
+/**
+ * GPU test: CustomHLSL PointProcessor translates point positions.
+ *
+ *   CreatePoints (4 known positions) -> CustomHLSL (X += 10, Y += 25, Z += 50) -> Gather
+ *
+ * Verifies each output point's position == input position + offset, using a different
+ * offset per axis to confirm all three components are written correctly.
+ */
+IMPLEMENT_CUSTOM_SIMPLE_AUTOMATION_TEST(
+	FPCGCustomHLSLPointTranslateTest,
+	FPCGTestBaseClass,
+	"Plugins.PCG.Compute.CustomHLSL.PointTranslate",
+	PCGTestsCommon::TestFlags | EAutomationTestFlags::NonNullRHI)
+
+bool FPCGCustomHLSLPointTranslateTest::RunTest(const FString& Parameters)
+{
+	FPCGGPUGraphTestRunner Runner;
+
+	// CreatePoints: 4 points at known positions.
+	UPCGCreatePointsSettings* CreatePointsSettings = nullptr;
+	UPCGNode* CreatePointsNode = Runner.Graph->AddNodeOfType<UPCGCreatePointsSettings>(CreatePointsSettings);
+	CreatePointsSettings->PointsToCreate.Empty();
+
+	const TArray<FVector> InputPositions = {
+		FVector(  0.f,   0.f,   0.f),
+		FVector(100.f, 200.f, 300.f),
+		FVector(-50.f,  75.f, 250.f),
+		FVector( 10.f, -30.f, -80.f),
+	};
+	for (const FVector& Pos : InputPositions)
+	{
+		FPCGPoint& Pt = CreatePointsSettings->PointsToCreate.AddDefaulted_GetRef();
+		Pt.Transform.SetLocation(Pos);
+	}
+	CreatePointsSettings->CoordinateSpace = EPCGCoordinateSpace::World;
+
+	// CustomHLSL PointProcessor: translate by a different amount per axis.
+	UPCGCustomHLSLSettings* HLSLSettings = nullptr;
+	UPCGNode* HLSLNode = Runner.Graph->AddNodeOfType<UPCGCustomHLSLSettings>(HLSLSettings);
+	HLSLSettings->KernelType = EPCGKernelType::PointProcessor;
+	HLSLSettings->SetSourceText(
+		TEXT("float3 Pos = In_GetPosition(In_DataIndex, ElementIndex);\n")
+		TEXT("Pos.x += 10.0f;\n")
+		TEXT("Pos.y += 25.0f;\n")
+		TEXT("Pos.z += 50.0f;\n")
+		TEXT("Out_SetPosition(Out_DataIndex, ElementIndex, Pos);\n")
+	);
+
+	// Useful to add gather node to grab output data from.
+	UPCGGatherSettings* GatherSettings = nullptr;
+	UPCGNode* GatherNode = Runner.Graph->AddNodeOfType<UPCGGatherSettings>(GatherSettings);
+
+	Runner.Graph->AddEdge(CreatePointsNode, PCGPinConstants::DefaultOutputLabel, HLSLNode, PCGPinConstants::DefaultInputLabel);
+	Runner.Graph->AddEdge(HLSLNode, PCGPinConstants::DefaultOutputLabel, GatherNode, PCGPinConstants::DefaultInputLabel);
+
+	TArray<FPCGDataCollection> OutNodeData;
+	if (!Runner.Execute({ this, { GatherNode } }, OutNodeData))
+	{
+		return false;
+	}
+
+	return PCGGPUTestCommon::ValidateSingleTaggedData(this, OutNodeData[0], [&](const FPCGTaggedData& TaggedData)
+	{
+		const UPCGBasePointData* OutPoints = Cast<UPCGBasePointData>(TaggedData.Data);
+		if (!TestNotNull(TEXT("Readback data is point data"), OutPoints))
+		{
+			return false;
+		}
+
+		bool bTestPassed = TestEqual(TEXT("Output point count matches input"), OutPoints->GetNumPoints(), InputPositions.Num());
+
+		if (OutPoints->GetNumPoints() == InputPositions.Num())
+		{
+			const FConstPCGPointValueRanges Ranges(OutPoints);
+			for (int32 i = 0; i < InputPositions.Num(); ++i)
+			{
+				const FVector ExpectedPos = InputPositions[i] + FVector(10.0, 25.0, 50.0);
+				const FVector ActualPos = Ranges.TransformRange[i].GetLocation();
+				bTestPassed &= TestNearlyEqual(*FString::Printf(TEXT("Point[%d] X (expected %.1f, got %.1f)"), i, ExpectedPos.X, ActualPos.X), ActualPos.X, ExpectedPos.X, 0.1);
+				bTestPassed &= TestNearlyEqual(*FString::Printf(TEXT("Point[%d] Y (expected %.1f, got %.1f)"), i, ExpectedPos.Y, ActualPos.Y), ActualPos.Y, ExpectedPos.Y, 0.1);
+				bTestPassed &= TestNearlyEqual(*FString::Printf(TEXT("Point[%d] Z (expected %.1f, got %.1f)"), i, ExpectedPos.Z, ActualPos.Z), ActualPos.Z, ExpectedPos.Z, 0.1);
+			}
+		}
+
+		return bTestPassed;
+	});
+}
+
+/**
+ * GPU test: Two chained CustomHLSL PointProcessor nodes apply different transforms.
+ *
+ *   CreatePoints (4 points) -> CustomHLSL (set scale to (2,3,4)) -> CustomHLSL (rotate i*90 deg around Z) -> Gather
+ *
+ * Verifies the output scale (set by node 1 and carried through node 2) and per-point rotation
+ * (0, 90, 180, 270 deg around Z set by node 2), confirming each node in the chain sees the
+ * previous node's output data.
+ */
+IMPLEMENT_CUSTOM_SIMPLE_AUTOMATION_TEST(
+	FPCGCustomHLSLChainedPointProcessorsTest,
+	FPCGTestBaseClass,
+	"Plugins.PCG.Compute.CustomHLSL.ChainedPointProcessors",
+	PCGTestsCommon::TestFlags | EAutomationTestFlags::NonNullRHI)
+
+bool FPCGCustomHLSLChainedPointProcessorsTest::RunTest(const FString& Parameters)
+{
+	FPCGGPUGraphTestRunner Runner;
+
+	// CreatePoints: 4 default points.
+	UPCGCreatePointsSettings* CreatePointsSettings = nullptr;
+	UPCGNode* CreatePointsNode = Runner.Graph->AddNodeOfType<UPCGCreatePointsSettings>(CreatePointsSettings);
+	CreatePointsSettings->PointsToCreate.Empty();
+	for (int32 i = 0; i < 4; ++i)
+	{
+		CreatePointsSettings->PointsToCreate.AddDefaulted();
+	}
+	CreatePointsSettings->CoordinateSpace = EPCGCoordinateSpace::World;
+
+	// First HLSL node: set non-uniform scale.
+	UPCGCustomHLSLSettings* HLSLSettings1 = nullptr;
+	UPCGNode* HLSLNode1 = Runner.Graph->AddNodeOfType<UPCGCustomHLSLSettings>(HLSLSettings1);
+	HLSLSettings1->KernelType = EPCGKernelType::PointProcessor;
+	HLSLSettings1->SetSourceText(
+		TEXT("Out_SetScale(Out_DataIndex, ElementIndex, float3(2.0f, 3.0f, 4.0f));\n")
+	);
+
+	// Second HLSL node: rotate each point by a different angle around Z (0, 90, 180, 270 deg).
+	// Scale from node 1 should be carried through.
+	UPCGCustomHLSLSettings* HLSLSettings2 = nullptr;
+	UPCGNode* HLSLNode2 = Runner.Graph->AddNodeOfType<UPCGCustomHLSLSettings>(HLSLSettings2);
+	HLSLSettings2->KernelType = EPCGKernelType::PointProcessor;
+	HLSLSettings2->SetSourceText(
+		TEXT("FQuat Rot = Quat_FromAxisAngle(half3(0.0f, 0.0f, 1.0f), ElementIndex * 1.5707963f);\n")
+		TEXT("Out_SetRotation(Out_DataIndex, ElementIndex, Rot);\n")
+	);
+
+	// Useful to add gather node to grab output data from.
+	UPCGGatherSettings* GatherSettings = nullptr;
+	UPCGNode* GatherNode = Runner.Graph->AddNodeOfType<UPCGGatherSettings>(GatherSettings);
+
+	Runner.Graph->AddEdge(CreatePointsNode, PCGPinConstants::DefaultOutputLabel, HLSLNode1, PCGPinConstants::DefaultInputLabel);
+	Runner.Graph->AddEdge(HLSLNode1,        PCGPinConstants::DefaultOutputLabel, HLSLNode2, PCGPinConstants::DefaultInputLabel);
+	Runner.Graph->AddEdge(HLSLNode2,        PCGPinConstants::DefaultOutputLabel, GatherNode, PCGPinConstants::DefaultInputLabel);
+
+	TArray<FPCGDataCollection> OutNodeData;
+	if (!Runner.Execute({ this, { GatherNode } }, OutNodeData))
+	{
+		return false;
+	}
+
+	return PCGGPUTestCommon::ValidateSingleTaggedData(this, OutNodeData[0], [&](const FPCGTaggedData& TaggedData)
+	{
+		const UPCGBasePointData* OutPoints = Cast<UPCGBasePointData>(TaggedData.Data);
+		if (!TestNotNull(TEXT("Readback data is point data"), OutPoints))
+		{
+			return false;
+		}
+
+		constexpr int32 ExpectedPointCount = 4;
+		bool bTestPassed = TestEqual(TEXT("Output point count matches input"), OutPoints->GetNumPoints(), ExpectedPointCount);
+
+		if (OutPoints->GetNumPoints() == ExpectedPointCount)
+		{
+			const FVector ExpectedScale(2.f, 3.f, 4.f);
+
+			const FConstPCGPointValueRanges Ranges(OutPoints);
+			for (int32 i = 0; i < ExpectedPointCount; ++i)
+			{
+				const FVector ActualScale = Ranges.TransformRange[i].GetScale3D();
+				bTestPassed &= TestNearlyEqual(*FString::Printf(TEXT("Point[%d] Scale X (expected %.1f, got %.1f)"), i, ExpectedScale.X, ActualScale.X), ActualScale.X, ExpectedScale.X, 0.01);
+				bTestPassed &= TestNearlyEqual(*FString::Printf(TEXT("Point[%d] Scale Y (expected %.1f, got %.1f)"), i, ExpectedScale.Y, ActualScale.Y), ActualScale.Y, ExpectedScale.Y, 0.01);
+				bTestPassed &= TestNearlyEqual(*FString::Printf(TEXT("Point[%d] Scale Z (expected %.1f, got %.1f)"), i, ExpectedScale.Z, ActualScale.Z), ActualScale.Z, ExpectedScale.Z, 0.01);
+
+				const FQuat ExpectedRot(FVector::ZAxisVector, FMath::DegreesToRadians(90.f * i));
+				const FQuat ActualRot = Ranges.TransformRange[i].GetRotation();
+				bTestPassed &= TestTrue(*FString::Printf(TEXT("Point[%d] rotation matches expected %d deg around Z"), i, 90 * i), ActualRot.Equals(ExpectedRot, 0.01f));
+			}
+		}
+
+		return bTestPassed;
+	});
+}
+
+/**
+ * GPU test: CustomHLSL AttributeSetProcessor doubles a float attribute.
+ *
+ *   CreateAttributeSet (Value=5) -> CustomHLSL (Value *= 2) -> Gather
+ *
+ * Verifies the output attribute set has Value == 10.
+ */
+IMPLEMENT_CUSTOM_SIMPLE_AUTOMATION_TEST(
+	FPCGCustomHLSLAttributeSetDoubleValueTest,
+	FPCGTestBaseClass,
+	"Plugins.PCG.Compute.CustomHLSL.AttributeSetDoubleValue",
+	PCGTestsCommon::TestFlags | EAutomationTestFlags::NonNullRHI)
+
+bool FPCGCustomHLSLAttributeSetDoubleValueTest::RunTest(const FString& Parameters)
+{
+	FPCGGPUGraphTestRunner Runner;
+
+	// CreateAttributeSet: one element with float attribute "Value" = 5.
+	UPCGCreateAttributeSetSettings* CreateAttrSettings = nullptr;
+	UPCGNode* CreateAttrNode = Runner.Graph->AddNodeOfType<UPCGCreateAttributeSetSettings>(CreateAttrSettings);
+	CreateAttrSettings->AttributeTypes.Type = EPCGMetadataTypes::Float;
+	CreateAttrSettings->AttributeTypes.FloatValue = 5.0f;
+	CreateAttrSettings->OutputTarget.SetAttributeName(FName(TEXT("Value")));
+
+	// CustomHLSL AttributeSetProcessor: Value *= 2.
+	UPCGCustomHLSLSettings* HLSLSettings = nullptr;
+	UPCGNode* HLSLNode = Runner.Graph->AddNodeOfType<UPCGCustomHLSLSettings>(HLSLSettings);
+	HLSLSettings->SetKernelType(EPCGKernelType::AttributeSetProcessor);
+	HLSLSettings->SetSourceText(
+		TEXT("float v = In_GetFloat(In_DataIndex, ElementIndex, 'Value');\n")
+		TEXT("Out_SetFloat(Out_DataIndex, ElementIndex, 'Value', v * 2.0f);\n")
+	);
+
+	// Useful to add gather node to grab output data from.
+	UPCGGatherSettings* GatherSettings = nullptr;
+	UPCGNode* GatherNode = Runner.Graph->AddNodeOfType<UPCGGatherSettings>(GatherSettings);
+
+	Runner.Graph->AddEdge(CreateAttrNode, PCGPinConstants::DefaultOutputLabel, HLSLNode, PCGPinConstants::DefaultInputLabel);
+	Runner.Graph->AddEdge(HLSLNode, PCGPinConstants::DefaultOutputLabel, GatherNode, PCGPinConstants::DefaultInputLabel);
+
+	TArray<FPCGDataCollection> OutNodeData;
+	if (!Runner.Execute({ this, { GatherNode } }, OutNodeData))
+	{
+		return false;
+	}
+
+	return PCGGPUTestCommon::ValidateSingleTaggedData(this, OutNodeData[0], [&](const FPCGTaggedData& TaggedData)
+	{
+		const UPCGParamData* OutParams = Cast<UPCGParamData>(TaggedData.Data);
+		if (!TestNotNull(TEXT("Readback data is attribute set data"), OutParams))
+		{
+			return false;
+		}
+
+		const UPCGMetadata* Metadata = OutParams->ConstMetadata();
+		if (!TestNotNull(TEXT("Output has metadata"), Metadata))
+		{
+			return false;
+		}
+
+		const FPCGMetadataAttribute<float>* FloatAttr = Metadata->GetConstTypedAttribute<float>(FName(TEXT("Value")));
+		if (!TestNotNull(TEXT("Output has 'Value' float attribute"), FloatAttr))
+		{
+			return false;
+		}
+
+		const double ActualValue = FloatAttr->GetValueFromItemKey(0);
+		return TestNearlyEqual(*FString::Printf(TEXT("Value attribute doubled (expected 10.0, got %.1f)"), ActualValue), ActualValue, 10.0, 0.1);
+	});
+}
+
+/**
+ * GPU test: CustomHLSL AttributeSetProcessor passes a string attribute through unchanged.
+ *
+ *   CreateAttributeSet (StringAttr="Hello") -> CustomHLSL (read StringAttr, write StringAttr) -> Gather
+ *
+ * Verifies the output attribute set has StringAttr == "Hello".
+ */
+IMPLEMENT_CUSTOM_SIMPLE_AUTOMATION_TEST(
+	FPCGCustomHLSLAttributeSetStringAttributeTest,
+	FPCGTestBaseClass,
+	"Plugins.PCG.Compute.CustomHLSL.AttributeSetStringAttribute",
+	PCGTestsCommon::TestFlags | EAutomationTestFlags::NonNullRHI)
+
+bool FPCGCustomHLSLAttributeSetStringAttributeTest::RunTest(const FString& Parameters)
+{
+	FPCGGPUGraphTestRunner Runner;
+
+	// CreateAttributeSet: one element with string attribute "StringAttr" = "Hello".
+	UPCGCreateAttributeSetSettings* CreateAttrSettings = nullptr;
+	UPCGNode* CreateAttrNode = Runner.Graph->AddNodeOfType<UPCGCreateAttributeSetSettings>(CreateAttrSettings);
+	CreateAttrSettings->AttributeTypes.Type = EPCGMetadataTypes::String;
+	CreateAttrSettings->AttributeTypes.StringValue = TEXT("Hello");
+	CreateAttrSettings->OutputTarget.SetAttributeName(FName(TEXT("StringAttr")));
+
+	// CustomHLSL AttributeSetProcessor: read StringAttr and write it back unchanged.
+	UPCGCustomHLSLSettings* HLSLSettings = nullptr;
+	UPCGNode* HLSLNode = Runner.Graph->AddNodeOfType<UPCGCustomHLSLSettings>(HLSLSettings);
+	HLSLSettings->SetKernelType(EPCGKernelType::AttributeSetProcessor);
+	HLSLSettings->SetSourceText(
+		TEXT("int StringAttrVal = In_GetStringKey(In_DataIndex, ElementIndex, 'StringAttr');\n")
+		TEXT("Out_SetStringKey(Out_DataIndex, ElementIndex, 'StringAttr', StringAttrVal);\n")
+	);
+
+	// Useful to add gather node to grab output data from.
+	UPCGGatherSettings* GatherSettings = nullptr;
+	UPCGNode* GatherNode = Runner.Graph->AddNodeOfType<UPCGGatherSettings>(GatherSettings);
+
+	Runner.Graph->AddEdge(CreateAttrNode, PCGPinConstants::DefaultOutputLabel, HLSLNode, PCGPinConstants::DefaultInputLabel);
+	Runner.Graph->AddEdge(HLSLNode, PCGPinConstants::DefaultOutputLabel, GatherNode, PCGPinConstants::DefaultInputLabel);
+
+	TArray<FPCGDataCollection> OutNodeData;
+	if (!Runner.Execute({ this, { GatherNode } }, OutNodeData))
+	{
+		return false;
+	}
+
+	return PCGGPUTestCommon::ValidateSingleTaggedData(this, OutNodeData[0], [&](const FPCGTaggedData& TaggedData)
+	{
+		const UPCGParamData* OutParams = Cast<UPCGParamData>(TaggedData.Data);
+		if (!TestNotNull(TEXT("Readback data is attribute set data"), OutParams))
+		{
+			return false;
+		}
+
+		const UPCGMetadata* Metadata = OutParams->ConstMetadata();
+		if (!TestNotNull(TEXT("Output has metadata"), Metadata))
+		{
+			return false;
+		}
+
+		const FPCGMetadataAttribute<FString>* StringAttr = Metadata->GetConstTypedAttribute<FString>(FName(TEXT("StringAttr")));
+		if (!TestNotNull(TEXT("Output has 'StringAttr' string attribute"), StringAttr))
+		{
+			return false;
+		}
+
+		const FString ActualValue = StringAttr->GetValueFromItemKey(0);
+		return TestEqual(
+			*FString::Printf(TEXT("'StringAttr' attribute passed through (expected 'Hello', got '%s')"), *ActualValue),
+			ActualValue, FString(TEXT("Hello")));
+	});
+}
+
+#endif // WITH_EDITOR
diff --git a/Engine/Plugins/PCG/Source/PCG/Private/Tests/Compute/PCGGPUTestCommon.cpp b/Engine/Plugins/PCG/Source/PCG/Private/Tests/Compute/PCGGPUTestCommon.cpp
new file mode 100644
--- /dev/null
+++ b/Engine/Plugins/PCG/Source/PCG/Private/Tests/Compute/PCGGPUTestCommon.cpp
@@ -0,0 +1,185 @@
+// Copyright Epic Games, Inc. All Rights Reserved.
+
+#if WITH_EDITOR
+
+#include "Tests/Compute/PCGGPUTestCommon.h"
+
+#include "PCGComponent.h"
+#include "PCGGraph.h"
+#include "Compute/Data/PCGProxyForGPUData.h"
+#include "Graph/PCGGraphExecutor.h"
+
+#include "ComputeWorkerInterface.h"
+#include "RenderingThread.h"
+#include "Components/BoxComponent.h"
+#include "ComputeFramework/ComputeFramework.h"
+#include "Containers/Ticker.h"
+#include "Engine/World.h"
+#include "HAL/PlatformProcess.h"
+#include "Misc/AutomationTest.h"
+
+bool PCGGPUTestCommon::ValidateSingleTaggedData(FAutomationTestBase* Test, const FPCGDataCollection& DataCollection, TFunctionRef<bool(const FPCGTaggedData&)> Validator)
+{
+	if (!Test->TestEqual(TEXT("Data collection contains a single data item"), DataCollection.TaggedData.Num(), 1))
+	{
+		return false;
+	}
+	return Validator(DataCollection.TaggedData[0]);
+}
+
+FPCGGPUGraphTestRunner::FPCGGPUGraphTestRunner()
+{
+	// The test actor has no geometry, so GetActorBounds() returns an invalid FBox, which
+	// causes CreateGenerateTask to abort. Add a UBoxComponent to give the actor valid bounds.
+	UBoxComponent* BoundsBox = NewObject<UBoxComponent>(TestData.TestActor, FName(TEXT("BoundsBox")), RF_Transient);
+	BoundsBox->SetBoxExtent(FVector(100.f, 100.f, 100.f));
+	BoundsBox->RegisterComponent();
+	TestData.TestActor->AddInstanceComponent(BoundsBox);
+	if (!TestData.TestActor->GetRootComponent())
+	{
+		TestData.TestActor->SetRootComponent(BoundsBox);
+	}
+
+	Graph = NewObject<UPCGGraph>(GetTransientPackage(), NAME_None, RF_Transient);
+}
+
+bool FPCGGPUGraphTestRunner::Execute(const FPCGGPUGraphTestRunnerParams& Params, TArray<FPCGDataCollection>& OutNodeData)
+{
+	check(Params.Test);
+
+	TestData.TestPCGComponent->SetGraph(Graph);
+
+	TSharedPtr<FPCGGraphExecutor> Executor = MakeShared<FPCGGraphExecutor>();
+	Executor->SetKeepIntermediateResults(true);
+
+	TArray<FPCGDataCollection> CapturedOutputs;
+	CapturedOutputs.SetNum(Params.CollectOutputDataForNodes.Num());
+	TArray<bool> NodeOutputCaptured;
+	NodeOutputCaptured.Init(false, Params.CollectOutputDataForNodes.Num());
+
+	// This will execute the graph and run the lambda on each node, which we use to capture output data.
+	const FPCGTaskId FinalTaskId = Executor->ScheduleDebugWithTaskCallback(
+		TestData.TestPCGComponent,
+		[&](FPCGTaskId /*TaskId*/, const UPCGNode* Node, const FPCGDataCollection& Output)
+		{
+			for (int32 i = 0; i < Params.CollectOutputDataForNodes.Num(); ++i)
+			{
+				if (Node == Params.CollectOutputDataForNodes[i])
+				{
+					CapturedOutputs[i] = Output;
+					NodeOutputCaptured[i] = true;
+					break;
+				}
+			}
+		}
+	);
+
+	// Sentinel for execution.
+	volatile bool bExecutionComplete = false;
+	Executor->ScheduleGeneric(
+		[&bExecutionComplete]() { bExecutionComplete = true; return true; },
+		TestData.TestPCGComponent,
+		{ FinalTaskId }
+	);
+
+	constexpr float FakeDeltaTime = 1.f / 30.f;
+	constexpr int32 MaxExecuteIterations = 10000;
+	UWorld* TestWorld = TestData.TestActor->GetWorld();
+	for (int32 i = 0; i < MaxExecuteIterations && !bExecutionComplete; ++i)
+	{
+		// Fires FPCGModule::Tick (flushing SleepUntilNextFrame unpause callbacks) AND any ExecuteOnGameThread
+		// tickers added by the previous iteration.
+		FTSTicker::GetCoreTicker().Tick(FakeDeltaTime);
+
+		// Advances any now-unpaused tasks, may enqueue GPU compute work.
+		Executor->Execute();
+
+		// Manually dispatches any queued EndOfFrameUpdate compute work for the scene, since without a real render
+		// frame the engine never does this automatically.
+		if (TestWorld && TestWorld->Scene)
+		{
+			// todo_pcg: Using this execution group is fine for now, but in the future this will be dictated by the requirements of the compute work.
+			ComputeFramework::FlushWork(TestWorld->Scene, ComputeTaskExecutionGroup::EndOfFrameUpdate);
+		}
+
+		// Render thread executes the RDG, then calls ExecuteOnGameThread() which registers export callbacks as
+		// one-shot core tickers for the next iteration.
+		FlushRenderingCommands();
+	}
+
+	if (!Params.Test->TestTrue(TEXT("Graph execution completed within iteration limit"), bExecutionComplete))
+	{
+		Executor->CancelAll();
+		return false;
+	}
+
+	bool bAllNodeOutputsCaptured = true;
+	for (int32 i = 0; i < Params.CollectOutputDataForNodes.Num(); ++i)
+	{
+		bAllNodeOutputsCaptured &= Params.Test->TestTrue(
+			*FString::Printf(TEXT("Output for node %d was captured by debug callback"), i),
+			NodeOutputCaptured[i]);
+	}
+	if (!bAllNodeOutputsCaptured)
+	{
+		return false;
+	}
+
+	if (Params.bReadbackGPUData)
+	{
+		FlushRenderingCommands();
+	}
+
+	OutNodeData.SetNum(Params.CollectOutputDataForNodes.Num());
+
+	// Gather all the output data collections for the requested nodes.
+	for (int32 i = 0; i < CapturedOutputs.Num(); ++i)
+	{
+		for (const FPCGTaggedData& TaggedData : CapturedOutputs[i].TaggedData)
+		{
+			if (Params.bReadbackGPUData)
+			{
+				if (const UPCGProxyForGPUData* Proxy = Cast<UPCGProxyForGPUData>(TaggedData.Data))
+				{
+					UPCGProxyForGPUData::FReadbackResult Result;
+					constexpr int32 MaxReadbackIterations = 10000;
+					for (int32 ReadbackIter = 0; ReadbackIter < MaxReadbackIterations; ++ReadbackIter)
+					{
+						Result = Proxy->GetCPUData(nullptr);
+						if (Result.bComplete)
+						{
+							break;
+						}
+
+						FPlatformProcess::SleepNoStats(0.001f);
+						FlushRenderingCommands();
+					}
+
+					if (!Params.Test->TestTrue(TEXT("GPU->CPU readback completed"), Result.bComplete))
+					{
+						return false;
+					}
+
+					if (!Params.Test->TestNotNull(TEXT("Readback produced a data object"), Result.TaggedData.Data.Get()))
+					{
+						return false;
+					}
+
+					OutNodeData[i].TaggedData.Add(Result.TaggedData);
+					continue;
+				}
+			}
+
+			OutNodeData[i].TaggedData.Add(TaggedData);
+		}
+	}
+
+	if (Executor->IsAnyGraphCurrentlyExecuting())
+	{
+		Executor->CancelAll();
+	}
+
+	return true;
+}
+
+#endif // WITH_EDITOR
diff --git a/Engine/Plugins/PCG/Source/PCG/Private/Tests/Compute/PCGGPUTestCommon.h b/Engine/Plugins/PCG/Source/PCG/Private/Tests/Compute/PCGGPUTestCommon.h
new file mode 100644
--- /dev/null
+++ b/Engine/Plugins/PCG/Source/PCG/Private/Tests/Compute/PCGGPUTestCommon.h
@@ -0,0 +1,66 @@
+// Copyright Epic Games, Inc. All Rights Reserved.
+
+#pragma once
+
+#if WITH_EDITOR
+
+#include "Tests/PCGTestsCommon.h"
+
+#include "PCGData.h"
+
+#include "Templates/Function.h"
+
+class FAutomationTestBase;
+class UPCGGraph;
+class UPCGNode;
+
+namespace PCGGPUTestCommon
+{
+	/**
+	 * Verifies DataCollection contains exactly one item, then calls Validator on it.
+	 * Returns false (with a logged failure) if the count check fails, otherwise returns
+	 * whatever Validator returns.
+	 */
+	bool ValidateSingleTaggedData(FAutomationTestBase* Test, const FPCGDataCollection& DataCollection, TFunctionRef<bool(const FPCGTaggedData&)> Validator);
+}
+
+/** Parameters for FPCGGPUGraphTestRunner::Execute(). */
+struct FPCGGPUGraphTestRunnerParams
+{
+	FAutomationTestBase* Test = nullptr;
+	TArray<UPCGNode*> CollectOutputDataForNodes;
+
+	/** If true (default), UPCGProxyForGPUData outputs are read back to CPU before being added to OutNodeData. */
+	bool bReadbackGPUData = true;
+};
+
+/**
+ * Helper for writing GPU PCG automation tests.
+ *
+ * Handles all boilerplate: actor/component setup, bounds, graph creation, executor scheduling, tick/execute/FlushWork loop,
+ * and GPU->CPU readback.
+ *
+ * Usage:
+ *   1. Construct FPCGGPUGraphTestRunner.
+ *   2. Add nodes to FPCGGPUGraphTestRunner::Graph and wire them.
+ *   3. Specify for which nodes the test runner should collect the output data collections from (FPCGGPUGraphTestRunnerParams::CollectOutputDataForNodes).
+ *   4. Call Runner.Execute(), if it returns true it completed.
+ *   5. Verify OutNodeData in the test body.
+ */
+struct FPCGGPUGraphTestRunner
+{
+	PCGTestsCommon::FTestData TestData;
+	UPCGGraph* Graph = nullptr;
+
+	/** Sets up the test actor (with valid bounds), PCG component, and graph. */
+	FPCGGPUGraphTestRunner();
+
+	/**
+	 * Runs the graph to completion, reads back (if Params.bReadbackGPUData) all
+	 * UPCGProxyForGPUData from each specified node's output to CPU data, and fills OutNodeData.
+	 * Logs test failures via Params.Test on error. Returns true if everything succeeded.
+	 */
+	bool Execute(const FPCGGPUGraphTestRunnerParams& Params, TArray<FPCGDataCollection>& OutNodeData);
+};
+
+#endif // WITH_EDITOR
diff --git a/Engine/Plugins/PCG/Source/PCG/Public/PCGContext.h b/Engine/Plugins/PCG/Source/PCG/Public/PCGContext.h
--- a/Engine/Plugins/PCG/Source/PCG/Public/PCGContext.h
+++ b/Engine/Plugins/PCG/Source/PCG/Public/PCGContext.h
@@ -20,6 +20,7 @@
 class UPCGBasePointData;
 class UPCGComponent;
 class FPCGGraphExecutor;
+class UPCGComputeGraph;
 class UPCGGraphInterface;
 class UPCGSettingsInterface;
 class UPCGSpatialData;
@@ -300,6 +301,9 @@ PRAGMA_ENABLE_DEPRECATION_WARNINGS
 	UE_API bool GetOutputData(FPCGTaskId InTaskId, FPCGDataCollection& OutData);
 	UE_API void ClearOutputData(FPCGTaskId InTaskId);
 
+	/** Retrieves a compiled compute graph from the executor that scheduled this task. */
+	UE_API UPCGComputeGraph* GetComputeGraph(const UPCGGraph* InGraph, uint32 GridSize, uint32 ComputeGraphIndex) const;
+
 	UE_API bool IsCacheEnabled();
 	UE_API void StoreInCache(const FPCGStoreInCacheParams& Params, const FPCGDataCollection& InOutput);
 	UE_API bool GetFromCache(const FPCGGetFromCacheParams& Params, FPCGDataCollection& OutOutput) const;
diff --git a/Engine/Plugins/PCG/Source/PCG/Public/Subsystems/IPCGBaseSubsystem.h b/Engine/Plugins/PCG/Source/PCG/Public/Subsystems/IPCGBaseSubsystem.h
--- a/Engine/Plugins/PCG/Source/PCG/Public/Subsystems/IPCGBaseSubsystem.h
+++ b/Engine/Plugins/PCG/Source/PCG/Public/Subsystems/IPCGBaseSubsystem.h
@@ -93,6 +93,8 @@ public:
 	PCG_API bool IsGraphCacheDebuggingEnabled() const;
 	
 	PCG_API FPCGGraphCompiler* GetGraphCompiler();
+
+	UE_DEPRECATED(5.8, "Use FPCGContext::GetComputeGraph instead, which resolves from the executing graph's compiler cache.")
 	PCG_API UPCGComputeGraph* GetComputeGraph(const UPCGGraph* InGraph, uint32 GridSize, uint32 ComputeGraphIndex);
 
 #if WITH_EDITOR
