Simplify PCG_CONTEXT_COMMAND_INFO and PCG_TOOL_COMMAND_INFO macros to fix loc strings gather

Localization text gathering commandlet doesn't support nested LOCTEXT or NSLOCTEXT macros expansion

#tests Run Localize UAT command with Editor set as loc target. Observe strings are gathred to loc data files

#Synced-CL 51601388

diff --git a/Engine/Plugins/PCG/Source/PCGEditor/Private/EditorMode/PCGEdModeCommands.cpp b/Engine/Plugins/PCG/Source/PCGEditor/Private/EditorMode/PCGEdModeCommands.cpp
--- a/Engine/Plugins/PCG/Source/PCGEditor/Private/EditorMode/PCGEdModeCommands.cpp
+++ b/Engine/Plugins/PCG/Source/PCGEditor/Private/EditorMode/PCGEdModeCommands.cpp
@@ -38,12 +38,12 @@ FPCGEditorModePaletteCommands::FPCGEditorModePaletteCommands()
  * @param Name The name of the tool
  * @param IconStyleName The style name for the icon to be used for the button
  */
-#define PCG_CONTEXT_COMMAND_INFO(CommandInfoSP, Name, IconStyleName)\
+#define PCG_CONTEXT_COMMAND_INFO(CommandInfoSP, CommandName, CommandLabel, ToolTip, IconStyleName)\
 FUICommandInfo::MakeCommandInfo(AsShared()\
 	, CommandInfoSP\
-	, TEXT("Enter" Name "Context")\
-	, LOCTEXT(Name "ContextLabel", Name)\
-	, LOCTEXT(Name "ContextToolTip", "Enter '" Name "' context and open corresponding toolkit")\
+	, CommandName\
+	, CommandLabel\
+	, ToolTip\
 	, FSlateIcon(FPCGEditorModeStyle::Get().GetStyleSetName(), IconStyleName)\
 	, EUserInterfaceActionType::ToggleButton\
 	, FInputChord())
@@ -51,11 +51,11 @@ FUICommandInfo::MakeCommandInfo(AsShared()\
 void FPCGEditorModePaletteCommands::RegisterCommands()
 {
 	/*** PCG EDITOR CONTEXTS **/
-	PCG_CONTEXT_COMMAND_INFO(LoadSplineContextPalette, "Spline", "PCGEditorMode.Context.DrawSpline");
-	PCG_CONTEXT_COMMAND_INFO(LoadPaintContextPalette, "Paint", "PCGEditorMode.Context.Paint");
-	PCG_CONTEXT_COMMAND_INFO(LoadVolumeContextPalette, "Volume", "PCGEditorMode.Context.Volume");
-	PCG_CONTEXT_COMMAND_INFO(LoadSelectContextPalette, "Select", "PCGEditorMode.Context.Select");
-	PCG_CONTEXT_COMMAND_INFO(LoadQueryContextPalette, "Query", "PCGEditorMode.Context.Query");
+	PCG_CONTEXT_COMMAND_INFO(LoadSplineContextPalette, TEXT("EnterSplineContext"), LOCTEXT("SplineContextLabel", "Spline"), LOCTEXT("SplineContextToolTip", "Enter 'Spline' context and open corresponding toolkit"), "PCGEditorMode.Context.DrawSpline");
+	PCG_CONTEXT_COMMAND_INFO(LoadPaintContextPalette, TEXT("EnterPaintContext"), LOCTEXT("PaintContextLabel", "Paint"), LOCTEXT("PaintContextToolTip", "Enter 'Paint' context and open corresponding toolkit"), "PCGEditorMode.Context.Paint");
+	PCG_CONTEXT_COMMAND_INFO(LoadVolumeContextPalette, TEXT("EnterVolumeContext"), LOCTEXT("VolumeContextLabel", "Volume"), LOCTEXT("VolumeContextToolTip", "Enter 'Volume' context and open corresponding toolkit"), "PCGEditorMode.Context.Volume");
+	PCG_CONTEXT_COMMAND_INFO(LoadSelectContextPalette, TEXT("EnterSelectContext"), LOCTEXT("SelectContextLabel", "Select"), LOCTEXT("SelectContextToolTip", "Enter 'Select' context and open corresponding toolkit"), "PCGEditorMode.Context.Select");
+	PCG_CONTEXT_COMMAND_INFO(LoadQueryContextPalette, TEXT("EnterQueryContext"), LOCTEXT("QueryContextLabel", "Query"), LOCTEXT("QueryContextToolTip", "Enter 'Query' context and open corresponding toolkit"), "PCGEditorMode.Context.Query");
 }
 
 #undef PCG_CONTEXT_COMMAND_INFO
@@ -74,12 +74,12 @@ FPCGEditorModeToolCommands::FPCGEditorModeToolCommands()
  * @param ToolTag The tag for the tool: can be used in the tooltip.
  * @param IconStyleName The style name for the icon to be used for the button
  */
-#define PCG_TOOL_COMMAND_INFO(CommandInfoSP, Name, Label, ToolTip, ToolTag, IconStyleName)\
+#define PCG_TOOL_COMMAND_INFO(CommandInfoSP, CommandName, CommandLabel, ToolTip, ToolTag, IconStyleName)\
 FUICommandInfo::MakeCommandInfo(AsShared()\
 	, CommandInfoSP\
-	, TEXT(Name"Tool")\
-	, LOCTEXT(Name"ToolLabel", Label)\
-	, FText::Format(LOCTEXT(Name"ToolTip", ToolTip), FFormatNamedArguments({{TEXT("ToolTag"), FText::FromName(ToolTag)}}))\
+	, CommandName\
+	, CommandLabel\
+	, FText::Format(ToolTip, FFormatNamedArguments({{TEXT("ToolTag"), FText::FromName(ToolTag)}}))\
 	, FSlateIcon(FPCGEditorModeStyle::Get().GetStyleSetName(), IconStyleName)\
 	, EUserInterfaceActionType::ToggleButton\
 	, FInputChord())
@@ -90,13 +90,13 @@ void FPCGEditorModeToolCommands::RegisterCommands()
 
 	/*** PCG EDITOR MODE TOOLS **/
 
-	PCG_TOOL_COMMAND_INFO(EnableDrawSplineTool, "DrawSpline", "Draw Spline", "Use the '{ToolTag}' tool to build a spline for PCG.", UPCGInteractiveToolSettings_Spline::StaticGetToolTag(), "PCGEditorMode.Tools.DrawSpline");
-	PCG_TOOL_COMMAND_INFO(EnableDrawSurfaceTool, "DrawSurface", "Draw Spline Surface", "Use the '{ToolTag}' tool to build a closed spline representing an area for PCG.", UPCGInteractiveToolSettings_SplineSurface::StaticGetToolTag(), "PCGEditorMode.Tools.DrawSurface");
-	PCG_TOOL_COMMAND_INFO(EnablePaintTool, "Paint", "Paint", "Use the '{ToolTag}' tool to apply PCG points with a brush.", UPCGInteractiveToolSettings_PaintTool::StaticGetToolTag(), "PCGEditorMode.Context.Paint");
-	PCG_TOOL_COMMAND_INFO(EnableVolumeTool, "Volume", "Volume", "Create or updates a volume in the scene using the '{ToolTag}' tool.", UPCGInteractiveToolSettings_Volume::StaticGetToolTag(), "PCGEditorMode.Tools.Volume");
-	PCG_TOOL_COMMAND_INFO(EnableSelectTool, "Select", "Select", "Apply PCG graph to selection using the '{ToolTag}' tool.", UPCGInteractiveToolSettings_Select::StaticGetToolTag(), "PCGEditorMode.Tools.Select");
-	PCG_TOOL_COMMAND_INFO(EnableQueryTool, "Query", "Query", "Queries scene data from raycasts.", UPCGQueryTool::StaticGetToolTag(), "PCGEditorMode.Tools.Query");
-	PCG_TOOL_COMMAND_INFO(EnableIsolateTool, "Isolate", "Isolate", "Isolates content from the actor on a tag basis.", UPCGIsolateTool::StaticGetToolTag(), "PCGEditorMode.Tools.Isolate");
+	PCG_TOOL_COMMAND_INFO(EnableDrawSplineTool, "DrawSplineTool", LOCTEXT("DrawSplineToolLabel", "Draw Spline"), LOCTEXT("DrawSplineToolTip", "Use the '{ToolTag}' tool to build a spline for PCG."), UPCGInteractiveToolSettings_Spline::StaticGetToolTag(), "PCGEditorMode.Tools.DrawSpline");
+	PCG_TOOL_COMMAND_INFO(EnableDrawSurfaceTool, "DrawSurfaceTool", LOCTEXT("DrawSurfaceToolLabel", "Draw Spline Surface"), LOCTEXT("DrawSurfaceToolTip", "Use the '{ToolTag}' tool to build a closed spline representing an area for PCG."), UPCGInteractiveToolSettings_SplineSurface::StaticGetToolTag(), "PCGEditorMode.Tools.DrawSurface");
+	PCG_TOOL_COMMAND_INFO(EnablePaintTool, "PaintTool", LOCTEXT("PaintToolLabel", "Paint"), LOCTEXT("PaintToolTip", "Use the '{ToolTag}' tool to apply PCG points with a brush."), UPCGInteractiveToolSettings_PaintTool::StaticGetToolTag(), "PCGEditorMode.Context.Paint");
+	PCG_TOOL_COMMAND_INFO(EnableVolumeTool, "VolumeTool", LOCTEXT("VolumeToolLabel", "Volume"), LOCTEXT("VolumeToolTip", "Create or updates a volume in the scene using the '{ToolTag}' tool."), UPCGInteractiveToolSettings_Volume::StaticGetToolTag(), "PCGEditorMode.Tools.Volume");
+	PCG_TOOL_COMMAND_INFO(EnableSelectTool, "SelectTool", LOCTEXT("SelectToolLabel", "Select"), LOCTEXT("SelectToolTip", "Apply PCG graph to selection using the '{ToolTag}' tool."), UPCGInteractiveToolSettings_Select::StaticGetToolTag(), "PCGEditorMode.Tools.Select");
+	PCG_TOOL_COMMAND_INFO(EnableQueryTool, "QueryTool", LOCTEXT("QueryToolLabel", "Query"), LOCTEXT("QueryToolTip", "Queries scene data from raycasts."), UPCGQueryTool::StaticGetToolTag(), "PCGEditorMode.Tools.Query");
+	PCG_TOOL_COMMAND_INFO(EnableIsolateTool, "IsolateTool", LOCTEXT("IsolateToolLabel", "Isolate"), LOCTEXT("IsolateToolTip", "Isolates content from the actor on a tag basis."), UPCGIsolateTool::StaticGetToolTag(), "PCGEditorMode.Tools.Isolate");
 }
 
 void FPCGEditorModeToolCommands::RegisterAllToolCommands()
