[PCG] Expose a hook for assigning the SubgraphSettings class. This will allow PVE to use its own SubgraphSettings class (PVSubgraphSettings) when dragging and dropping the embedded graph node within the graph.

#Synced-CL 51579190

diff --git a/Engine/Plugins/Experimental/ProceduralVegetationEditor/Source/ProceduralVegetationEditor/Private/PVEditorSchema.cpp b/Engine/Plugins/Experimental/ProceduralVegetationEditor/Source/ProceduralVegetationEditor/Private/PVEditorSchema.cpp
--- a/Engine/Plugins/Experimental/ProceduralVegetationEditor/Source/ProceduralVegetationEditor/Private/PVEditorSchema.cpp
+++ b/Engine/Plugins/Experimental/ProceduralVegetationEditor/Source/ProceduralVegetationEditor/Private/PVEditorSchema.cpp
@@ -15,6 +15,11 @@ EPCGElementType UPVEditorSchema::GetElementTypeFiltering() const
 	return EPCGElementType::Native;
 }
 
+TSubclassOf<UPCGBaseSubgraphSettings> UPVEditorSchema::GetSubgraphSettingsClass() const
+{
+	return UPVSubgraphSettings::StaticClass();
+}
+
 const FPCGGraphEditorCustomization& UPVEditorSchema::GetGraphEditorCustomization(const UEdGraph* InEdGraph) const
 {
 	static const FPCGGraphEditorCustomization GraphEditorCustomization
diff --git a/Engine/Plugins/Experimental/ProceduralVegetationEditor/Source/ProceduralVegetationEditor/Private/PVEditorSchema.h b/Engine/Plugins/Experimental/ProceduralVegetationEditor/Source/ProceduralVegetationEditor/Private/PVEditorSchema.h
--- a/Engine/Plugins/Experimental/ProceduralVegetationEditor/Source/ProceduralVegetationEditor/Private/PVEditorSchema.h
+++ b/Engine/Plugins/Experimental/ProceduralVegetationEditor/Source/ProceduralVegetationEditor/Private/PVEditorSchema.h
@@ -14,6 +14,8 @@ public:
 	// Return the element type filtering for the current editor. By default it is all.
 	virtual EPCGElementType GetElementTypeFiltering() const override;
 
+	virtual TSubclassOf<UPCGBaseSubgraphSettings> GetSubgraphSettingsClass() const override;
+
 	// Return the graph customization for this editor. By default it is the one provided by the graph.
 	virtual const FPCGGraphEditorCustomization& GetGraphEditorCustomization(const UEdGraph* InEdGraph) const override;
 };
diff --git a/Engine/Plugins/PCG/Source/PCGEditor/Private/Schema/PCGEditorGraphSchema.cpp b/Engine/Plugins/PCG/Source/PCGEditor/Private/Schema/PCGEditorGraphSchema.cpp
--- a/Engine/Plugins/PCG/Source/PCGEditor/Private/Schema/PCGEditorGraphSchema.cpp
+++ b/Engine/Plugins/PCG/Source/PCGEditor/Private/Schema/PCGEditorGraphSchema.cpp
@@ -62,6 +62,11 @@ void UPCGEditorGraphSchema::GetPaletteActions(FGraphActionMenuBuilder& ActionMen
 	GetGraphActions(ActionMenuBuilder, ActionFilter, /*bIsContextual=*/false);
 }
 
+TSubclassOf<UPCGBaseSubgraphSettings> UPCGEditorGraphSchema::GetSubgraphSettingsClass() const
+{
+	return UPCGSubgraphSettings::StaticClass();
+}
+
 const FPCGGraphEditorCustomization& UPCGEditorGraphSchema::GetGraphEditorCustomization(const UEdGraph* InEdGraph) const
 {
 	static const FPCGGraphEditorCustomization Default{};
diff --git a/Engine/Plugins/PCG/Source/PCGEditor/Private/Schema/PCGEditorGraphSchemaActions.cpp b/Engine/Plugins/PCG/Source/PCGEditor/Private/Schema/PCGEditorGraphSchemaActions.cpp
--- a/Engine/Plugins/PCG/Source/PCGEditor/Private/Schema/PCGEditorGraphSchemaActions.cpp
+++ b/Engine/Plugins/PCG/Source/PCGEditor/Private/Schema/PCGEditorGraphSchemaActions.cpp
@@ -647,8 +647,10 @@ UEdGraphNode* FPCGEditorGraphSchemaAction_NewSubgraphElement::MakeGraphNode(UPCG
 	
 	if (!bInCreateLoop)
 	{
-		NewPCGNode = PCGGraph->AddNodeOfType(UPCGSubgraphSettings::StaticClass(), DefaultNodeSettings);
-		UPCGSubgraphSettings* DefaultSubgraphSettings = CastChecked<UPCGSubgraphSettings>(DefaultNodeSettings);
+		const UPCGEditorGraphSchema* PCGSchema = CastChecked<UPCGEditorGraphSchema>(InEditorGraph->GetSchema());
+		
+		NewPCGNode = PCGGraph->AddNodeOfType(PCGSchema->GetSubgraphSettingsClass(), DefaultNodeSettings);
+		UPCGBaseSubgraphSettings* DefaultSubgraphSettings = CastChecked<UPCGBaseSubgraphSettings>(DefaultNodeSettings);
 		DefaultSubgraphSettings->SetSubgraph(InGraph);
 	}
 	else
diff --git a/Engine/Plugins/PCG/Source/PCGEditor/Public/Schema/PCGEditorGraphSchema.h b/Engine/Plugins/PCG/Source/PCGEditor/Public/Schema/PCGEditorGraphSchema.h
--- a/Engine/Plugins/PCG/Source/PCGEditor/Public/Schema/PCGEditorGraphSchema.h
+++ b/Engine/Plugins/PCG/Source/PCGEditor/Public/Schema/PCGEditorGraphSchema.h
@@ -12,6 +12,7 @@
 
 #include "PCGEditorGraphSchema.generated.h"
 
+class UPCGBaseSubgraphSettings;
 enum class EPCGElementType : uint8;
 
 class IAssetReferenceFilter;
@@ -37,8 +38,11 @@ public:
 	PCGEDITOR_API void GetPaletteActions(FGraphActionMenuBuilder& ActionMenuBuilder, const FPCGActionsFilter& ActionFilter) const;
 
 	// Return the element type filtering for the current editor. By default it is all.
-	virtual EPCGElementType GetElementTypeFiltering() const { return EPCGElementType::All; };
+	virtual EPCGElementType GetElementTypeFiltering() const { return EPCGElementType::All; }
 
+	// Returns the settings class to use when creating a subgraph node. Override in derived schemas to provide a custom subgraph settings class.
+	PCGEDITOR_API virtual TSubclassOf<UPCGBaseSubgraphSettings> GetSubgraphSettingsClass() const;
+	
 	// Return the graph customization for this editor. By default it is the one provided by the graph.
 	PCGEDITOR_API virtual const FPCGGraphEditorCustomization& GetGraphEditorCustomization(const UEdGraph* InEdGraph) const;
 
