[PCG] Follow up feedback to GPU tests.

Mainly a refactor to fix weird line of code where context assigns compute graph to itself. Decided that the get compute graph method was not necessary at the FPCGContext level so pushed it down to the compute graph context.

#rnx
#Synced-CL 51579517

diff --git a/Engine/Plugins/PCG/Source/PCG/Private/Compute/Elements/PCGComputeGraphElement.cpp b/Engine/Plugins/PCG/Source/PCG/Private/Compute/Elements/PCGComputeGraphElement.cpp
--- a/Engine/Plugins/PCG/Source/PCG/Private/Compute/Elements/PCGComputeGraphElement.cpp
+++ b/Engine/Plugins/PCG/Source/PCG/Private/Compute/Elements/PCGComputeGraphElement.cpp
@@ -14,6 +14,7 @@
 #include "Elements/PCGStaticMeshSpawnerKernel.h"
 #include "Editor/IPCGEditorModule.h"
 #include "Graph/PCGGraphCompiler.h"
+#include "Graph/PCGGraphExecutor.h"
 #include "Helpers/PCGHelpers.h"
 #include "Subsystems/IPCGBaseSubsystem.h"
 
@@ -94,6 +95,16 @@ bool FPCGComputeGraphContext::HasPendingAsyncOperations() const
 	return !ProvidersWithBufferExports.IsEmpty();
 }
 
+bool FPCGComputeGraphContext::ResolveComputeGraph(const UPCGGraph* InGraph, uint32 GridSize, uint32 InComputeGraphIndex)
+{
+	if (TSharedPtr<FPCGGraphExecutor> Executor = GetGraphExecutor().Pin())
+	{
+		ComputeGraph = Executor->GetCompiler()->GetComputeGraph(InGraph, GridSize, InComputeGraphIndex);
+	}
+
+	return ComputeGraph != nullptr;
+}
+
 void FPCGComputeGraphContext::AddExtraStructReferencedObjects(FReferenceCollector& Collector)
 {
 	if (bComputeGraphInstanceInitialized)
@@ -182,9 +193,7 @@ bool FPCGComputeGraphElement::ExecuteInternal(FPCGContext* InContext) const
 			Context->GenerationGridSize = Source->GetExecutionState().GetGenerationGridSize();
 		}
 		
-		Context->ComputeGraph = Context->GetComputeGraph(DynamicSubgraph ? DynamicSubgraph : TopGraph, Context->GenerationGridSize, ComputeGraphIndex);
-
-		if (!Context->ComputeGraph)
+		if (!Context->ResolveComputeGraph(DynamicSubgraph ? DynamicSubgraph : TopGraph, Context->GenerationGridSize, ComputeGraphIndex))
 		{
 #if WITH_EDITOR
 			// Normally emitted when a graph compilation is flushed while generation in progress.
diff --git a/Engine/Plugins/PCG/Source/PCG/Private/Compute/Elements/PCGComputeGraphElement.h b/Engine/Plugins/PCG/Source/PCG/Private/Compute/Elements/PCGComputeGraphElement.h
--- a/Engine/Plugins/PCG/Source/PCG/Private/Compute/Elements/PCGComputeGraphElement.h
+++ b/Engine/Plugins/PCG/Source/PCG/Private/Compute/Elements/PCGComputeGraphElement.h
@@ -42,6 +42,9 @@ public:
 
 	bool HasPendingAsyncOperations() const;
 
+	/** Resolves and assigns ComputeGraph from the executor's compiler cache. Returns true if the graph was found. */
+	bool ResolveComputeGraph(const UPCGGraph* InGraph, uint32 GridSize, uint32 InComputeGraphIndex);
+
 protected:
 	virtual void AddExtraStructReferencedObjects(FReferenceCollector& Collector) override;
 
diff --git a/Engine/Plugins/PCG/Source/PCG/Private/PCGContext.cpp b/Engine/Plugins/PCG/Source/PCG/Private/PCGContext.cpp
--- a/Engine/Plugins/PCG/Source/PCG/Private/PCGContext.cpp
+++ b/Engine/Plugins/PCG/Source/PCG/Private/PCGContext.cpp
@@ -730,16 +730,6 @@ void FPCGContext::ClearOutputData(FPCGTaskId InTaskId)
 	}
 }
 
-UPCGComputeGraph* FPCGContext::GetComputeGraph(const UPCGGraph* InGraph, uint32 GridSize, uint32 ComputeGraphIndex) const
-{
-	if (TSharedPtr<FPCGGraphExecutor> Executor = GraphExecutor.Pin())
-	{
-		return Executor->GetCompiler()->GetComputeGraph(InGraph, GridSize, ComputeGraphIndex);
-	}
-
-	return nullptr;
-}
-
 void FPCGContext::InitializeGraphExecutor(FPCGContext* InContext)
 {
 	check(InContext);
diff --git a/Engine/Plugins/PCG/Source/PCG/Private/Tests/Compute/PCGGPUTestCommon.cpp b/Engine/Plugins/PCG/Source/PCG/Private/Tests/Compute/PCGGPUTestCommon.cpp
--- a/Engine/Plugins/PCG/Source/PCG/Private/Tests/Compute/PCGGPUTestCommon.cpp
+++ b/Engine/Plugins/PCG/Source/PCG/Private/Tests/Compute/PCGGPUTestCommon.cpp
@@ -174,9 +174,10 @@ bool FPCGGPUGraphTestRunner::Execute(const FPCGGPUGraphTestRunnerParams& Params,
 		}
 	}
 
-	if (Executor->IsAnyGraphCurrentlyExecuting())
+	if (!Params.Test->TestFalse(TEXT("All generation should be complete"), Executor->IsAnyGraphCurrentlyExecuting()))
 	{
 		Executor->CancelAll();
+		return false;
 	}
 
 	return true;
diff --git a/Engine/Plugins/PCG/Source/PCG/Public/PCGContext.h b/Engine/Plugins/PCG/Source/PCG/Public/PCGContext.h
--- a/Engine/Plugins/PCG/Source/PCG/Public/PCGContext.h
+++ b/Engine/Plugins/PCG/Source/PCG/Public/PCGContext.h
@@ -301,9 +301,6 @@ PRAGMA_ENABLE_DEPRECATION_WARNINGS
 	UE_API bool GetOutputData(FPCGTaskId InTaskId, FPCGDataCollection& OutData);
 	UE_API void ClearOutputData(FPCGTaskId InTaskId);
 
-	/** Retrieves a compiled compute graph from the executor that scheduled this task. */
-	UE_API UPCGComputeGraph* GetComputeGraph(const UPCGGraph* InGraph, uint32 GridSize, uint32 ComputeGraphIndex) const;
-
 	UE_API bool IsCacheEnabled();
 	UE_API void StoreInCache(const FPCGStoreInCacheParams& Params, const FPCGDataCollection& InOutput);
 	UE_API bool GetFromCache(const FPCGGetFromCacheParams& Params, FPCGDataCollection& OutOutput) const;
@@ -326,6 +323,8 @@ protected:
 	
 	UE_API void InitializeGraphExecutor(FPCGContext* InContext);
 
+	TWeakPtr<FPCGGraphExecutor> GetGraphExecutor() { return GraphExecutor; }
+
 private:
 	template<class T, typename... Args>
 	T* NewObject_AnyThread_Impl(Args&&... InArgs)
