[PCG] Move standalone execution sources to base class

#rnx
#Synced-CL 51473272

diff --git a/Engine/Plugins/PCG/Source/PCG/Private/Subsystems/IPCGBaseSubsystem.cpp b/Engine/Plugins/PCG/Source/PCG/Private/Subsystems/IPCGBaseSubsystem.cpp
--- a/Engine/Plugins/PCG/Source/PCG/Private/Subsystems/IPCGBaseSubsystem.cpp
+++ b/Engine/Plugins/PCG/Source/PCG/Private/Subsystems/IPCGBaseSubsystem.cpp
@@ -3,9 +3,11 @@
 #include "Subsystems/IPCGBaseSubsystem.h"
 
 #include "PCGComponent.h"
+#include "PCGDefaultExecutionSource.h"
 #include "PCGGraph.h"
 #include "PCGModule.h"
 #include "PCGTrackingManager.h"
+#include "Editor/IPCGEditorModule.h"
 #include "Graph/PCGGraphCache.h"
 #include "Graph/PCGGraphCompiler.h"
 #include "Graph/PCGGraphExecutor.h"
@@ -200,6 +202,14 @@ void IPCGBaseSubsystem::ClearOutputData(FPCGTaskId TaskId)
 }
 
 #if WITH_EDITOR
+void IPCGBaseSubsystem::AddReferencedObjects(FReferenceCollector& Collector)
+{
+	for (FGraphExecution& GraphExecution : GraphExecutions)
+	{
+		Collector.AddReferencedObject(GraphExecution.ExecutionSource);
+	}
+}
+
 void IPCGBaseSubsystem::NotifyGraphChanged(UPCGGraph* InGraph, EPCGChangeType ChangeType)
 {
 	if (GraphExecutor)
@@ -270,6 +280,30 @@ uint32 IPCGBaseSubsystem::GetGraphCacheEntryCount(IPCGElement* InElement) const
 void IPCGBaseSubsystem::OnPCGSourceGenerationDone(IPCGGraphExecutionSource* InExecutionSource, EPCGGenerationStatus InStatus)
 {
 	OnPCGSourceGenerationDoneDelegate.Broadcast(this, InExecutionSource, InStatus);
+
+	UPCGDefaultExecutionSource* ExecutionSource = Cast<UPCGDefaultExecutionSource>(InExecutionSource);
+	if (!ExecutionSource)
+	{
+		return;
+	}
+
+	const int32 Index = GraphExecutions.IndexOfByPredicate([ExecutionSource](const FGraphExecution& GraphExecution) { return ExecutionSource == GraphExecution.ExecutionSource.Get(); });
+
+	if (Index != INDEX_NONE)
+	{
+		FGraphExecution& GraphExecution = GraphExecutions[Index];
+		GraphExecution.GenerationCallback.Broadcast(this, ExecutionSource, InStatus);
+		GraphExecutions.RemoveAtSwap(Index);
+
+		if (IPCGEditorModule* PCGEditorModule = IPCGEditorModule::Get())
+		{
+			PCGEditorModule->ClearExecutionMetadata(ExecutionSource);
+		}
+
+		// Mark source as invalid, request GC (non-blocking)
+		ExecutionSource->MarkAsGarbage();
+		GEngine->ForceGarbageCollection(true);
+	}
 }
 
 void IPCGBaseSubsystem::SetDisableClearResults(bool bInDisableClearResults)
diff --git a/Engine/Plugins/PCG/Source/PCG/Private/Subsystems/PCGEngineSubsystem.cpp b/Engine/Plugins/PCG/Source/PCG/Private/Subsystems/PCGEngineSubsystem.cpp
--- a/Engine/Plugins/PCG/Source/PCG/Private/Subsystems/PCGEngineSubsystem.cpp
+++ b/Engine/Plugins/PCG/Source/PCG/Private/Subsystems/PCGEngineSubsystem.cpp
@@ -25,11 +25,7 @@ void UPCGEngineSubsystem::AddReferencedObjects(UObject* InThis, FReferenceCollec
 
 	UPCGEngineSubsystem* This = Cast<UPCGEngineSubsystem>(InThis);
 	check(This);
-
-	for (FGraphExecution& GraphExecution : This->GraphExecutions)
-	{
-		Collector.AddReferencedObject(GraphExecution.ExecutionSource);
-	}
+	This->IPCGBaseSubsystem::AddReferencedObjects(Collector);
 }
 
 void UPCGEngineSubsystem::Initialize(FSubsystemCollectionBase& Collection)
@@ -61,35 +57,6 @@ void UPCGEngineSubsystem::NotifyGraphChanged(UPCGGraph* InGraph, EPCGChangeType
 	IPCGBaseSubsystem::NotifyGraphChanged(InGraph, ChangeType);
 }
 
-void UPCGEngineSubsystem::OnPCGSourceGenerationDone(IPCGGraphExecutionSource* InExecutionSource, EPCGGenerationStatus InStatus)
-{
-	IPCGBaseSubsystem::OnPCGSourceGenerationDone(InExecutionSource, InStatus);
-
-	UPCGDefaultExecutionSource* ExecutionSource = Cast<UPCGDefaultExecutionSource>(InExecutionSource);
-	if (!ExecutionSource)
-	{
-		return;
-	}
-
-	const int32 Index = GraphExecutions.IndexOfByPredicate([ExecutionSource](const FGraphExecution& GraphExecution) { return ExecutionSource == GraphExecution.ExecutionSource.Get(); });
-
-	if (Index != INDEX_NONE)
-	{
-		FGraphExecution& GraphExecution = GraphExecutions[Index];
-		GraphExecution.GenerationCallback.ExecuteIfBound(ExecutionSource, InStatus);
-		GraphExecutions.RemoveAtSwap(Index);
-
-		if (IPCGEditorModule* PCGEditorModule = IPCGEditorModule::Get())
-		{
-			PCGEditorModule->ClearExecutionMetadata(ExecutionSource);
-		}
-
-		// Mark source as invalid, request GC (non-blocking)
-		ExecutionSource->MarkAsGarbage();
-		GEngine->ForceGarbageCollection(true);
-	}
-}
-
 void UPCGEngineSubsystem::GenerateGraph(const FPCGGenerateGraphParams& InParams)
 {
 	if (!InParams.Graph)
diff --git a/Engine/Plugins/PCG/Source/PCG/Public/Editor/IPCGEditorModule.h b/Engine/Plugins/PCG/Source/PCG/Public/Editor/IPCGEditorModule.h
--- a/Engine/Plugins/PCG/Source/PCG/Public/Editor/IPCGEditorModule.h
+++ b/Engine/Plugins/PCG/Source/PCG/Public/Editor/IPCGEditorModule.h
@@ -47,6 +47,7 @@ public:
 #endif
 
 protected:
+	friend class IPCGBaseSubsystem;
 	friend class FPCGEditor;
 	friend class FPCGGraphExecutor;
 	friend class SPCGEditorGraphDebugObjectTree;
diff --git a/Engine/Plugins/PCG/Source/PCG/Public/Subsystems/IPCGBaseSubsystem.h b/Engine/Plugins/PCG/Source/PCG/Public/Subsystems/IPCGBaseSubsystem.h
--- a/Engine/Plugins/PCG/Source/PCG/Public/Subsystems/IPCGBaseSubsystem.h
+++ b/Engine/Plugins/PCG/Source/PCG/Public/Subsystems/IPCGBaseSubsystem.h
@@ -13,20 +13,21 @@
 
 #include "IPCGBaseSubsystem.generated.h"
 
-class UPCGComputeGraph;
-class UPCGData;
-class UPCGGraph;
-
 class IPCGBaseSubsystem;
 class IPCGElement;
 class IPCGGraphCache;
 class FPCGGraphCompiler;
 class FPCGGraphExecutor;
+class UPCGComputeGraph;
+class UPCGData;
+class UPCGDefaultExecutionSource;
+class UPCGGraph;
+class UPCGSettings;
+class UWorld;
 struct FPCGDataCollection;
 struct FPCGScheduleGenericParams;
 struct FPCGStack;
-class UPCGSettings;
-class UWorld;
+struct FPCGDefaultExecutionSourceParams;
 
 #if WITH_EDITOR
 DECLARE_MULTICAST_DELEGATE_ThreeParams(FPCGOnPCGSourceGenerationDone, IPCGBaseSubsystem*, IPCGGraphExecutionSource*, EPCGGenerationStatus);
@@ -118,6 +119,38 @@ public:
 
 	FPCGOnPCGSourceGenerationDone& GetOnPCGSourceGenerationDone() { return OnPCGSourceGenerationDoneDelegate; }
 	FPCGOnPCGSourceUnregistered& GetOnPCGSourceUnregistered() { return OnPCGSourceUnregisteredDelegate; }
+
+	/** Creates an execution source from the params and optionally tracks its lifetime and generation. Note this call is NOT THREADSAFE and as such this should be called only from the game thread. */
+	template<typename SourceType, typename ParamsType>
+		requires std::is_base_of_v<IPCGGraphExecutionSource, SourceType>
+	SourceType* CreateExecutionSource(const ParamsType& InParams, UObject* Outer)
+	{
+		check(IsInGameThread());
+		if (!InParams.GraphInterface || !Outer)
+		{
+			return nullptr;
+		}
+
+		SourceType* NewSource = NewObject<SourceType>(Outer);
+		NewSource->Initialize(InParams);
+
+		if (InParams.bFireAndForgetExecution)
+		{
+			FPCGOnPCGSourceGenerationDone GenerationCallbackWrapper;
+			if (InParams.GenerationCallback.IsBound())
+			{
+				GenerationCallbackWrapper.AddLambda([GenerationCallback = InParams.GenerationCallback](IPCGBaseSubsystem*, IPCGGraphExecutionSource* Source, EPCGGenerationStatus Status)
+				{
+					GenerationCallback.ExecuteIfBound(Source, Status);
+				});
+			}
+
+			GraphExecutions.Emplace(NewSource, GenerationCallbackWrapper);
+			NewSource->Generate();
+		}
+
+		return NewSource;
+	}
 #endif // WITH_EDITOR
 	
 protected:
@@ -136,6 +169,16 @@ protected:
 #if WITH_EDITOR
 	PCG_API void SetDisableClearResults(bool bInDisableClearResults);
 
+	PCG_API void AddReferencedObjects(FReferenceCollector& Collector);
+
+	struct FGraphExecution
+	{
+		TObjectPtr<UPCGDefaultExecutionSource> ExecutionSource;
+		FPCGOnPCGSourceGenerationDone GenerationCallback;
+	};
+
+	TArray<FGraphExecution> GraphExecutions;
+
 	FPCGOnPCGSourceGenerationDone OnPCGSourceGenerationDoneDelegate;
 	FPCGOnPCGSourceUnregistered OnPCGSourceUnregisteredDelegate;
 #endif
diff --git a/Engine/Plugins/PCG/Source/PCG/Public/Subsystems/PCGEngineSubsystem.h b/Engine/Plugins/PCG/Source/PCG/Public/Subsystems/PCGEngineSubsystem.h
--- a/Engine/Plugins/PCG/Source/PCG/Public/Subsystems/PCGEngineSubsystem.h
+++ b/Engine/Plugins/PCG/Source/PCG/Public/Subsystems/PCGEngineSubsystem.h
@@ -50,28 +50,12 @@ public:
 		requires std::is_base_of_v<IPCGGraphExecutionSource, SourceType> && std::is_base_of_v<FPCGDefaultExecutionSourceParams, ParamsType>
 	SourceType* CreateExecutionSource(const ParamsType& InParams)
 	{
-		if (!InParams.GraphInterface)
-		{
-			return nullptr;
-		}
-
-		SourceType* NewSource = NewObject<SourceType>(this);
-		NewSource->Initialize(InParams);
-
-		if (InParams.bFireAndForgetExecution)
-		{
-			GraphExecutions.Emplace(NewSource, InParams.GenerationCallback);
-			NewSource->Generate();
-		}
-
-		return NewSource;
+		return IPCGBaseSubsystem::CreateExecutionSource<SourceType, ParamsType>(InParams, this);
 	}
 
 	PCG_API void GenerateGraph(const FPCGGenerateGraphParams& InParams);
 
 	PCG_API virtual void NotifyGraphChanged(UPCGGraph* InGraph, EPCGChangeType ChangeType) override;
-
-	PCG_API virtual void OnPCGSourceGenerationDone(IPCGGraphExecutionSource* InExecutionSource, EPCGGenerationStatus InStatus) override;
 #endif // WITH_EDITOR
 
 protected:
@@ -81,14 +65,4 @@ protected:
 	PCG_API virtual TStatId GetStatId() const override;
 	virtual bool IsTickableInEditor() const { return true; }
 	//~ END FTickableGameObject Interface
-
-#if WITH_EDITOR
-	struct FGraphExecution
-	{
-		TObjectPtr<UPCGDefaultExecutionSource> ExecutionSource;
-		FPCGOnEditorGenerationDone GenerationCallback;
-	};
-
-	TArray<FGraphExecution> GraphExecutions;
-#endif // WITH_EDITOR
 };
\ No newline at end of file
