[PCG] Added type filtering in the contextual menu.

#Synced-CL 51473272

diff --git a/Engine/Plugins/PCG/Source/PCG/Private/PCGDataAsset.cpp b/Engine/Plugins/PCG/Source/PCG/Private/PCGDataAsset.cpp
--- a/Engine/Plugins/PCG/Source/PCG/Private/PCGDataAsset.cpp
+++ b/Engine/Plugins/PCG/Source/PCG/Private/PCGDataAsset.cpp
@@ -1,4 +1,4 @@
-﻿// Copyright Epic Games, Inc. All Rights Reserved.
+// Copyright Epic Games, Inc. All Rights Reserved.
 
 #include "PCGDataAsset.h"
 
@@ -69,7 +69,15 @@ PCGDataAsset::FGetAssetDataOutput UPCGDataAsset::GetAssetRegistryData(const TSof
 #if WITH_EDITOR
 	const FAssetRegistryModule& AssetRegistryModule = FModuleManager::LoadModuleChecked<FAssetRegistryModule>(TEXT("AssetRegistry"));
 	const FAssetData AssetData = AssetRegistryModule.Get().GetAssetByObjectPath(InSoftAsset.ToSoftObjectPath());
+	return GetAssetRegistryData(AssetData);
+#else // !WITH_EDITOR
+	return {};
+#endif // WITH_EDITOR
+}
 
+#if WITH_EDITOR
+PCGDataAsset::FGetAssetDataOutput UPCGDataAsset::GetAssetRegistryData(const FAssetData& AssetData)
+{
 	if (const UClass* AssetClass = AssetData.GetClass(); AssetClass && AssetClass->IsChildOf(UPCGDataAsset::StaticClass()))
 	{
 		PCGDataAsset::FGetAssetDataOutput Result =
@@ -93,7 +101,7 @@ PCGDataAsset::FGetAssetDataOutput UPCGDataAsset::GetAssetRegistryData(const TSof
 		{
 			// Initialize it with the default value.
 			Result.Color.Emplace(FLinearColor::White);
-			
+
 			TempString.Reset();
 			AssetData.GetTagValue<FString>(GET_MEMBER_NAME_CHECKED(UPCGDataAsset, Color), TempString);
 			if (!TempString.IsEmpty())
@@ -104,7 +112,7 @@ PCGDataAsset::FGetAssetDataOutput UPCGDataAsset::GetAssetRegistryData(const TSof
 
 		return Result;
 	}
-#endif // WITH_EDITOR
 
 	return {};
-}
\ No newline at end of file
+}
+#endif // WITH_EDITOR
\ No newline at end of file
diff --git a/Engine/Plugins/PCG/Source/PCG/Public/PCGDataAsset.h b/Engine/Plugins/PCG/Source/PCG/Public/PCGDataAsset.h
--- a/Engine/Plugins/PCG/Source/PCG/Public/PCGDataAsset.h
+++ b/Engine/Plugins/PCG/Source/PCG/Public/PCGDataAsset.h
@@ -71,6 +71,9 @@ public:
 	PCG_API TArray<FPCGPinProperties> GetPins() const;
 
 	static PCG_API PCGDataAsset::FGetAssetDataOutput GetAssetRegistryData(const TSoftObjectPtr<UPCGDataAsset>& InSoftAsset, const bool bShouldLoadAsset);
+#if WITH_EDITOR
+	static PCG_API PCGDataAsset::FGetAssetDataOutput GetAssetRegistryData(const FAssetData& AssetData);
+#endif // WITH_EDITOR
 
 	/** Contains direct data owned by this data asset */
 	UPROPERTY(BlueprintReadWrite, VisibleAnywhere, Category = Data, meta = (NoResetToDefault))
diff --git a/Engine/Plugins/PCG/Source/PCGEditor/Private/PCGEditor.cpp b/Engine/Plugins/PCG/Source/PCGEditor/Private/PCGEditor.cpp
--- a/Engine/Plugins/PCG/Source/PCGEditor/Private/PCGEditor.cpp
+++ b/Engine/Plugins/PCG/Source/PCGEditor/Private/PCGEditor.cpp
@@ -1790,7 +1790,7 @@ FReply FPCGEditor::OnSpawnNodeByShortcut(FInputChord InChord, const FVector2f& I
 	}
 }
 
-FActionMenuContent FPCGEditor::OnCreateActionMenuContent(UEdGraph* InGraph, const FVector2f& Location, const TArray<UEdGraphPin*>& InDraggedPins, bool bAutoExpand, SGraphEditor::FActionMenuClosed OnMenuClosed)
+FActionMenuContent FPCGEditor::OnCreateActionMenu(UEdGraph* InGraph, const FVector2f& Location, const TArray<UEdGraphPin*>& InDraggedPins, bool bAutoExpand, SGraphEditor::FActionMenuClosed OnMenuClosed)
 {
 	TSharedRef<SGraphEditorActionMenu> Menu = SNew(SGraphEditorActionMenu)
 		.GraphObj(InGraph)
@@ -1804,7 +1804,136 @@ FActionMenuContent FPCGEditor::OnCreateActionMenuContent(UEdGraph* InGraph, cons
 				return SNew(SPCGGraphActionWidget, Data);
 			}));
 
-	return FActionMenuContent(Menu, Menu->GetFilterTextBox());
+	// No need for the contextual menu options if it is not dragged from pin(s).
+	if (InDraggedPins.IsEmpty())
+	{
+		return FActionMenuContent(Menu, Menu->GetFilterTextBox());
+	}
+
+	TSharedRef<SWidget> MenuWithOptions =
+		SNew(SVerticalBox)
+		+ SVerticalBox::Slot()
+		.AutoHeight()
+		.Padding(4, 4, 4, 0)
+		[
+			SNew(SHorizontalBox)
+			+SHorizontalBox::Slot()
+			.HAlign(HAlign_Left)
+			.VAlign(VAlign_Fill)
+			.AutoWidth()
+			.Padding(0)
+			[
+				SNew(STextBlock)
+				.Text_Lambda([this]()
+				{
+					return IsContextualFilteringEnabled(PCGActionsHelpers::ECompatibilityMask::Compatible) ? LOCTEXT("PCGAllNodesShown", "Compatible Nodes") : LOCTEXT("PCGOnlyCompatibleNodesShown", "All Nodes");
+				})
+			]
+			+SHorizontalBox::Slot()
+			.FillWidth(1)
+			.HAlign(HAlign_Center)
+			.Padding(0)
+			[
+				// Left blank to separate the text from the buttons
+				SNullWidget::NullWidget
+			]
+			+SHorizontalBox::Slot()
+			.VAlign(VAlign_Fill)
+			.HAlign(HAlign_Right)
+			.AutoWidth()
+			.Padding(0)
+			[
+				SNew(SCheckBox)
+				.ToolTipText(LOCTEXT("PCGShowContextual", "Show contextual nodes only, e.g. nodes that can be connected directly."))
+				.Style(&FAppStyle::Get().GetWidgetStyle<FCheckBoxStyle>("ToggleButtonCheckbox"))
+				.OnCheckStateChanged(FOnCheckStateChanged::CreateSP(this, &FPCGEditor::OnContextualFilteringChanged, Menu.ToWeakPtr(), PCGActionsHelpers::ECompatibilityMask::Compatible))
+				.IsChecked(TAttribute<ECheckBoxState>::CreateSP(this, &FPCGEditor::IsContextualFilteringChecked, PCGActionsHelpers::ECompatibilityMask::Compatible))
+				.Padding(FMargin(2.0f, 0))
+				[
+					SNew(STextBlock)
+					.Text(LOCTEXT("PCGContextualMenuActivated", "Contextual"))
+				]
+			]
+			+ SHorizontalBox::Slot()
+			.VAlign(VAlign_Fill)
+			.HAlign(HAlign_Right)
+			.AutoWidth()
+			.Padding(0)
+			[
+				// Accept filters
+				SNew(SCheckBox)
+				.ToolTipText(LOCTEXT("PCGAcceptFiltering", "Show nodes that would apply a filter"))
+				.Style(&FAppStyle::Get().GetWidgetStyle<FCheckBoxStyle>("ToggleButtonCheckbox"))
+				.OnCheckStateChanged(FOnCheckStateChanged::CreateSP(this, &FPCGEditor::OnContextualFilteringChanged, Menu.ToWeakPtr(), PCGActionsHelpers::ECompatibilityMask::Filter))
+				.IsChecked(TAttribute<ECheckBoxState>::CreateSP(this, &FPCGEditor::IsContextualFilteringChecked, PCGActionsHelpers::ECompatibilityMask::Filter))
+				.IsEnabled(TAttribute<bool>::CreateSP(this, &FPCGEditor::IsContextualFilteringEnabled, PCGActionsHelpers::ECompatibilityMask::Compatible))
+				.Padding(FMargin(2.0f, 0))
+				[
+					SNew(SImage)
+					.Image(FPCGEditorStyle::Get().GetBrush(PCGEditorStyleConstants::NodeFilter))
+				]
+			]
+			+ SHorizontalBox::Slot()
+			.VAlign(VAlign_Fill)
+			.HAlign(HAlign_Right)
+			.AutoWidth()
+			.Padding(0)
+			[
+				// Accept conversions
+				SNew(SCheckBox)
+				.ToolTipText(LOCTEXT("PCGAcceptConversion", "Show nodes that would apply a conversion"))
+				.Style(&FAppStyle::Get().GetWidgetStyle<FCheckBoxStyle>("ToggleButtonCheckbox"))
+				.OnCheckStateChanged(FOnCheckStateChanged::CreateSP(this, &FPCGEditor::OnContextualFilteringChanged, Menu.ToWeakPtr(), PCGActionsHelpers::ECompatibilityMask::Conversion))
+				.IsChecked(TAttribute<ECheckBoxState>::CreateSP(this, &FPCGEditor::IsContextualFilteringChecked, PCGActionsHelpers::ECompatibilityMask::Conversion))
+				.IsEnabled(TAttribute<bool>::CreateSP(this, &FPCGEditor::IsContextualFilteringEnabled, PCGActionsHelpers::ECompatibilityMask::Compatible))
+				.Padding(FMargin(2.0f, 0))
+				[
+					SNew(SImage)
+					.Image(FPCGEditorStyle::Get().GetBrush(PCGEditorStyleConstants::NodeConvert))
+				]
+			]
+		]
+		+ SVerticalBox::Slot()
+		.Padding(0)
+		.AutoHeight()
+		[
+			Menu
+		];
+
+	return FActionMenuContent(MenuWithOptions, Menu->GetFilterTextBox());
+}
+
+PCGActionsHelpers::ECompatibilityMask FPCGEditor::GetContextualCompatibilityMask() const
+{
+	// Behavior note: there is compatibility at all unless the Compatible flag is enabled, hence why we return NoMask when this happens.
+	return !!(ContextualActionMask & PCGActionsHelpers::ECompatibilityMask::Compatible) ? ContextualActionMask : PCGActionsHelpers::ECompatibilityMask::NoMask;
+}
+
+void FPCGEditor::OnContextualFilteringChanged(ECheckBoxState NewState, const TWeakPtr<SGraphEditorActionMenu> WeakActionMenu, PCGActionsHelpers::ECompatibilityMask AffectedMask)
+{
+	if (NewState == ECheckBoxState::Checked)
+	{
+		ContextualActionMask |= AffectedMask;
+	}
+	else
+	{
+		ContextualActionMask &= ~AffectedMask;
+	}
+
+	if (TSharedPtr<SGraphEditorActionMenu> ActionMenu = WeakActionMenu.Pin())
+	{
+		ActionMenu->RefreshAllActions();
+	}
+}
+
+bool FPCGEditor::IsContextualFilteringEnabled(PCGActionsHelpers::ECompatibilityMask AffectedMask) const
+{
+	return !!(ContextualActionMask & AffectedMask);
+}
+
+ECheckBoxState FPCGEditor::IsContextualFilteringChecked(PCGActionsHelpers::ECompatibilityMask AffectedMask) const
+{
+	return IsContextualFilteringEnabled(AffectedMask) ? ECheckBoxState::Checked : ECheckBoxState::Unchecked;
 }
 
 void FPCGEditor::OnCollapseNodesInSubgraph(bool bCollapseToEmbeddedSubgraph)
@@ -3099,7 +3228,7 @@ TSharedRef<SGraphEditor> FPCGEditor::CreateGraphEditorWidget(TSharedPtr<FTabInfo
 	InEvents.OnTextCommitted = FOnNodeTextCommitted::CreateSP(this, &FPCGEditor::OnNodeTitleCommitted);
 	InEvents.OnNodeDoubleClicked = FSingleNodeEvent::CreateSP(this, &FPCGEditor::OnNodeDoubleClicked);
 	InEvents.OnSpawnNodeByShortcutAtLocation = SGraphEditor::FOnSpawnNodeByShortcutAtLocation::CreateSP(this, &FPCGEditor::OnSpawnNodeByShortcut, InGraph);
-	InEvents.OnCreateActionMenuAtLocation = SGraphEditor::FOnCreateActionMenuAtLocation::CreateSP(this, &FPCGEditor::OnCreateActionMenuContent);
+	InEvents.OnCreateActionMenuAtLocation = SGraphEditor::FOnCreateActionMenuAtLocation::CreateSP(this, &FPCGEditor::OnCreateActionMenu);
 
 	TSharedRef<SGraphEditor> GraphWidget = SNew(SGraphEditor)
 		.AdditionalCommands(GraphEditorCommands)
diff --git a/Engine/Plugins/PCG/Source/PCGEditor/Private/PCGEditorStyle.cpp b/Engine/Plugins/PCG/Source/PCGEditor/Private/PCGEditorStyle.cpp
--- a/Engine/Plugins/PCG/Source/PCGEditor/Private/PCGEditorStyle.cpp
+++ b/Engine/Plugins/PCG/Source/PCGEditor/Private/PCGEditorStyle.cpp
@@ -80,6 +80,9 @@ FPCGEditorStyle::FPCGEditorStyle() : FSlateStyleSet("PCGEditorStyle")
 	Set(PCGNodeConstants::Icons::CompactNodeFilter, new IMAGE_BRUSH_SVG("Style/PCG_Graph_Filter", FVector2D(28.0f, 28.0f)));
 	Set(PCGNodeConstants::Icons::CompactNodeConvert, new IMAGE_BRUSH_SVG("Style/PCG_Graph_To", PinSize));
 
+	Set(PCGEditorStyleConstants::NodeFilter, new IMAGE_BRUSH_SVG("Style/PCG_Graph_Filter", Icon16x16));
+	Set(PCGEditorStyleConstants::NodeConvert, new IMAGE_BRUSH_SVG("Style/PCG_Graph_To", Icon16x16));
+
 	FInlineEditableTextBlockStyle NodeTitleStyle = FInlineEditableTextBlockStyle(FAppStyle::Get().GetWidgetStyle<FInlineEditableTextBlockStyle>("Graph.Node.NodeTitleInlineEditableText"));
 	FTextBlockStyle GraphNodeItalicTitle = FTextBlockStyle(FAppStyle::Get().GetWidgetStyle<FTextBlockStyle>("Graph.Node.NodeTitle"))
 		.SetFont(DEFAULT_FONT("BoldCondensedItalic", 10));
diff --git a/Engine/Plugins/PCG/Source/PCGEditor/Private/PCGEditorStyle.h b/Engine/Plugins/PCG/Source/PCGEditor/Private/PCGEditorStyle.h
--- a/Engine/Plugins/PCG/Source/PCGEditor/Private/PCGEditorStyle.h
+++ b/Engine/Plugins/PCG/Source/PCGEditor/Private/PCGEditorStyle.h
@@ -56,6 +56,9 @@ namespace PCGEditorStyleConstants
 
 	const FName Node_Overlay_GridSizeLabel_Active_Border = TEXT("PCG.Node.Overlay.ThisGridSizeLabel.Active.Border");
 
+	const FName NodeFilter = TEXT("PCG.Graph.NodeFilter");
+	const FName NodeConvert = TEXT("PCG.Graph.NodeConvert");
+
 	const float Node_Overlay_GridSizeLabel_BorderRadius = 11.0f;
 	const float Node_Overlay_GridSizeLabel_BorderStroke = 1.5f;
 }
\ No newline at end of file
diff --git a/Engine/Plugins/PCG/Source/PCGEditor/Private/Schema/PCGEditorGraphSchema.cpp b/Engine/Plugins/PCG/Source/PCGEditor/Private/Schema/PCGEditorGraphSchema.cpp
--- a/Engine/Plugins/PCG/Source/PCGEditor/Private/Schema/PCGEditorGraphSchema.cpp
+++ b/Engine/Plugins/PCG/Source/PCGEditor/Private/Schema/PCGEditorGraphSchema.cpp
@@ -6,6 +6,7 @@
 #include "PCGDataAsset.h"
 #include "PCGEdge.h"
 #include "PCGGraph.h"
+#include "PCGModule.h"
 #include "PCGPin.h"
 #include "PCGSubgraph.h"
 #include "Elements/PCGCollapseElement.h"
@@ -28,6 +29,7 @@
 #include "Nodes/PCGEditorGraphNodeReroute.h"
 #include "Schema/PCGEditorGraphSchemaActions.h"
 
+#include "Algo/AnyOf.h"
 #include "AssetRegistry/AssetData.h"
 #include "AssetRegistry/AssetRegistryModule.h"
 #include "Blueprint/BlueprintSupport.h"
@@ -45,6 +47,87 @@
 
 #define LOCTEXT_NAMESPACE "PCGEditorGraphSchema"
 
+namespace PCGActionsHelpers
+{
+	bool IsValid(const EPCGDataTypeCompatibilityResult Result, const ECompatibilityMask Mask)
+	{
+		const bool bIsValid = PCGDataTypeCompatibilityResult::IsValid(Result);
+
+		if (!bIsValid || Mask == ECompatibilityMask::NoMask)
+		{
+			return bIsValid;
+		}
+		else
+		{
+			return !!(Mask & static_cast<ECompatibilityMask>(1 << std::underlying_type_t<EPCGDataTypeCompatibilityResult>(Result)));
+		}
+	}
+
+	ECompatibilityMask GetPinAndCompatibilityMask(const UEdGraphPin* Pin, const UPCGEditorGraph* PCGEdGraph, const UPCGPin*& OutPin)
+	{
+		const UEdGraphNode* Node = Pin ? Pin->GetOwningNode() : nullptr;
+		const UPCGEditorGraphNodeBase* EditorNode = Cast<UPCGEditorGraphNodeBase>(Node);
+		const UPCGNode* PCGNode = EditorNode ? EditorNode->GetPCGNode() : nullptr;
+		OutPin = (Pin && PCGNode) ? (Pin->Direction == EGPD_Output ? PCGNode->GetOutputPin(Pin->PinName) : PCGNode->GetInputPin(Pin->PinName)) : nullptr;
+
+		const FPCGEditor* PCGEditor = PCGEdGraph ? PCGEdGraph->GetEditor().Pin().Get() : nullptr;
+		return PCGEditor ? PCGEditor->GetContextualCompatibilityMask() : ECompatibilityMask::NoMask;
+	}
+
+	bool HasCompatiblePins(const UPCGPin* Pin, const TArray<FPCGPinProperties>& PinProps, ECompatibilityMask CompatibilityMask)
+	{
+		if (!Pin || CompatibilityMask == ECompatibilityMask::NoMask)
+		{
+			// Not super relevant here because we don't have enough information to filter against.
+			return true;
+		}
+
+		if (Pin->IsOutputPin())
+		{
+			bool bHasRequiredPins = Algo::AnyOf(PinProps, [](const FPCGPinProperties& Props) { return Props.IsRequiredPin(); });
+
+			return Algo::AnyOf(PinProps, [Pin, bHasRequiredPins, CompatibilityMask](const FPCGPinProperties& Props)
+				{
+					return (!Props.IsAdvancedPin() &&
+						(!bHasRequiredPins || Props.IsRequiredPin()) &&
+						IsValid(FPCGModule::GetConstDataTypeRegistry().IsCompatible(Pin->GetCurrentTypesID(), Props.AllowedTypes), CompatibilityMask));
+				});
+		}
+		else
+		{
+			return Algo::AnyOf(PinProps, [Pin, CompatibilityMask](const FPCGPinProperties& Props) { return IsValid(FPCGModule::GetConstDataTypeRegistry().IsCompatible(Props.AllowedTypes, Pin->GetCurrentTypesID()), CompatibilityMask); });
+		}
+	}
+
+	bool IsCompatibleSettings(const UPCGSettings* Settings, const UPCGPin* Pin, ECompatibilityMask CompatibilityMask)
+	{
+		if (!Settings || !Pin || CompatibilityMask == ECompatibilityMask::NoMask)
+		{
+			// Not super relevant here because we either don't have settings (so there's nothing to be compatible with) or don't have enough information to filter against.
+			return true;
+		}
+
+		return HasCompatiblePins(Pin, Pin->IsOutputPin() ? Settings->InputPinProperties() : Settings->OutputPinProperties(), CompatibilityMask);
+	}
+
+	bool IsCompatibleType(const UPCGPin* Pin, FPCGDataTypeIdentifier AllowedTypes, ECompatibilityMask CompatibilityMask)
+	{
+		if (!Pin)
+		{
+			return true;
+		}
+
+		if (Pin->IsOutputPin())
+		{
+			return IsValid(FPCGModule::GetConstDataTypeRegistry().IsCompatible(Pin->GetCurrentTypesID(), { AllowedTypes }), CompatibilityMask);
+		}
+		else
+		{
+			return IsValid(FPCGModule::GetConstDataTypeRegistry().IsCompatible({ AllowedTypes }, Pin->GetCurrentTypesID()), CompatibilityMask);
+		}
+	}
+}
+
 FPCGActionsFilter::FPCGActionsFilter(const UEdGraph* InEdGraph, EPCGElementType InElementFilterType, FPCGGraphEditorCustomization InCustomization)
 	: FilterType(InElementFilterType)
 	, Graph(Cast<UPCGEditorGraph>(InEdGraph))
@@ -410,6 +493,10 @@ void UPCGEditorGraphSchema::GetNativeElementActions(FGraphActionMenuBuilder& Act
 		}
 	}
 
+	const UEdGraphPin* Pin = ActionMenuBuilder.FromPin;
+	const UPCGPin* PCGPin = nullptr;
+	PCGActionsHelpers::ECompatibilityMask CompatibilityMask = PCGActionsHelpers::GetPinAndCompatibilityMask(Pin, ActionFilter.Graph, PCGPin);
+
 	for (UClass* SettingsClass : SettingsClasses)
 	{
 		if (const UPCGSettings* PCGSettings = SettingsClass->GetDefaultObject<UPCGSettings>())
@@ -425,6 +512,12 @@ void UPCGEditorGraphSchema::GetNativeElementActions(FGraphActionMenuBuilder& Act
 					continue;
 				}
 
+				//@todo_pcg: we might want to treat different preconfigured settings as having different compatibility states
+				if(!PCGActionsHelpers::IsCompatibleSettings(PCGSettings, PCGPin, CompatibilityMask))
+				{
+					continue;
+				}
+
 				TArray<FPCGPreConfiguredSettingsInfo> AllPreconfiguredInfo = PCGSettings->GetPreconfiguredInfo();
 
 				if (AllPreconfiguredInfo.IsEmpty() || !PCGSettings->OnlyExposePreconfiguredSettings())
@@ -463,7 +556,11 @@ void UPCGEditorGraphSchema::GetNativeElementActions(FGraphActionMenuBuilder& Act
 
 	const FText UserParameterCategory = LOCTEXT("UserParametersCategoryName", "Graph Parameters");
 
-	if (ActionFilter.Accepts(UserParameterCategory))
+	// @todo_pcg review usage of the compatible call also for specific subtypes, but might need more info from the user parameters
+	// Add getters only when it's an input pin (or no pin), and if we have filtering, don't show them if the pin isn't compatible with attribute sets.
+	if ((!Pin || Pin->Direction == EGPD_Input) && 
+		PCGActionsHelpers::IsCompatibleType(PCGPin, EPCGDataType::Param, CompatibilityMask) &&
+		ActionFilter.Accepts(UserParameterCategory))
 	{
 		const UPCGEditorGraph* Graph = ActionFilter.Graph;
 		const UPCGGraph* PCGGraph = Graph ? const_cast<UPCGEditorGraph*>(Graph)->GetPCGGraph() : nullptr;
@@ -722,7 +819,17 @@ void UPCGEditorGraphSchema::GetNamedRerouteUsageActions(FGraphActionMenuBuilder&
 
 void UPCGEditorGraphSchema::GetDataAssetActions(FGraphActionMenuBuilder& ActionMenuBuilder, const FPCGActionsFilter& ActionFilter) const
 {
-	PCGEditorUtils::ForEachPCGAssetData([&ActionMenuBuilder, &ActionFilter](const FAssetData& AssetData)
+	// Early out if the 'from pin' is an output pin - it's not common to pass data to a load asset node.
+	// @todo_pcg: it is possible to have the input be a param data for the 'from input' mode so we probably should consider this
+	if (ActionMenuBuilder.FromPin && ActionMenuBuilder.FromPin->Direction == EGPD_Output)
+	{
+		return;
+	}
+
+	const UPCGPin* PCGPin = nullptr;
+	PCGActionsHelpers::ECompatibilityMask CompatibilityMask = PCGActionsHelpers::GetPinAndCompatibilityMask(ActionMenuBuilder.FromPin, ActionFilter.Graph, PCGPin);
+
+	PCGEditorUtils::ForEachPCGAssetData([&ActionMenuBuilder, &ActionFilter, PCGPin, CompatibilityMask](const FAssetData& AssetData)
 	{
 		const bool bExposeToLibrary = AssetData.GetTagValueRef<bool>(GET_MEMBER_NAME_CHECKED(UPCGDataAsset, bExposeToLibrary));
 		if (bExposeToLibrary)
@@ -737,6 +844,16 @@ void UPCGEditorGraphSchema::GetDataAssetActions(FGraphActionMenuBuilder& ActionM
 				return true;
 			}
 
+			// @todo_pcg : If there is a pin, check against the cached pins, if they were properly set
+			if (CompatibilityMask != PCGActionsHelpers::ECompatibilityMask::NoMask && PCGPin)
+			{
+				PCGDataAsset::FGetAssetDataOutput AssetDataOutput = UPCGDataAsset::GetAssetRegistryData(AssetData);
+				if (AssetDataOutput.CachedPins.bHasCachedPins && !PCGActionsHelpers::HasCompatiblePins(PCGPin, AssetDataOutput.CachedPins.CachedPins, CompatibilityMask))
+				{
+					return true;
+				}
+			}
+
 			const FSoftClassPath SettingsClassPath = FSoftClassPath(AssetData.GetTagValueRef<FString>(GET_MEMBER_NAME_CHECKED(UPCGDataAsset, SettingsClass)));
 			TSubclassOf<UPCGSettings> SettingsClass = SettingsClassPath.TryLoadClass<UPCGSettings>();
 
@@ -1312,4 +1429,4 @@ void FPCGEditorConnectionDrawingPolicy::DetermineWiringStyle(UEdGraphPin* Output
 	}
 }
 
-#undef LOCTEXT_NAMESPACE
+#undef LOCTEXT_NAMESPACE
\ No newline at end of file
diff --git a/Engine/Plugins/PCG/Source/PCGEditor/Public/PCGEditor.h b/Engine/Plugins/PCG/Source/PCGEditor/Public/PCGEditor.h
--- a/Engine/Plugins/PCG/Source/PCGEditor/Public/PCGEditor.h
+++ b/Engine/Plugins/PCG/Source/PCGEditor/Public/PCGEditor.h
@@ -252,7 +252,10 @@ PRAGMA_ENABLE_DEPRECATION_WARNINGS
 	FReply OnSpawnNodeByShortcut(FInputChord InChord, const FVector2D& InPosition, UPCGEditorGraph* InGraph);
 	PCGEDITOR_API FReply OnSpawnNodeByShortcut(FInputChord InChord, const FVector2f& InPosition, UPCGEditorGraph* InGraph);
 
-	FActionMenuContent OnCreateActionMenuContent(UEdGraph* InGraph, const FVector2f& Location, const TArray<UEdGraphPin*>& InDraggedPins, bool bAutoExpand, SGraphEditor::FActionMenuClosed OnMenuClosed);
+	FActionMenuContent OnCreateActionMenu(UEdGraph* InGraph, const FVector2f& Location, const TArray<UEdGraphPin*>& InDraggedPins, bool bAutoExpand, SGraphEditor::FActionMenuClosed OnMenuClosed);
+
+	/** Returns whether contextual menu filtering should be applied. */
+	PCGActionsHelpers::ECompatibilityMask GetContextualCompatibilityMask() const;
 
 	/** Can determinism be tested on the current graph */
 	bool CanRunDeterminismGraphTest() const;
@@ -392,6 +395,13 @@ private:
 	/** Whether selected nodes are being debugged or not */
 	ECheckBoxState GetDebugCheckState() const;
 
+	/** Called when the contextual menu filtering toggle is changed. */
+	void OnContextualFilteringChanged(ECheckBoxState NewState, const TWeakPtr<SGraphEditorActionMenu> WeakActionMenu, PCGActionsHelpers::ECompatibilityMask AffectedMask);
+
+	/** Whether the contextual menu filtering is enabled or not. */
+	bool IsContextualFilteringEnabled(PCGActionsHelpers::ECompatibilityMask AffectedMask) const;
+	ECheckBoxState IsContextualFilteringChecked(PCGActionsHelpers::ECompatibilityMask AffectedMask) const;
+
 	/** Enable node debug state for selected nodes and disable for others */
 	void OnDebugOnlySelected();
 
@@ -636,4 +646,7 @@ private:
 
 	/** Factory that spawns graph editors; used to look up all tabs spawned by it. */
 	TWeakPtr<FDocumentTabFactory> DocumentTabFactory;
+
+	/** Contextual menu filter toggle state variable. Starts as compatible only to show more likely choices. */
+	PCGActionsHelpers::ECompatibilityMask ContextualActionMask = PCGActionsHelpers::ECompatibilityMask::Compatible;
 };
diff --git a/Engine/Plugins/PCG/Source/PCGEditor/Public/Schema/PCGEditorGraphSchema.h b/Engine/Plugins/PCG/Source/PCGEditor/Public/Schema/PCGEditorGraphSchema.h
--- a/Engine/Plugins/PCG/Source/PCGEditor/Public/Schema/PCGEditorGraphSchema.h
+++ b/Engine/Plugins/PCG/Source/PCGEditor/Public/Schema/PCGEditorGraphSchema.h
@@ -10,6 +10,8 @@
 #include "ConnectionDrawingPolicy.h"
 #include "StructUtils/PropertyBag.h"
 
+#include "Data/Registry/PCGDataTypeCommon.h"
+
 #include "PCGEditorGraphSchema.generated.h"
 
 class UPCGBaseSubgraphSettings;
@@ -18,6 +20,20 @@ enum class EPCGElementType : uint8;
 class IAssetReferenceFilter;
 class UPCGEditorGraph;
 
+namespace PCGActionsHelpers
+{
+	enum class ECompatibilityMask
+	{
+		NoMask = 0,
+		Compatible = 1 << std::underlying_type_t<EPCGDataTypeCompatibilityResult>(EPCGDataTypeCompatibilityResult::Compatible),
+		Filter = 1 << std::underlying_type_t<EPCGDataTypeCompatibilityResult>(EPCGDataTypeCompatibilityResult::RequireFilter),
+		Conversion = 1 << std::underlying_type_t<EPCGDataTypeCompatibilityResult>(EPCGDataTypeCompatibilityResult::RequireConversion),
+		// note: SubtypeNotCompatible isn't used yet, but we'll keep it here for completeness
+		SubtypeNotCompatible = 1 << std::underlying_type_t<EPCGDataTypeCompatibilityResult>(EPCGDataTypeCompatibilityResult::TypeCompatibleSubtypeNotCompatible)
+	};
+	ENUM_CLASS_FLAGS(ECompatibilityMask);
+}
+
 struct FPCGActionsFilter
 {
 	FPCGActionsFilter(const UEdGraph* InEdGraph, EPCGElementType InElementFilterType, FPCGGraphEditorCustomization InCustomization);
