[PCG] Generalized default actor source to world object source for more flexibility.

#Synced-CL 51473272

diff --git a/Engine/Plugins/PCG/Source/PCG/Private/PCGDefaultActorExecutionSource.cpp b/Engine/Plugins/PCG/Source/PCG/Private/PCGDefaultActorExecutionSource.cpp
--- a/Engine/Plugins/PCG/Source/PCG/Private/PCGDefaultActorExecutionSource.cpp
+++ b/Engine/Plugins/PCG/Source/PCG/Private/PCGDefaultActorExecutionSource.cpp
@@ -11,18 +11,25 @@
 #include UE_INLINE_GENERATED_CPP_BY_NAME(PCGDefaultActorExecutionSource)
 
 FPCGDefaultActorExecutionState::FPCGDefaultActorExecutionState(UPCGDefaultActorExecutionSource* InSource)
-	: FPCGDefaultExecutionState(InSource)
+	: FPCGDefaultWorldObjectExecutionState(InSource)
 {
 }
 
 FString FPCGDefaultActorExecutionState::GetDebugName() const
 {
-	return FString(TEXT("PCGDefaultActorExecutionSource - "))
+	if (AActor* Actor = GetActor())
+	{
+		return FString(TEXT("PCGDefaultActorExecutionSource - "))
 #if WITH_EDITOR
-		+ FString(GetActor() ? GetActor()->GetActorLabel() : TEXT("Unknown actor"));
+			+ Actor->GetActorLabel();
 #else
-		+ FString(GetActor() ? GetActor()->GetName() : TEXT("Unknown actor"));
+			+ Actor->GetName();
 #endif
+	}
+	else
+	{
+		return FPCGDefaultWorldObjectExecutionState::GetDebugName();
+	}
 }
 
 UPCGData* FPCGDefaultActorExecutionState::GetSelfData() const
@@ -36,7 +43,7 @@ UPCGData* FPCGDefaultActorExecutionState::GetSelfData() const
 		}
 	}
 
-	return FPCGDefaultExecutionState::GetSelfData();
+	return FPCGDefaultWorldObjectExecutionState::GetSelfData();
 }
 
 void FPCGDefaultActorExecutionState::ExecutePreGraph(FPCGContext* Context)
@@ -73,7 +80,7 @@ void FPCGDefaultActorExecutionState::ExecutePreGraph(FPCGContext* Context)
 			{
 				SelfData = const_cast<UPCGData*>(Output.Collection.TaggedData[0].Data.Get());
 			}
-			
+
 			PCGSubsystem->SetSelfPCGData(GetTypedSource()->CurrentGenerationTask, SelfData);
 		}
 	}
@@ -87,7 +94,7 @@ FTransform FPCGDefaultActorExecutionState::GetTransform() const
 	}
 	else
 	{
-		return FPCGDefaultExecutionState::GetTransform();
+		return FPCGDefaultWorldObjectExecutionState::GetTransform();
 	}
 }
 
@@ -96,18 +103,6 @@ UObject* FPCGDefaultActorExecutionState::GetTarget() const
 	return GetActor();
 }
 
-UWorld* FPCGDefaultActorExecutionState::GetWorld() const
-{
-	if (AActor* Actor = GetActor())
-	{
-		return Actor->GetWorld();
-	}
-	else
-	{
-		return FPCGDefaultExecutionState::GetWorld();
-	}
-}
-
 bool FPCGDefaultActorExecutionState::HasAuthority() const
 {
 	if (AActor* Actor = GetActor())
@@ -116,7 +111,7 @@ bool FPCGDefaultActorExecutionState::HasAuthority() const
 	}
 	else
 	{
-		return FPCGDefaultExecutionState::HasAuthority();
+		return FPCGDefaultWorldObjectExecutionState::HasAuthority();
 	}
 }
 
@@ -128,7 +123,7 @@ FBox FPCGDefaultActorExecutionState::GetBounds() const
 	}
 	else
 	{
-		return FPCGDefaultExecutionState::GetBounds();
+		return FPCGDefaultWorldObjectExecutionState::GetBounds();
 	}
 }
 
@@ -140,7 +135,7 @@ FBox FPCGDefaultActorExecutionState::GetLocalSpaceBounds() const
 	}
 	else
 	{
-		return FPCGDefaultExecutionState::GetLocalSpaceBounds();
+		return FPCGDefaultWorldObjectExecutionState::GetLocalSpaceBounds();
 	}
 }
 
@@ -152,16 +147,16 @@ UPCGDefaultActorExecutionSource::UPCGDefaultActorExecutionSource()
 
 void UPCGDefaultActorExecutionSource::Initialize(const FPCGDefaultActorExecutionSourceParams& InParams)
 {
-	Super::Initialize(InParams);
+	UPCGDefaultExecutionSource::Initialize(InParams);
 	SetActor(InParams.Actor);
 }
 
 void UPCGDefaultActorExecutionSource::SetActor(AActor* InActor)
 {
-	Actor = InActor;
+	SetWorldObject(InActor);
 }
 
 AActor* UPCGDefaultActorExecutionSource::GetActor()
 {
-	return Actor.Get();
+	return Cast<AActor>(GetWorldObject());
 }
\ No newline at end of file
diff --git a/Engine/Plugins/PCG/Source/PCG/Private/PCGDefaultWorldObjectExecutionSource.cpp b/Engine/Plugins/PCG/Source/PCG/Private/PCGDefaultWorldObjectExecutionSource.cpp
new file mode 100644
--- /dev/null
+++ b/Engine/Plugins/PCG/Source/PCG/Private/PCGDefaultWorldObjectExecutionSource.cpp
@@ -0,0 +1,39 @@
+// Copyright Epic Games, Inc. All Rights Reserved.
+
+#include "PCGDefaultWorldObjectExecutionSource.h"
+
+#include UE_INLINE_GENERATED_CPP_BY_NAME(PCGDefaultWorldObjectExecutionSource)
+
+FPCGDefaultWorldObjectExecutionState::FPCGDefaultWorldObjectExecutionState(UPCGDefaultWorldObjectExecutionSource* InSource)
+	: FPCGDefaultExecutionState(InSource)
+{
+}
+
+FString FPCGDefaultWorldObjectExecutionState::GetDebugName() const
+{
+	return FString(TEXT("PCGDefaultWorldObjectExecutionSource - ")) + (GetWorldObject() ? GetWorldObject()->GetName() : TEXT("Unknown object"));
+}
+
+UWorld* FPCGDefaultWorldObjectExecutionState::GetWorld() const
+{
+	if (UObject* Object = GetWorldObject())
+	{
+		return Object->GetWorld();
+	}
+	else
+	{
+		return FPCGDefaultExecutionState::GetWorld();
+	}
+}
+
+UPCGDefaultWorldObjectExecutionSource::UPCGDefaultWorldObjectExecutionSource()
+{
+	// Inefficient but more foolproof - we'll replace the previously created state from the base class by our version.
+	State = MakeUnique<FPCGDefaultWorldObjectExecutionState>(this);
+}
+
+void UPCGDefaultWorldObjectExecutionSource::Initialize(const FPCGDefaultWorldObjectExecutionSourceParams& InParams)
+{
+	UPCGDefaultExecutionSource::Initialize(InParams);
+	SetWorldObject(InParams.WorldObject);
+}
\ No newline at end of file
diff --git a/Engine/Plugins/PCG/Source/PCG/Public/PCGDefaultActorExecutionSource.h b/Engine/Plugins/PCG/Source/PCG/Public/PCGDefaultActorExecutionSource.h
--- a/Engine/Plugins/PCG/Source/PCG/Public/PCGDefaultActorExecutionSource.h
+++ b/Engine/Plugins/PCG/Source/PCG/Public/PCGDefaultActorExecutionSource.h
@@ -2,11 +2,11 @@
 
 #pragma once
 
-#include "PCGDefaultExecutionSource.h"
-#include "PCGManagedResourceContainer.h"
+#include "PCGDefaultWorldObjectExecutionSource.h"
 
 #include "PCGDefaultActorExecutionSource.generated.h"
 
+class AActor;
 class FPCGDefaultActorExecutionState;
 
 struct FPCGDefaultActorExecutionSourceParams : public FPCGDefaultExecutionSourceParams
@@ -15,7 +15,7 @@ struct FPCGDefaultActorExecutionSourceParams : public FPCGDefaultExecutionSource
 };
 
 UCLASS(MinimalAPI)
-class UPCGDefaultActorExecutionSource : public UPCGDefaultExecutionSource
+class UPCGDefaultActorExecutionSource : public UPCGDefaultWorldObjectExecutionSource
 {
 	friend FPCGDefaultActorExecutionState;
 
@@ -26,12 +26,9 @@ public:
 	PCG_API void Initialize(const FPCGDefaultActorExecutionSourceParams& InParams);
 	PCG_API AActor* GetActor();
 	PCG_API void SetActor(AActor* InActor);
-
-private:
-	TWeakObjectPtr<AActor> Actor = nullptr;
 };
 
-class FPCGDefaultActorExecutionState : public FPCGDefaultExecutionState
+class FPCGDefaultActorExecutionState : public FPCGDefaultWorldObjectExecutionState
 {
 public:
 	FPCGDefaultActorExecutionState() = default;
@@ -41,13 +38,10 @@ public:
 	PCG_API virtual UPCGData* GetSelfData() const override;
 	PCG_API virtual FTransform GetTransform() const override;
 	PCG_API virtual UObject* GetTarget() const override;
-	PCG_API virtual UWorld* GetWorld() const override;
 	PCG_API virtual bool HasAuthority() const override;
 	PCG_API virtual FBox GetBounds() const override;
 	PCG_API virtual FBox GetLocalSpaceBounds() const override;
 
-	PCG_API virtual FPCGManagedResourceContainer* GetManagedResourceContainer() override { return &ManagedResourceContainer; }
-	PCG_API virtual FTransactionallySafeCriticalSection* GetManagedResourceContainerLock() override { return &ManagedResourcesLock; }
 	PCG_API virtual void ExecutePreGraph(FPCGContext* InContext) override;
 
 protected:
@@ -55,7 +49,4 @@ protected:
 
 private:
 	AActor* GetActor() const { return GetTypedSource()->GetActor(); }
-
-	FPCGManagedResourceContainer ManagedResourceContainer;
-	mutable FTransactionallySafeCriticalSection ManagedResourcesLock;
 };
\ No newline at end of file
diff --git a/Engine/Plugins/PCG/Source/PCG/Public/PCGDefaultWorldObjectExecutionSource.h b/Engine/Plugins/PCG/Source/PCG/Public/PCGDefaultWorldObjectExecutionSource.h
new file mode 100644
--- /dev/null
+++ b/Engine/Plugins/PCG/Source/PCG/Public/PCGDefaultWorldObjectExecutionSource.h
@@ -0,0 +1,52 @@
+// Copyright Epic Games, Inc. All Rights Reserved.
+
+#pragma once
+
+#include "PCGDefaultExecutionSource.h"
+#include "PCGManagedResourceContainer.h"
+
+#include "PCGDefaultWorldObjectExecutionSource.generated.h"
+
+class FPCGDefaultWorldObjectExecutionState;
+
+struct FPCGDefaultWorldObjectExecutionSourceParams : public FPCGDefaultExecutionSourceParams
+{
+	UObject* WorldObject = nullptr;
+};
+
+UCLASS(MinimalAPI)
+class UPCGDefaultWorldObjectExecutionSource : public UPCGDefaultExecutionSource
+{
+	friend FPCGDefaultWorldObjectExecutionState;
+
+	GENERATED_BODY()
+	
+public:
+	PCG_API UPCGDefaultWorldObjectExecutionSource();
+	PCG_API void Initialize(const FPCGDefaultWorldObjectExecutionSourceParams& InParams);
+	UObject* GetWorldObject() { return WorldObject.Get(); }
+	void SetWorldObject(UObject* InWorldObject) { WorldObject = InWorldObject; }
+
+private:
+	TWeakObjectPtr<UObject> WorldObject = nullptr;
+};
+
+class FPCGDefaultWorldObjectExecutionState : public FPCGDefaultExecutionState
+{
+public:
+	FPCGDefaultWorldObjectExecutionState() = default;
+	explicit FPCGDefaultWorldObjectExecutionState(UPCGDefaultWorldObjectExecutionSource* InSource);
+
+	PCG_API virtual FString GetDebugName() const override;
+	PCG_API virtual UWorld* GetWorld() const override;
+	virtual FPCGManagedResourceContainer* GetManagedResourceContainer() override { return &ManagedResourceContainer; }
+	virtual FTransactionallySafeCriticalSection* GetManagedResourceContainerLock() override { return &ManagedResourcesLock; }
+
+protected:
+	UPCGDefaultWorldObjectExecutionSource* GetTypedSource() const { return CastChecked<UPCGDefaultWorldObjectExecutionSource>(GetSource()); }
+	UObject* GetWorldObject() const { return GetTypedSource()->GetWorldObject(); }
+
+private:
+	FPCGManagedResourceContainer ManagedResourceContainer;
+	mutable FTransactionallySafeCriticalSection ManagedResourcesLock;
+};
\ No newline at end of file
