[PCG] Fix landscape data interface missing landscape components at negative tile indices

#rnx
#Synced-CL 51639529

diff --git a/Engine/Plugins/PCG/Source/PCG/Private/Compute/DataInterfaces/PCGLandscapeDataInterface.cpp b/Engine/Plugins/PCG/Source/PCG/Private/Compute/DataInterfaces/PCGLandscapeDataInterface.cpp
--- a/Engine/Plugins/PCG/Source/PCG/Private/Compute/DataInterfaces/PCGLandscapeDataInterface.cpp
+++ b/Engine/Plugins/PCG/Source/PCG/Private/Compute/DataInterfaces/PCGLandscapeDataInterface.cpp
@@ -482,16 +482,16 @@ void UPCGLandscapeDataProvider::InitializeFromLandscapeCollision(UPCGDataBinding
 		SystemMinInLandscape.ComponentMin(SystemMaxInLandscape),
 		SystemMinInLandscape.ComponentMax(SystemMaxInLandscape));
 
-	// Transform the above box into a range of integers covering the cells of the landscape first clamp it at 0.
-	SystemBoundsInLandscape.Min = SystemBoundsInLandscape.Min.ComponentMax(FVector::ZeroVector);
-	SystemBoundsInLandscape.Max = SystemBoundsInLandscape.Max.ComponentMax(FVector::ZeroVector);
-
-	// Next rescale based on the quad size.
+	// Rescale based on the quad size.
 	const double QuadSizeScaleFactor = 1.0 / ((double)InLandscape->ComponentSizeQuads);
 	SystemBoundsInLandscape.Min *= QuadSizeScaleFactor;
 	SystemBoundsInLandscape.Max *= QuadSizeScaleFactor;
 
-	// Truncate to integers.
+	// Clamp between Min/Max int32.
+	const FVector MinIntValue = FVector(double(TNumericLimits<int32>::Min()));
+	SystemBoundsInLandscape.Min = SystemBoundsInLandscape.Min.ComponentMax(MinIntValue);
+	SystemBoundsInLandscape.Max = SystemBoundsInLandscape.Max.ComponentMax(MinIntValue);
+
 	const FVector MaxIntValue = FVector(double(TNumericLimits<int32>::Max()));
 	SystemBoundsInLandscape.Min = SystemBoundsInLandscape.Min.ComponentMin(MaxIntValue);
 	SystemBoundsInLandscape.Max = SystemBoundsInLandscape.Max.ComponentMin(MaxIntValue);
