[PCG] Fix crash when cancelling save dialog of Save Texture to Asset node

#rnx
#Synced-CL 51639529

diff --git a/Engine/Plugins/PCG/Source/PCG/Private/Elements/PCGSaveTextureToAsset.cpp b/Engine/Plugins/PCG/Source/PCG/Private/Elements/PCGSaveTextureToAsset.cpp
--- a/Engine/Plugins/PCG/Source/PCG/Private/Elements/PCGSaveTextureToAsset.cpp
+++ b/Engine/Plugins/PCG/Source/PCG/Private/Elements/PCGSaveTextureToAsset.cpp
@@ -124,11 +124,29 @@ bool FPCGSaveTextureToAssetElement::ExecuteInternal(FPCGContext* InContext) cons
 			},
 			Context);
 
-		// Initialize exported texture with the readback data.
-		ExportedTexture->Source.Init(Context->ReadbackWidth, Context->ReadbackHeight, /*NewNumSlices=*/1, /*NewNumMips=*/1, TSF_BGRA8, Context->RawReadbackData);
-		ExportedTexture->UpdateResource();
+		if (ExportedTexture)
+		{
+			// Initialize exported texture with the readback data.
+			ExportedTexture->Source.Init(Context->ReadbackWidth, Context->ReadbackHeight, /*NewNumSlices=*/1, /*NewNumMips=*/1, TSF_BGRA8, Context->RawReadbackData);
+			ExportedTexture->UpdateResource();
+
+			Context->ExportedTexture = ExportedTexture;
+		}
+		else
+		{
+			if (Settings->ExporterParams.AssetPath.IsEmpty() || Settings->ExporterParams.AssetName.IsEmpty())
+			{
+				PCGLog::LogErrorOnGraph(LOCTEXT("CreateAssetFailedNoParams", "Failed to create texture asset. Invalid asset path or name was provided."), InContext);
+			}
+			else
+			{
+				PCGLog::LogErrorOnGraph(FText::Format(LOCTEXT("CreateAssetFailedWithParams", "Failed to create texture asset at path '{0}' with name '{1}'."),
+					FText::FromString(Settings->ExporterParams.AssetPath),
+					FText::FromString(Settings->ExporterParams.AssetName)), InContext);
+			}
 
-		Context->ExportedTexture = ExportedTexture;
+			return true;
+		}
 	}
 
 	if (Context->ExportedTexture->HasPendingInitOrStreaming())
