[PCG] Added callbacks for post-process & on abort on default execution sources.

#rnx
#Synced-CL 51473272

diff --git a/Engine/Plugins/PCG/Source/PCG/Private/PCGDefaultExecutionSource.cpp b/Engine/Plugins/PCG/Source/PCG/Private/PCGDefaultExecutionSource.cpp
--- a/Engine/Plugins/PCG/Source/PCG/Private/PCGDefaultExecutionSource.cpp
+++ b/Engine/Plugins/PCG/Source/PCG/Private/PCGDefaultExecutionSource.cpp
@@ -117,6 +117,11 @@ FPCGGraphExecutionInspection& FPCGDefaultExecutionState::GetInspection()
 }
 #endif // WITH_EDITOR
 
+void FPCGDefaultExecutionState::OnPostProcess(const FPCGDataCollection& InDataCollection)
+{
+	PostProcessCallback.ExecuteIfBound(InDataCollection);
+}
+
 void UPCGDefaultExecutionSource::SetGraphInterface(UPCGGraphInterface* InGraphInterface)
 {
 	if (InGraphInterface == GraphInterface)
@@ -207,6 +212,8 @@ void UPCGDefaultExecutionSource::Generate()
 						return true;
 					}
 
+					This->State->OnPostProcess(InContext->InputData);
+
 					This->CurrentGenerationTask = InvalidPCGTaskId;
 					if (WeakSubsystem.IsValid())
 					{
@@ -218,8 +225,9 @@ void UPCGDefaultExecutionSource::Generate()
 					return true;
 				},
 				/*InExecutionSource=*/this,
-				/*InExecutionDependencies=*/{GenerationTaskId},
-				/*InDataDependencies=*/{},
+				/*InExecutionDependencies=*/{},
+				// Implementation note: here we're taking the generation task as a data dependency to have data passed in the callback, because we'll need it for the post process.
+				/*InDataDependencies=*/{GenerationTaskId},
 				/*bSupportBasePointDataInput=*/true
 			};
 			
@@ -259,6 +267,7 @@ void UPCGDefaultExecutionSource::Initialize(const FPCGDefaultExecutionSourcePara
 {
 	SetGraphInterface(InParams.GraphInterface);
 	SetSeed(InParams.Seed);
+	State->SetPostProcessCallback(InParams.PostProcessCallback);
 }
 
 void UPCGDefaultExecutionSource::SetSeed(int32 InSeed)
diff --git a/Engine/Plugins/PCG/Source/PCG/Public/PCGDefaultExecutionSource.h b/Engine/Plugins/PCG/Source/PCG/Public/PCGDefaultExecutionSource.h
--- a/Engine/Plugins/PCG/Source/PCG/Public/PCGDefaultExecutionSource.h
+++ b/Engine/Plugins/PCG/Source/PCG/Public/PCGDefaultExecutionSource.h
@@ -12,6 +12,9 @@ class UPCGDefaultExecutionSource;
 class UPCGGraph;
 class UPCGGraphInstance;
 class UPCGGraphInterface;
+struct FPCGDataCollection;
+
+DECLARE_DELEGATE_OneParam(FPCGGenerationPostProcessCallback, const FPCGDataCollection&);
 
 class FPCGDefaultExecutionState : public IPCGGraphExecutionState
 {
@@ -33,6 +36,9 @@ public:
 	PCG_API virtual bool IsGenerating() const override;
 	PCG_API virtual IPCGGraphExecutionSource* GetOriginalSource() const override;
 
+	void SetPostProcessCallback(FPCGGenerationPostProcessCallback InCallback) { PostProcessCallback = MoveTemp(InCallback); }
+	PCG_API void OnPostProcess(const FPCGDataCollection& InDataCollection);
+
 #if WITH_EDITOR
 	PCG_API virtual const PCGUtils::FExtraCapture& GetExtraCapture() const override;
 	PCG_API virtual PCGUtils::FExtraCapture& GetExtraCapture() override;
@@ -54,6 +60,7 @@ protected:
 
 private:
 	UPCGDefaultExecutionSource* Source = nullptr;
+	FPCGGenerationPostProcessCallback PostProcessCallback;
 };
 
 struct FPCGDefaultExecutionSourceParams
@@ -64,6 +71,8 @@ struct FPCGDefaultExecutionSourceParams
 
 	/** Optional callback when generation completes. */
 	FPCGOnEditorGenerationDone GenerationCallback;
+	/** Optional callback to process graph execution results, before the Generation Done callback. */
+	FPCGGenerationPostProcessCallback PostProcessCallback;
 };
 
 UCLASS(MinimalAPI)
