[PCG] Fixed issue when using clear pcg link for resources on created actors

#Synced-CL 51473272

diff --git a/Engine/Plugins/PCG/Source/PCG/Private/PCGManagedResource.cpp b/Engine/Plugins/PCG/Source/PCG/Private/PCGManagedResource.cpp
--- a/Engine/Plugins/PCG/Source/PCG/Private/PCGManagedResource.cpp
+++ b/Engine/Plugins/PCG/Source/PCG/Private/PCGManagedResource.cpp
@@ -760,40 +760,39 @@ bool UPCGManagedComponentBase::MoveResource(const FPCGMoveResourceParams& Params
 		TObjectPtr<AActor> OldOwner = GeneratedComponent->GetOwner();
 		check(OldOwner);
 
-		// Prevent moving of components on external (or spawned) actors
-		if (Params.ExpectedPreviousOwner && OldOwner != Params.ExpectedPreviousOwner)
+		// Prevent moving of components on external (or spawned) actors, but still transfer resource
+		if (!Params.ExpectedPreviousOwner || OldOwner == Params.ExpectedPreviousOwner)
 		{
-			continue;
-		}
-
-		bool bDetached = false;
-		bool bAttached = false;
+			bool bDetached = false;
+			bool bAttached = false;
 
-		GeneratedComponent->UnregisterComponent();
+			GeneratedComponent->UnregisterComponent();
 
-		// Need to change owner first to avoid that the PCG Component will react to this component changes.
-		GeneratedComponent->Rename(nullptr, TargetActor);
+			// Need to change owner first to avoid that the PCG Component will react to this component changes.
+			GeneratedComponent->Rename(nullptr, TargetActor);
 
-		// Check if it is a scene component, and if so, use its method to attach/detach to root component
-		if (TObjectPtr<USceneComponent> GeneratedSceneComponent = Cast<USceneComponent>(GeneratedComponent.Get()))
-		{
-			GeneratedSceneComponent->DetachFromComponent(FDetachmentTransformRules::KeepWorldTransform);
-			bDetached = true;
-			bAttached = GeneratedSceneComponent->AttachToComponent(TargetActor->GetRootComponent(), FAttachmentTransformRules::KeepWorldTransform);
-		}
+			// Check if it is a scene component, and if so, use its method to attach/detach to root component
+			if (TObjectPtr<USceneComponent> GeneratedSceneComponent = Cast<USceneComponent>(GeneratedComponent.Get()))
+			{
+				GeneratedSceneComponent->DetachFromComponent(FDetachmentTransformRules::KeepWorldTransform);
+				bDetached = true;
+				bAttached = GeneratedSceneComponent->AttachToComponent(TargetActor->GetRootComponent(), FAttachmentTransformRules::KeepWorldTransform);
+			}
 
-		// Otherwise use the default one.
-		if (!bAttached)
-		{
-			if (!bDetached)
+			// Otherwise use the default one.
+			if (!bAttached)
 			{
-				OldOwner->RemoveInstanceComponent(GeneratedComponent.Get());
+				if (!bDetached)
+				{
+					OldOwner->RemoveInstanceComponent(GeneratedComponent.Get());
+				}
+
+				TargetActor->AddInstanceComponent(GeneratedComponent.Get());
 			}
 
-			TargetActor->AddInstanceComponent(GeneratedComponent.Get());
+			GeneratedComponent->RegisterComponent();
 		}
 
-		GeneratedComponent->RegisterComponent();
 		ForgetComponent(ComponentIndex);
 		bMovedResources = true;
 	}
