[PCG] Add node to get editor selection (actors).

#Synced-CL 51473272

diff --git a/Engine/Plugins/PCG/Source/PCG/Internal/Elements/Editor/PCGGetEditorSelection.h b/Engine/Plugins/PCG/Source/PCG/Internal/Elements/Editor/PCGGetEditorSelection.h
new file mode 100644
--- /dev/null
+++ b/Engine/Plugins/PCG/Source/PCG/Internal/Elements/Editor/PCGGetEditorSelection.h
@@ -0,0 +1,52 @@
+// Copyright Epic Games, Inc. All Rights Reserved.
+
+#pragma once
+
+#include "PCGSettings.h"
+
+#include "PCGGetEditorSelection.generated.h"
+
+/**
+ * Produces a single point per editor-selected actor, with the actor transform and bounds.
+ * Mirrors the GetSinglePoint mode of the Get Actor Data node but sources actors from the
+ * current editor viewport selection instead of the actor selector.
+ *
+ * Note: returns no data in non-editor (cooked) builds.
+ */
+UCLASS(MinimalAPI, BlueprintType, ClassGroup = (Procedural))
+class UPCGGetEditorSelectionSettings : public UPCGSettings
+{
+	GENERATED_BODY()
+
+public:
+	//~Begin UPCGSettings interface
+#if WITH_EDITOR
+	virtual FName GetDefaultNodeName() const override { return FName(TEXT("GetEditorSelection")); }
+	virtual FText GetDefaultNodeTitle() const override { return NSLOCTEXT("PCGGetEditorSelectionElement", "NodeTitle", "Get Editor Selection"); }
+	virtual FText GetNodeTooltipText() const override { return NSLOCTEXT("PCGGetEditorSelectionElement", "NodeTooltip", "Produces a single point per editor-selected actor, carrying the actor transform and bounds."); }
+	virtual EPCGSettingsType GetType() const override { return EPCGSettingsType::Spatial; }
+#endif
+
+protected:
+	virtual TArray<FPCGPinProperties> InputPinProperties() const override { return {}; }
+	virtual TArray<FPCGPinProperties> OutputPinProperties() const override { return UPCGSettings::DefaultPointOutputPinProperties(); }
+	virtual FPCGElementPtr CreateElement() const override;
+	//~End UPCGSettings interface
+
+public:
+	/** Silence warnings that attribute names were sanitized to replace invalid characters. */
+	UPROPERTY(BlueprintReadWrite, EditAnywhere, Category = Settings, AdvancedDisplay)
+	bool bSilenceSanitizedAttributeNameWarnings = false;
+};
+
+class FPCGGetEditorSelectionElement : public IPCGElement
+{
+public:
+	// Editor selection must be queried from the main thread.
+	virtual bool CanExecuteOnlyOnMainThread(FPCGContext* Context) const override { return true; }
+	// Selection changes constantly; never cache results.
+	virtual bool IsCacheable(const UPCGSettings* InSettings) const override { return false; }
+
+protected:
+	virtual bool ExecuteInternal(FPCGContext* Context) const override;
+};
diff --git a/Engine/Plugins/PCG/Source/PCG/Private/Elements/Editor/PCGGetEditorSelection.cpp b/Engine/Plugins/PCG/Source/PCG/Private/Elements/Editor/PCGGetEditorSelection.cpp
new file mode 100644
--- /dev/null
+++ b/Engine/Plugins/PCG/Source/PCG/Private/Elements/Editor/PCGGetEditorSelection.cpp
@@ -0,0 +1,93 @@
+// Copyright Epic Games, Inc. All Rights Reserved.
+
+#include "Elements/Editor/PCGGetEditorSelection.h"
+
+#include "PCGContext.h"
+#include "PCGPin.h"
+#include "Data/PCGBasePointData.h"
+
+#if WITH_EDITOR
+#include "Editor.h"
+#include "Engine/Selection.h"
+#include "GameFramework/Actor.h"
+#endif // WITH_EDITOR
+
+#include UE_INLINE_GENERATED_CPP_BY_NAME(PCGGetEditorSelection)
+
+#define LOCTEXT_NAMESPACE "PCGGetEditorSelectionElement"
+
+FPCGElementPtr UPCGGetEditorSelectionSettings::CreateElement() const
+{
+	return MakeShared<FPCGGetEditorSelectionElement>();
+}
+
+bool FPCGGetEditorSelectionElement::ExecuteInternal(FPCGContext* Context) const
+{
+	TRACE_CPUPROFILER_EVENT_SCOPE(FPCGGetEditorSelectionElement::Execute);
+	check(Context);
+
+	const UPCGGetEditorSelectionSettings* Settings = Context->GetInputSettings<UPCGGetEditorSelectionSettings>();
+	check(Settings);
+
+#if WITH_EDITOR
+	if (!GEditor)
+	{
+		PCGE_LOG(Warning, GraphAndLog, LOCTEXT("NoEditor", "GEditor is not available; no data will be produced."));
+		return true;
+	}
+
+	TArray<AActor*> SelectedActors;
+	if (USelection* Selection = GEditor->GetSelectedActors())
+	{
+		Selection->GetSelectedObjects<AActor>(SelectedActors);
+	}
+
+	if (SelectedActors.IsEmpty())
+	{
+		PCGE_LOG(Verbose, LogOnly, LOCTEXT("NoActorsSelected", "No actors are currently selected."));
+		return true;
+	}
+
+	// All selected actors merged into a single point data element.
+	UPCGBasePointData* PointData = nullptr;
+	bool bAnyAttributeNameWasSanitized = false;
+
+	for (AActor* Actor : SelectedActors)
+	{
+		if (!Actor)
+		{
+			continue;
+		}
+
+		if (!PointData)
+		{
+			PointData = FPCGContext::NewPointData_AnyThread(Context);
+		}
+
+		bool bAttributeNameWasSanitized = false;
+		PointData->AddSinglePointFromActor(Actor, &bAttributeNameWasSanitized);
+		bAnyAttributeNameWasSanitized |= bAttributeNameWasSanitized;
+	}
+
+	if (!PointData)
+	{
+		PCGE_LOG(Warning, GraphAndLog, LOCTEXT("OnlyInvalidActors", "All selected actors were invalid, will not produce any data."));
+		return true;
+	}
+
+	if (bAnyAttributeNameWasSanitized && !Settings->bSilenceSanitizedAttributeNameWarnings)
+	{
+		PCGE_LOG(Warning, GraphAndLog, LOCTEXT("TagAttributeNamesSanitized", "One or more tag names contained invalid characters and were sanitized when creating the corresponding attributes."));
+	}
+
+	FPCGTaggedData& TaggedData = Context->OutputData.TaggedData.Emplace_GetRef();
+	TaggedData.Data = PointData;
+
+#else // WITH_EDITOR
+	PCGE_LOG(Warning, GraphAndLog, LOCTEXT("NotAvailableAtRuntime", "Get Data From Selection is only available in editor builds."));
+#endif // WITH_EDITOR
+
+	return true;
+}
+
+#undef LOCTEXT_NAMESPACE
