[PCG] Added node to retrieve editor viewport camera info.

#Synced-CL 51473272

diff --git a/Engine/Plugins/PCG/Source/PCG/Internal/Elements/Editor/PCGGetEditorCameras.h b/Engine/Plugins/PCG/Source/PCG/Internal/Elements/Editor/PCGGetEditorCameras.h
new file mode 100644
--- /dev/null
+++ b/Engine/Plugins/PCG/Source/PCG/Internal/Elements/Editor/PCGGetEditorCameras.h
@@ -0,0 +1,52 @@
+// Copyright Epic Games, Inc. All Rights Reserved.
+
+#pragma once
+
+#include "PCGSettings.h"
+
+#include "PCGGetEditorCameras.generated.h"
+
+/**
+ * Produces one point per editor viewport camera, with the camera world transform as the point transform. 
+ * Perspective viewports appear first, followed by orthographic viewports. Each point carries the following attributes:
+ *   IsPerspective        (bool)   - true for perspective viewports
+ *   ViewportType         (string) - "Perspective", "Top", "Bottom", "Front", "Back", "Left", "Right", "OrthoFreelook" or "Unknown"
+ *   FOV                  (float)  - horizontal field of view in degrees (perspective only; 0 otherwise)
+ *   OrthoZoom            (float)  - orthographic zoom amount (orthographic only; 0 otherwise)
+ *   NearClipPlane        (float)  - near clipping plane distance
+ *   FarClipPlaneOverride (float)  - far clipping plane override; negative if not overridden (engine default applies)
+ *
+ * Note: returns no data in non-editor (cooked) builds.
+ */
+UCLASS(MinimalAPI, BlueprintType, ClassGroup = (Procedural))
+class UPCGGetEditorCamerasSettings : public UPCGSettings
+{
+	GENERATED_BODY()
+
+public:
+	//~Begin UPCGSettings interface
+#if WITH_EDITOR
+	virtual FName GetDefaultNodeName() const override { return FName(TEXT("GetEditorCameras")); }
+	virtual FText GetDefaultNodeTitle() const override;
+	virtual FText GetNodeTooltipText() const override;
+	virtual EPCGSettingsType GetType() const override { return EPCGSettingsType::Generic; }
+#endif
+
+protected:
+	virtual TArray<FPCGPinProperties> InputPinProperties() const override { return {}; }
+	virtual TArray<FPCGPinProperties> OutputPinProperties() const override;
+	virtual FPCGElementPtr CreateElement() const override;
+	//~End UPCGSettings interface
+};
+
+class FPCGGetEditorCamerasElement : public IPCGElement
+{
+public:
+	// Editor viewport clients must be queried from the main thread.
+	virtual bool CanExecuteOnlyOnMainThread(FPCGContext* Context) const override { return true; }
+	// Viewport transforms change every frame; never cache results.
+	virtual bool IsCacheable(const UPCGSettings* InSettings) const override { return false; }
+
+protected:
+	virtual bool ExecuteInternal(FPCGContext* Context) const override;
+};
diff --git a/Engine/Plugins/PCG/Source/PCG/Private/Elements/Editor/PCGGetEditorCameras.cpp b/Engine/Plugins/PCG/Source/PCG/Private/Elements/Editor/PCGGetEditorCameras.cpp
new file mode 100644
--- /dev/null
+++ b/Engine/Plugins/PCG/Source/PCG/Private/Elements/Editor/PCGGetEditorCameras.cpp
@@ -0,0 +1,190 @@
+// Copyright Epic Games, Inc. All Rights Reserved.
+
+#include "Elements/Editor/PCGGetEditorCameras.h"
+
+#include "PCGContext.h"
+#include "PCGPin.h"
+#include "Data/PCGBasePointData.h"
+#include "Helpers/PCGHelpers.h"
+
+#if WITH_EDITOR
+#include "Editor.h"
+#include "EditorViewportClient.h"
+#include "Editor/UnrealEdTypes.h"
+#endif // WITH_EDITOR
+
+#include UE_INLINE_GENERATED_CPP_BY_NAME(PCGGetEditorCameras)
+
+#define LOCTEXT_NAMESPACE "PCGGetEditorCamerasElement"
+
+namespace PCGGetEditorCameras
+{
+	namespace Constants
+	{
+		const FLazyName IsPerspectiveName = TEXT("IsPerspective");
+		const FLazyName ViewportTypeName = TEXT("ViewportType");
+		const FLazyName FOVName = TEXT("FOV");
+		const FLazyName OrthoZoomName = TEXT("OrthoZoom");
+		const FLazyName NearClipPlaneName = TEXT("NearClipPlane");
+		const FLazyName FarClipPlaneOverrideName = TEXT("FarClipPlaneOverride");
+	}
+
+	namespace Helpers
+	{
+#if WITH_EDITOR
+		static FString ViewportTypeToString(ELevelViewportType InType)
+		{
+			switch (InType)
+			{
+			case LVT_Perspective: return TEXT("Perspective");
+			case LVT_OrthoXY: return TEXT("Top");
+			case LVT_OrthoNegativeXY: return TEXT("Bottom");
+			case LVT_OrthoNegativeYZ: return TEXT("Front");
+			case LVT_OrthoYZ: return TEXT("Back");
+			case LVT_OrthoNegativeXZ: return TEXT("Left");
+			case LVT_OrthoXZ: return TEXT("Right");
+			case LVT_OrthoFreelook: return TEXT("OrthoFreelook");
+			default: return TEXT("Unknown");
+			}
+		}
+#endif
+	}
+}
+
+#if WITH_EDITOR
+FText UPCGGetEditorCamerasSettings::GetDefaultNodeTitle() const
+{
+	return LOCTEXT("NodeTitle", "Get Editor Cameras");
+}
+
+FText UPCGGetEditorCamerasSettings::GetNodeTooltipText() const
+{
+	return LOCTEXT("NodeTooltip", 
+		"Produces one point per editor viewport camera.\n"
+		"Perspective viewports appear first.\n"
+		"Each point carries camera attributes:\n"
+		"IsPerspective, ViewportType, FOV, OrthoZoom, NearClipPlane, FarClipPlaneOverride.");
+}
+#endif // WITH_EDITOR
+
+TArray<FPCGPinProperties> UPCGGetEditorCamerasSettings::OutputPinProperties() const
+{
+	TArray<FPCGPinProperties> PinProperties;
+	PinProperties.Emplace(PCGPinConstants::DefaultOutputLabel, EPCGDataType::Point);
+	return PinProperties;
+}
+
+FPCGElementPtr UPCGGetEditorCamerasSettings::CreateElement() const
+{
+	return MakeShared<FPCGGetEditorCamerasElement>();
+}
+
+bool FPCGGetEditorCamerasElement::ExecuteInternal(FPCGContext* Context) const
+{
+	TRACE_CPUPROFILER_EVENT_SCOPE(FPCGGetEditorCamerasElement::Execute);
+	check(Context);
+
+#if WITH_EDITOR
+	if (!GEditor)
+	{
+		PCGE_LOG(Warning, GraphAndLog, LOCTEXT("NoEditor", "GEditor is not available; no data will be produced."));
+		return true;
+	}
+
+	// Separate perspective and orthographic clients; perspective comes first.
+	TArray<FEditorViewportClient*> PerspectiveClients;
+	TArray<FEditorViewportClient*> OrthoClients;
+
+	for (FEditorViewportClient* Client : GEditor->GetAllViewportClients())
+	{
+		if (!Client || !Client->IsVisible() || !Client->IsLevelEditorClient())
+		{
+			continue;
+		}
+
+		if (Client->IsPerspective())
+		{
+			PerspectiveClients.Add(Client);
+		}
+		else
+		{
+			OrthoClients.Add(Client);
+		}
+	}
+
+	if (PerspectiveClients.IsEmpty() && OrthoClients.IsEmpty())
+	{
+		PCGE_LOG(Verbose, LogOnly, LOCTEXT("NoViewports", "No visible viewport clients found."));
+		return true;
+	}
+
+	const int32 TotalCount = PerspectiveClients.Num() + OrthoClients.Num();
+
+	UPCGBasePointData* PointData = FPCGContext::NewPointData_AnyThread(Context);
+	PointData->SetNumPoints(TotalCount);
+	PointData->AllocateProperties(EPCGPointNativeProperties::Transform | EPCGPointNativeProperties::Seed | EPCGPointNativeProperties::MetadataEntry);
+
+	TPCGValueRange<FTransform> TransformRange = PointData->GetTransformValueRange(/*bAllocate=*/false);
+	TPCGValueRange<int32> SeedRange = PointData->GetSeedValueRange(/*bAllocate=*/false);
+	TPCGValueRange<int64> MetadataEntryRange = PointData->GetMetadataEntryValueRange(/*bAllocate=*/false);
+
+	UPCGMetadata* Metadata = PointData->MutableMetadata();
+	check(Metadata);
+
+	using namespace PCGGetEditorCameras::Constants;
+
+	FPCGMetadataAttribute<bool>* IsPerspectiveAttribute = Metadata->FindOrCreateAttribute<bool>(IsPerspectiveName, false, /*bAllowsInterpolation=*/false);
+	FPCGMetadataAttribute<FString>* ViewportTypeAttribute = Metadata->FindOrCreateAttribute<FString>(ViewportTypeName, FString(), /*bAllowsInterpolation=*/false);
+	FPCGMetadataAttribute<float>* FOVAttribute = Metadata->FindOrCreateAttribute<float>  (FOVName, 0.f, /*bAllowsInterpolation=*/true);
+	FPCGMetadataAttribute<float>* OrthoZoomAttribute = Metadata->FindOrCreateAttribute<float>  (OrthoZoomName, 0.f, /*bAllowsInterpolation=*/true);
+	FPCGMetadataAttribute<float>* NearClipAttribute = Metadata->FindOrCreateAttribute<float>  (NearClipPlaneName, 0.f, /*bAllowsInterpolation=*/true);
+	FPCGMetadataAttribute<float>* FarClipAttribute = Metadata->FindOrCreateAttribute<float>  (FarClipPlaneOverrideName, -1.f, /*bAllowsInterpolation=*/true);
+
+	auto AddViewportPoint = [&TransformRange, &SeedRange, &MetadataEntryRange, Metadata, IsPerspectiveAttribute, ViewportTypeAttribute, FOVAttribute, OrthoZoomAttribute, NearClipAttribute, FarClipAttribute](FEditorViewportClient* Client, int32 PointIndex)
+	{
+		const FVector  Location = Client->GetViewLocation();
+		const FRotator Rotation = Client->GetViewRotation();
+
+		TransformRange[PointIndex] = FTransform(Rotation, Location);
+		SeedRange[PointIndex] = PCGHelpers::ComputeSeedFromPosition(Location);
+		const PCGMetadataEntryKey Entry = Metadata->AddEntry();
+		MetadataEntryRange[PointIndex] = Entry;
+
+		const bool bIsPerspective = Client->IsPerspective();
+		IsPerspectiveAttribute->SetValue(Entry, bIsPerspective);
+		ViewportTypeAttribute->SetValue(Entry, PCGGetEditorCameras::Helpers::ViewportTypeToString(Client->GetViewportType()));
+
+		if (bIsPerspective)
+		{
+			FOVAttribute->SetValue(Entry, Client->ViewFOV);
+		}
+		else
+		{
+			OrthoZoomAttribute->SetValue(Entry, Client->GetOrthoZoom());
+		}
+
+		NearClipAttribute->SetValue(Entry, Client->GetNearClipPlane());
+		FarClipAttribute->SetValue(Entry, Client->GetFarClipPlaneOverride());
+	};
+
+	int32 PointIndex = 0;
+	for (FEditorViewportClient* Client : PerspectiveClients)
+	{
+		AddViewportPoint(Client, PointIndex++);
+	}
+
+	for (FEditorViewportClient* Client : OrthoClients)
+	{
+		AddViewportPoint(Client, PointIndex++);
+	}
+
+	Context->OutputData.TaggedData.Emplace_GetRef().Data = PointData;
+
+#else // WITH_EDITOR
+	PCGE_LOG(Warning, GraphAndLog, LOCTEXT("NotAvailableAtRuntime", "Get Editor Cameras is only available in editor builds."));
+#endif // WITH_EDITOR
+
+	return true;
+}
+
+#undef LOCTEXT_NAMESPACE
