[PCG] Switch to proper widgets for GPU upload/download markers, so that we can have tooltips.

#rnx
#Synced-CL 51663631

diff --git a/Engine/Plugins/PCG/Source/PCGEditor/Private/Nodes/SPCGEditorGraphNode.cpp b/Engine/Plugins/PCG/Source/PCGEditor/Private/Nodes/SPCGEditorGraphNode.cpp
--- a/Engine/Plugins/PCG/Source/PCGEditor/Private/Nodes/SPCGEditorGraphNode.cpp
+++ b/Engine/Plugins/PCG/Source/PCGEditor/Private/Nodes/SPCGEditorGraphNode.cpp
@@ -29,11 +29,18 @@
 #include "ThumbnailRendering/ThumbnailManager.h"
 
 #include "Widgets/SBoxPanel.h"
+#include "Widgets/Images/SImage.h"
 #include "Widgets/Input/SButton.h"
 #include "Widgets/Text/SInlineEditableTextBlock.h"
 
 #define LOCTEXT_NAMESPACE "SPCGEditorGraphNode"
 
+namespace PCGEditorGraphNode
+{
+	constexpr float OverlayIconSize = 16.0f;
+	constexpr float OverlayHalfIconSize = OverlayIconSize / 2.0f;
+}
+
 void SPCGEditorGraphNode::Construct(const FArguments& InArgs, UPCGEditorGraphNodeBase* InNode)
 {
 	GraphNode = InNode;
@@ -44,6 +51,18 @@ void SPCGEditorGraphNode::Construct(const FArguments& InArgs, UPCGEditorGraphNod
 		InNode->OnNodeChangedDelegate.BindSP(this, &SPCGEditorGraphNode::OnNodeChanged);
 	}
 
+	// Create GPU transfer indicator widgets once and reuse them across GetOverlayWidgets() calls. A stable widget pointer is required so that Slate can match the widget seen
+	// during OnArrangeChildren (hit-testing) with the one seen during Paint, which is the prerequisite for hover state tracking and tooltip display.
+	GPUUploadIndicatorWidget =
+		SNew(SImage)
+		.Image(FPCGEditorStyle::Get().GetBrush(TEXT("PCG.NodeOverlay.GPUUpload")))
+		.ToolTipText(LOCTEXT("GPUUploadTooltip", "One or more outputs of this node were uploaded from the CPU to the GPU. Uploads are often expensive and should be minimized."));
+
+	GPUReadbackIndicatorWidget =
+		SNew(SImage)
+		.Image(FPCGEditorStyle::Get().GetBrush(TEXT("PCG.NodeOverlay.GPUReadback")))
+		.ToolTipText(LOCTEXT("GPUReadbackTooltip", "One or more inputs of this node were read back form the GPU to the CPU. Readbacks are often expensive and should be minimized."));
+
 	UpdateGraphNode();
 }
 
@@ -310,6 +329,16 @@ TArray<FOverlayWidgetInfo> SPCGEditorGraphNode::GetOverlayWidgets(bool bSelected
 		AddGPUOverlayWidget(OverlayWidgets);
 	}
 
+	if (PCGEditorGraphNode->GetTriggeredGPUUpload())
+	{
+		AddGPUUploadWidget(OverlayWidgets);
+	}
+
+	if (PCGEditorGraphNode->GetTriggeredGPUReadback())
+	{
+		AddGPUReadbackWidget(OverlayWidgets);
+	}
+
 	return OverlayWidgets;
 }
 
@@ -317,9 +346,10 @@ void SPCGEditorGraphNode::GetOverlayBrushes(bool bSelected, const FVector2f& Wid
 {
 	check(PCGEditorGraphNode);
 	
-	float YOffsetLeft = 0.0f;
-	// Start lower down to be clear of grid size label.
-	float YOffsetRight = UsesHiGenOverlay() ? 18.0f : 0.0f;
+	// Reserve space already occupied by overlay widgets on the same side so brushes stack below them.
+	float YOffsetLeft = PCGEditorGraphNode->GetTriggeredGPUReadback() ? PCGEditorGraphNode::OverlayIconSize : 0.0f;
+	// Start lower down to be clear of the HiGen grid label and the upload widget (both on the right edge).
+	float YOffsetRight = (UsesHiGenOverlay() ? PCGEditorGraphNode::OverlayIconSize + 2.0f : 0.0f) + (PCGEditorGraphNode->GetTriggeredGPUUpload() ? PCGEditorGraphNode::OverlayIconSize : 0.0f);
 
 	auto AddOverlayBrush = [&YOffsetLeft, &YOffsetRight, &Brushes, this](const FName& BrushName, bool bRightSide = false)
 	{
@@ -344,16 +374,6 @@ void SPCGEditorGraphNode::GetOverlayBrushes(bool bSelected, const FVector2f& Wid
 		}
 	};
 
-	if (PCGEditorGraphNode->GetTriggeredGPUUpload())
-	{
-		AddOverlayBrush(TEXT("PCG.NodeOverlay.GPUUpload"), /*bRightSide=*/true);
-	}
-
-	if (PCGEditorGraphNode->GetTriggeredGPUReadback())
-	{
-		AddOverlayBrush(TEXT("PCG.NodeOverlay.GPUReadback"));
-	}
-
 	if (PCGEditorGraphNode->IsCulledFromExecution())
 	{
 		AddOverlayBrush(PCGEditorStyleConstants::Node_Overlay_Inactive);
@@ -575,6 +595,28 @@ void SPCGEditorGraphNode::AddGPUOverlayWidget(TArray<FOverlayWidgetInfo>& Overla
 	OverlayWidgets.Add(GPUUsageLabelInfo);
 }
 
+void SPCGEditorGraphNode::AddGPUUploadWidget(TArray<FOverlayWidgetInfo>& OverlayWidgets) const
+{
+	check(GPUUploadIndicatorWidget.IsValid());
+
+	// Mirror the position used in GetOverlayBrushes: right edge of the node, offset downward when the HiGen label is also visible so the two don't overlap.
+	const float YOffset = (UsesHiGenOverlay() ? PCGEditorGraphNode::OverlayIconSize + 2.0f : 0.0f) - PCGEditorGraphNode::OverlayHalfIconSize;
+
+	FOverlayWidgetInfo Info(GPUUploadIndicatorWidget);
+	Info.OverlayOffset = FVector2D(GetDesiredSize().X - PCGEditorGraphNode::OverlayHalfIconSize, YOffset);
+	OverlayWidgets.Add(Info);
+}
+
+void SPCGEditorGraphNode::AddGPUReadbackWidget(TArray<FOverlayWidgetInfo>& OverlayWidgets) const
+{
+	check(GPUReadbackIndicatorWidget.IsValid());
+
+	// Mirror the position used in GetOverlayBrushes: left edge of the node, top.
+	FOverlayWidgetInfo Info(GPUReadbackIndicatorWidget);
+	Info.OverlayOffset = FVector2D(-PCGEditorGraphNode::OverlayHalfIconSize, -PCGEditorGraphNode::OverlayHalfIconSize);
+	OverlayWidgets.Add(Info);
+}
+
 const FSlateBrush* SPCGEditorGraphNode::GetBorderBrush(uint32 InspectedGrid, uint32 NodeGrid) const
 {
 	if (InspectedGrid == NodeGrid)
diff --git a/Engine/Plugins/PCG/Source/PCGEditor/Private/Nodes/SPCGEditorGraphNode.h b/Engine/Plugins/PCG/Source/PCGEditor/Private/Nodes/SPCGEditorGraphNode.h
--- a/Engine/Plugins/PCG/Source/PCGEditor/Private/Nodes/SPCGEditorGraphNode.h
+++ b/Engine/Plugins/PCG/Source/PCGEditor/Private/Nodes/SPCGEditorGraphNode.h
@@ -70,4 +70,11 @@ private:
 	const FSlateBrush* GetBorderBrush(uint32 InspectedGrid, uint32 NodeGrid) const;
 
 	UPCGEditorGraphNodeBase* PCGEditorGraphNode = nullptr;
+
+	/**
+	* Stable widget instances for GPU transfer indicators. Must be the same pointer across GetOverlayWidgets() calls (Paint vs OnArrangeChildren) so that Slate
+	* can track hover state and show tooltips.
+	*/
+	TSharedPtr<SWidget> GPUUploadIndicatorWidget;
+	TSharedPtr<SWidget> GPUReadbackIndicatorWidget;
 };
