[PCG] Add a basic graph layout button to the editor.

#Synced-CL 51713122

diff --git a/Engine/Plugins/PCG/Source/PCGEditor/Private/PCGEditor.cpp b/Engine/Plugins/PCG/Source/PCGEditor/Private/PCGEditor.cpp
--- a/Engine/Plugins/PCG/Source/PCGEditor/Private/PCGEditor.cpp
+++ b/Engine/Plugins/PCG/Source/PCGEditor/Private/PCGEditor.cpp
@@ -69,6 +69,7 @@
 #include "ToolMenuSection.h"
 #include "UnrealEdGlobals.h"
 #include "Algo/AnyOf.h"
+#include "Algo/TopologicalSort.h"
 #include "Editor/UnrealEdEngine.h"
 #include "Fonts/FontMeasure.h"
 #include "Framework/Application/SlateApplication.h"
@@ -120,6 +121,11 @@ namespace FPCGEditor_private
 	{
 		PCGEditorTabs::CodeEditorID,
 	};
+
+	constexpr int32 DefaultFixedNodeWidth = 250;
+	constexpr int32 DefaultFixedNodeHeight = 200;
+	constexpr int32 NodeLayoutHorizontalGap = 200;
+	constexpr int32 NodeLayoutVerticalGap = 50;
 }
 
 UPCGEditorGraph* FPCGEditor::GetPCGEditorGraph(UPCGGraph* InGraph)
@@ -818,6 +824,7 @@ void FPCGEditor::RegisterToolbarInternal(FToolBarBuilder& ToolbarBuilder) const
 	RegisterToolbarButton(ToolbarBuilder, EPCGToolbarButtons::Find);
 	RegisterToolbarButton(ToolbarBuilder, EPCGToolbarButtons::PauseRegen);
 	RegisterToolbarButton(ToolbarBuilder, EPCGToolbarButtons::ForceRegen);
+	RegisterToolbarButton(ToolbarBuilder, EPCGToolbarButtons::AutoLayoutNodes);
 	RegisterToolbarButton(ToolbarBuilder, EPCGToolbarButtons::CancelExecution);
 	RegisterToolbarButton(ToolbarBuilder, EPCGToolbarButtons::OpenDebugObjectTreeTab);
 
@@ -859,6 +866,11 @@ void FPCGEditor::RegisterToolbarButton(FToolBarBuilder& ToolbarBuilder, EPCGTool
 				}));
 			break;
 		}
+	case EPCGToolbarButtons::AutoLayoutNodes:
+		{
+			ToolbarBuilder.AddToolBarButton(PCGEditorCommands.AutoLayoutNodes);
+			break;
+		}
 	case EPCGToolbarButtons::CancelExecution:
 		{
 			ToolbarBuilder.AddToolBarButton(PCGEditorCommands.CancelExecution, NAME_None, TAttribute<FText>(), TAttribute<FText>(),
@@ -920,6 +932,10 @@ void FPCGEditor::BindCommands()
 		PCGEditorCommands.ForceGraphRegeneration,
 		FExecuteAction::CreateSP(this, &FPCGEditor::OnForceGraphRegeneration_Clicked));
 
+	ToolkitCommands->MapAction(
+		PCGEditorCommands.AutoLayoutNodes,
+		FExecuteAction::CreateSP(this, &FPCGEditor::OnAutoLayoutNodes_Clicked));
+
 	ToolkitCommands->MapAction(
 		PCGEditorCommands.CancelExecution,
 		FExecuteAction::CreateSP(this, &FPCGEditor::OnCancelExecution_Clicked),
@@ -1182,6 +1198,25 @@ void FPCGEditor::OnForceGraphRegeneration_Clicked()
 	}
 }
 
+void FPCGEditor::OnAutoLayoutNodes_Clicked()
+{
+	if (MainGraph)
+	{
+		FScopedTransaction NodeMoveTransaction(LOCTEXT("MoveNodesAction", "Auto-Layout Nodes"));
+		
+		if (AutoLayoutFullGraph())
+		{
+			// Force a refresh
+			MainGraph->GetEditorGraph()->ReconstructGraph();
+		}
+		else
+		{
+			// Do not spam the undo buffer if the graph was already well positioned.
+			NodeMoveTransaction.Cancel();
+		}
+	}
+}
+
 void FPCGEditor::OnCancelExecution_Clicked()
 {
 	IPCGBaseSubsystem* Subsystem = GetSubsystem();
@@ -3612,6 +3647,176 @@ TSharedPtr<SPCGEditorViewport> FPCGEditor::GetViewportWidget(int32 Index) const
 	return ViewportWidgets[Index];
 }
 
+bool FPCGEditor::AutoLayoutFullGraph()
+{
+	UPCGNode* InputNode = MainGraph->InputNode;
+	check(InputNode);
+	UPCGNode* OutputNode = MainGraph->OutputNode;
+	check(OutputNode);
+
+	TArray<UPCGNode*> AllNodes = MainGraph->GetNodes();
+	AllNodes.Add(InputNode);
+	AllNodes.Add(OutputNode);
+	
+	// Precompute predecessor since we use it more than once
+	TMap<UPCGNode*, TSet<UPCGNode*>> PredecessorNodes;
+	for (UPCGNode* Node : AllNodes)
+	{
+		TSet<UPCGNode*>& CurPred = PredecessorNodes.Add(Node);
+		for (UPCGPin* Pin : Node->GetInputPins())
+		{
+			for (UPCGEdge* Edge : Pin->Edges)
+			{
+				if (Edge && Edge->InputPin && Edge->InputPin->Node)
+				{
+					CurPred.Add(Edge->InputPin->Node);
+				}
+			}
+		}
+	}
+	
+	auto GetNodePredecessor = [&PredecessorNodes](const UPCGNode* Node) -> TSet<UPCGNode*>
+	{
+		return PredecessorNodes[Node];
+	};
+	
+	// Use the min position as the base position to keep the graph approx. where it is
+	FIntVector2 BasePosition{AllNodes[0]->PositionX, AllNodes[0]->PositionY};
+	for (const UPCGNode* Node : AllNodes)
+	{
+		FIntVector2 CurPosition{Node->PositionX, Node->PositionY};
+		BasePosition.X = FMath::Min(BasePosition.X, CurPosition.X);
+		BasePosition.Y = FMath::Min(BasePosition.Y, CurPosition.Y);
+	}
+	
+	bool bNodeMovementOccurred = false;
+	auto SetNodePosition = [&bNodeMovementOccurred](UPCGNode* Node, int32 NewPositionX, int32 NewPositionY)
+	{
+		int32 CurPositionX = 0;
+		int32 CurPositionY = 0;
+		Node->GetNodePosition(CurPositionX, CurPositionY);
+		if (CurPositionX != NewPositionX || CurPositionY != NewPositionY)
+		{
+			bNodeMovementOccurred = true;
+			Node->Modify();
+			Node->SetNodePosition(NewPositionX, NewPositionY);
+		}
+	};
+	
+	// Extract lonely nodes (Nodes without any connection), we'll place them manually
+	auto IsLonely = [](const UPCGNode* Node) -> bool
+	{
+		return !Node->GetFirstConnectedInputPin() && !Node->GetFirstConnectedOutputPin();
+	};
+	
+	TArray<UPCGNode*> LonelyNodes;
+	TArray<UPCGNode*> SortedAttachedNodes;
+	for (UPCGNode* Node : MainGraph->GetNodes())
+	{
+		if (IsLonely(Node))
+		{
+			LonelyNodes.Add(Node);
+		}
+		else
+		{
+			SortedAttachedNodes.Add(Node);
+		}
+	}
+	
+	const bool bInputIsLonely = IsLonely(InputNode);
+	const bool bOutputIsLonely = IsLonely(OutputNode);
+	
+	if (!bInputIsLonely)
+	{
+		SortedAttachedNodes.Add(InputNode);
+	}
+	
+	if (!bOutputIsLonely)
+	{
+		SortedAttachedNodes.Add(OutputNode);
+	}
+	
+	Algo::TopologicalSort(SortedAttachedNodes, GetNodePredecessor);
+	
+	// Give space for the input node
+	if (bInputIsLonely)
+	{
+		BasePosition.X += (FPCGEditor_private::DefaultFixedNodeWidth + FPCGEditor_private::NodeLayoutHorizontalGap);
+	}
+	
+	// Associate a layer to each node
+	int MaxLayerIndex = 0;
+	TMap<UPCGNode*, int> NodeToLayerIndex;
+	for (UPCGNode* Node : SortedAttachedNodes)
+	{
+		NodeToLayerIndex.Add(Node, 0);
+	}
+
+	for (UPCGNode* Node : SortedAttachedNodes)
+	{
+		int MaxLayerOfPred = 0;
+		for (UPCGNode* PredNode: GetNodePredecessor(Node))
+		{
+			MaxLayerOfPred = FMath::Max(MaxLayerOfPred, NodeToLayerIndex[PredNode] + 1);
+		}
+		NodeToLayerIndex[Node] = MaxLayerOfPred;
+		MaxLayerIndex = FMath::Max(MaxLayerIndex, MaxLayerOfPred);
+	}
+	
+	TArray<TArray<UPCGNode*>> LayerMembersByIndex;
+	LayerMembersByIndex.SetNum(MaxLayerIndex+1);
+	for (UPCGNode* Node : SortedAttachedNodes)
+	{
+		LayerMembersByIndex[NodeToLayerIndex[Node]].Add(Node);
+	}
+
+	// Set Positions for attached nodes
+	int32 CurrentLayerPosX = BasePosition.X;
+	int32 MaxPositionY = BasePosition.Y;
+	for (TArray<UPCGNode*>& CurLayer : LayerMembersByIndex)
+	{
+		if (CurLayer.IsEmpty())
+		{
+			continue;
+		}
+		
+		int32 CurLayerPosY = BasePosition.Y;
+		int32 MaxNodeWidthForLayer = 0;
+		for (UPCGNode* Node : CurLayer)
+		{
+			SetNodePosition(Node, CurrentLayerPosX, CurLayerPosY);
+			
+			// Update next Pos Y for this layer
+			CurLayerPosY += FPCGEditor_private::DefaultFixedNodeHeight + FPCGEditor_private::NodeLayoutVerticalGap;
+			
+			MaxNodeWidthForLayer = FMath::Max(MaxNodeWidthForLayer, FPCGEditor_private::DefaultFixedNodeWidth);
+			MaxPositionY = FMath::Max(MaxPositionY, CurLayerPosY);
+		}
+		
+		//Update position X for the next layer
+		CurrentLayerPosX += MaxNodeWidthForLayer + FPCGEditor_private::NodeLayoutHorizontalGap;
+	}
+	
+	if (bInputIsLonely)
+	{
+		SetNodePosition(InputNode, BasePosition.X - (FPCGEditor_private::DefaultFixedNodeWidth + FPCGEditor_private::NodeLayoutHorizontalGap), BasePosition.Y);
+	}
+	
+	if (bOutputIsLonely)
+	{
+		SetNodePosition(OutputNode, CurrentLayerPosX, BasePosition.Y);
+	}
+	
+	int32 LonelyLayerPosX = BasePosition.X;
+	for (UPCGNode* LonelyNode : LonelyNodes)
+	{
+		SetNodePosition(LonelyNode, LonelyLayerPosX, MaxPositionY);
+		LonelyLayerPosX += FPCGEditor_private::DefaultFixedNodeWidth + FPCGEditor_private::NodeLayoutHorizontalGap;
+	}
+	
+	return bNodeMovementOccurred;
+}
+
 TSharedRef<SPCGEditorViewport> FPCGEditor::CreateViewportWidget()
 {
 	return SNew(SPCGEditorViewport);
diff --git a/Engine/Plugins/PCG/Source/PCGEditor/Private/PCGEditorCommands.cpp b/Engine/Plugins/PCG/Source/PCGEditor/Private/PCGEditorCommands.cpp
--- a/Engine/Plugins/PCG/Source/PCGEditor/Private/PCGEditorCommands.cpp
+++ b/Engine/Plugins/PCG/Source/PCGEditor/Private/PCGEditorCommands.cpp
@@ -30,6 +30,7 @@ void FPCGEditorCommands::RegisterCommands()
 	UI_COMMAND(ShowSelectedDetails, "Show Node Details", "Opens a details panel for the selected nodes.", EUserInterfaceActionType::Button, FInputChord(EKeys::F4));
 	UI_COMMAND(PauseAutoRegeneration, "Pause Regen", "Pause automatic regeneration of the current graph.", EUserInterfaceActionType::ToggleButton, FInputChord(EModifierKey::Alt, EKeys::R));
 	UI_COMMAND(ForceGraphRegeneration, "Force Regen", "Manually force a regeneration of the current graph.\nCtrl-click will also perform a flush cache before the regeneration.", EUserInterfaceActionType::Button, FInputChord());
+	UI_COMMAND(AutoLayoutNodes, "Auto Layout", "Automatically reposition the nodes in the graph for better layout.", EUserInterfaceActionType::Button, FInputChord());
 	UI_COMMAND(OpenDebugObjectTreeTab, "Debug Object Tree", "Open the Debug Object Tree tab to display and select graph invocations to debug.", EUserInterfaceActionType::Button, FInputChord());
 	UI_COMMAND(RunDeterminismNodeTest, "Run Determinism Test on Node", "Evaluate the current node for determinism.", EUserInterfaceActionType::Button, FInputChord(EModifierKey::Alt, EKeys::T));
 	UI_COMMAND(RunDeterminismGraphTest, "Graph Determinism Test", "Evaluate the current graph for determinism.", EUserInterfaceActionType::Button, FInputChord());
diff --git a/Engine/Plugins/PCG/Source/PCGEditor/Public/PCGEditor.h b/Engine/Plugins/PCG/Source/PCGEditor/Public/PCGEditor.h
--- a/Engine/Plugins/PCG/Source/PCGEditor/Public/PCGEditor.h
+++ b/Engine/Plugins/PCG/Source/PCGEditor/Public/PCGEditor.h
@@ -77,6 +77,7 @@ enum class EPCGToolbarButtons
 	Find,
 	PauseRegen,
 	ForceRegen,
+	AutoLayoutNodes,
 	CancelExecution,
 	OpenDebugObjectTreeTab,
 	GraphParams,
@@ -340,6 +341,9 @@ protected:
 
 	/** Returns viewport at specified index if it exists, nullptr otherwise */
 	PCGEDITOR_API TSharedPtr<SPCGEditorViewport> GetViewportWidget(int32 Index) const;
+	
+	/** Auto layout the nodes for better visualization. Returns true if a node has moved.*/
+	PCGEDITOR_API bool AutoLayoutFullGraph();
 
 	// TO BE REMOVED IN 5.8
 	// @todo_pcg
@@ -377,6 +381,9 @@ private:
 	/** Force a regeneration by invoking the graph notifications  */
 	void OnForceGraphRegeneration_Clicked();
 
+	/** Auto-Layout nodes in the graph for better layout */
+	void OnAutoLayoutNodes_Clicked();
+
 	/** Whether selected nodes are inspected or not */
 	ECheckBoxState GetInspectedCheckState() const;
 
diff --git a/Engine/Plugins/PCG/Source/PCGEditor/Public/PCGEditorCommands.h b/Engine/Plugins/PCG/Source/PCGEditor/Public/PCGEditorCommands.h
--- a/Engine/Plugins/PCG/Source/PCGEditor/Public/PCGEditorCommands.h
+++ b/Engine/Plugins/PCG/Source/PCGEditor/Public/PCGEditorCommands.h
@@ -27,6 +27,7 @@ public:
 	TSharedPtr<FUICommandInfo> ShowSelectedDetails;
 	TSharedPtr<FUICommandInfo> PauseAutoRegeneration;
 	TSharedPtr<FUICommandInfo> ForceGraphRegeneration;
+	TSharedPtr<FUICommandInfo> AutoLayoutNodes;
 	TSharedPtr<FUICommandInfo> OpenDebugObjectTreeTab;
 	TSharedPtr<FUICommandInfo> RunDeterminismNodeTest;
 	TSharedPtr<FUICommandInfo> RunDeterminismGraphTest;
