[PCG] Replace bIsEditorOnly with per-platform ShouldCook property on UPCGGraph

Replaces the boolean bIsEditorOnly flag on UPCGGraph with FPerPlatformBool ShouldCook, allowing PCG graphs to be selectively excluded from cooked builds on a per-platform basis. When ShouldCook is disabled for a target platform, the graph and all its compiled data are stripped via NeedsLoadForTargetPlatform. UPCGGraphInstance and UPCGComponent also override NeedsLoadForTargetPlatform, delegating to their graph so referencing objects are excluded automatically.

Changes:
- UPCGGraph: Replaced bIsEditorOnly (bool) with ShouldCook (FPerPlatformBool, defaults to true) in the Cooking category. Added NeedsLoadForTargetPlatform override that checks ShouldCook for the target platform. Added AnyPlatformShouldCook() helper used by IsEditorOnly to determine if the graph is editor-only across all platforms.
- UPCGGraph: Deprecated bIsEditorOnly with UE_DEPRECATED(5.8). PostLoad migrates bIsEditorOnly=true to ShouldCook=false.
- UPCGGraphInstance: Added NeedsLoadForTargetPlatform override delegating to its graph.
- UPCGComponent: Added NeedsLoadForTargetPlatform override delegating to its graph.

#rnx

#Synced-CL 51693885

diff --git a/Engine/Plugins/PCG/Source/PCG/Private/PCGComponent.cpp b/Engine/Plugins/PCG/Source/PCG/Private/PCGComponent.cpp
--- a/Engine/Plugins/PCG/Source/PCG/Private/PCGComponent.cpp
+++ b/Engine/Plugins/PCG/Source/PCG/Private/PCGComponent.cpp
@@ -38,6 +38,7 @@
 #include "Engine/BlueprintGeneratedClass.h"
 #include "Engine/Engine.h"
 #include "Engine/Level.h"
+#include "Interfaces/ITargetPlatform.h"
 #include "Kismet/GameplayStatics.h"
 #include "UObject/Package.h"
 
@@ -2094,6 +2095,17 @@ bool UPCGComponent::IsEditorOnly() const
 	return Super::IsEditorOnly() || (GraphInstance && GraphInstance->IsEditorOnly());
 }
 
+bool UPCGComponent::NeedsLoadForTargetPlatform(const ITargetPlatform* TargetPlatform) const
+{
+	if (!Super::NeedsLoadForTargetPlatform(TargetPlatform))
+	{
+		return false;
+	}
+
+	const UPCGGraph* PCGGraph = GetGraph();
+	return !PCGGraph || PCGGraph->NeedsLoadForTargetPlatform(TargetPlatform);
+}
+
 void UPCGComponent::PostInitProperties()
 {
 	LLM_SCOPE_BYTAG(PCG);
diff --git a/Engine/Plugins/PCG/Source/PCG/Private/PCGGraph.cpp b/Engine/Plugins/PCG/Source/PCG/Private/PCGGraph.cpp
--- a/Engine/Plugins/PCG/Source/PCG/Private/PCGGraph.cpp
+++ b/Engine/Plugins/PCG/Source/PCG/Private/PCGGraph.cpp
@@ -23,6 +23,7 @@
 #include "Helpers/PCGPropertyHelpers.h"
 
 #include "AssetRegistry/AssetData.h"
+#include "Interfaces/ITargetPlatform.h"
 #include "Logging/StructuredLog.h"
 #include "Serialization/ArchiveReplaceObjectRef.h"
 #include "UObject/FortniteMainBranchObjectVersion.h"
@@ -654,6 +655,12 @@ void UPCGGraph::PostLoad()
 		HiGenGridSizeMultiplier = static_cast<double>(1 << HiGenExponential);
 		HiGenExponential = 0;
 	}
+
+	if (bIsEditorOnly)
+	{
+		ShouldCook = false;
+		bIsEditorOnly = false;
+	}
 	PRAGMA_ENABLE_DEPRECATION_WARNINGS
 #endif // WITH_EDITOR
 }
@@ -667,7 +674,7 @@ bool UPCGGraph::IsEditorOnly() const
 	}
 #endif
 
-	bool bIsCurrentlyEditorOnly = (Super::IsEditorOnly() || bIsEditorOnly);
+	bool bIsCurrentlyEditorOnly = (Super::IsEditorOnly() || !AnyPlatformShouldCook());
 
 	if (!bIsCurrentlyEditorOnly)
 	{
@@ -705,7 +712,21 @@ void UPCGGraph::Serialize(FArchive& Ar)
 
 bool UPCGGraph::IsEditorOnly_Internal() const
 {
-	return Super::IsEditorOnly() || bIsEditorOnly;
+	return Super::IsEditorOnly() || !AnyPlatformShouldCook();
+}
+
+bool UPCGGraph::NeedsLoadForTargetPlatform(const ITargetPlatform* TargetPlatform) const
+{
+	if (!Super::NeedsLoadForTargetPlatform(TargetPlatform))
+	{
+		return false;
+	}
+
+#if WITH_EDITOR || WITH_PREVIEW_PPX_DATA
+	return ShouldCook.GetValueForPlatform(TargetPlatform ? FName(*TargetPlatform->IniPlatformName()) : NAME_None);
+#else
+	return ShouldCook.GetValue();
+#endif
 }
 
 bool UPCGGraph::Contains(const UPCGGraph* InGraph) const
@@ -2543,6 +2564,28 @@ uint32 UPCGGraph::GetNodeGenerationGridSize(const UPCGNode* InNode, uint32 InDef
 	}
 }
 
+bool UPCGGraph::AnyPlatformShouldCook() const
+{
+#if WITH_EDITOR || WITH_PREVIEW_PPX_DATA
+	if (ShouldCook.Default)
+	{
+		return true;
+	}
+
+	for (const auto& Pair : ShouldCook.PerPlatform)
+	{
+		if (Pair.Value)
+		{
+			return true;
+		}
+	}
+
+	return false;
+#else
+	return ShouldCook.GetValue();
+#endif
+}
+
 uint32 UPCGGraph::CalculateNodeGridSizeRecursive_Unsafe(const UPCGNode* InNode, uint32 InDefaultGridSize) const
 {
 	if (const uint32* CachedGridSize = NodeToGridSize.Find(InNode))
@@ -2696,6 +2739,17 @@ void UPCGGraphInstance::PostLoad()
 #endif // WITH_EDITOR
 }
 
+bool UPCGGraphInstance::NeedsLoadForTargetPlatform(const ITargetPlatform* TargetPlatform) const
+{
+	if (!Super::NeedsLoadForTargetPlatform(TargetPlatform))
+	{
+		return false;
+	}
+
+	const UPCGGraph* PCGGraph = GetGraph();
+	return !PCGGraph || PCGGraph->NeedsLoadForTargetPlatform(TargetPlatform);
+}
+
 void UPCGGraphInstance::BeginDestroy()
 {
 #if WITH_EDITOR
diff --git a/Engine/Plugins/PCG/Source/PCG/Public/PCGComponent.h b/Engine/Plugins/PCG/Source/PCG/Public/PCGComponent.h
--- a/Engine/Plugins/PCG/Source/PCG/Public/PCGComponent.h
+++ b/Engine/Plugins/PCG/Source/PCG/Public/PCGComponent.h
@@ -135,6 +135,7 @@ public:
 	UE_API virtual void PostInitProperties() override;
 	UE_API virtual void BeginDestroy() override;
 	UE_API virtual bool IsEditorOnly() const override;
+	UE_API virtual bool NeedsLoadForTargetPlatform(const ITargetPlatform* TargetPlatform) const override;
 
 #if WITH_EDITOR
 	UE_API virtual bool CanEditChange(const FProperty* InProperty) const override;
diff --git a/Engine/Plugins/PCG/Source/PCG/Public/PCGGraph.h b/Engine/Plugins/PCG/Source/PCG/Public/PCGGraph.h
--- a/Engine/Plugins/PCG/Source/PCG/Public/PCGGraph.h
+++ b/Engine/Plugins/PCG/Source/PCG/Public/PCGGraph.h
@@ -18,6 +18,7 @@
 #include "ComputeFramework/ComputeGraphInstance.h"
 #include "StructUtils/PropertyBag.h"
 #include "UObject/ObjectPtr.h"
+#include "UObject/PerPlatformProperties.h"
 
 #include "PCGGraph.generated.h"
 
@@ -327,6 +328,7 @@ public:
 	/** ~Begin UObject interface */
 	UE_API virtual void PostLoad() override;
 	UE_API virtual bool IsEditorOnly() const override;
+	UE_API virtual bool NeedsLoadForTargetPlatform(const ITargetPlatform* TargetPlatform) const override;
 	UE_API virtual void Serialize(FArchive& Ar) override;
 
 #if WITH_EDITOR
@@ -579,6 +581,9 @@ public:
 	TObjectPtr<UPCGGraphCompilationData> GetCookedCompilationData() { return CookedCompilationData; }
 	const TObjectPtr<UPCGGraphCompilationData> GetCookedCompilationData() const { return CookedCompilationData; }
 
+	/** Returns true if ShouldCook is enabled for the default or any per-platform override. When false, the graph is effectively editor-only and will be excluded from all cooked builds. */
+	bool AnyPlatformShouldCook() const;
+
 protected:
 #if WITH_EDITOR
 	UFUNCTION()
@@ -707,9 +712,14 @@ protected:
 	mutable TMap<uint32, PCGHiGenGrid::FSizeArray> ChildGridSizeToParentGridSizes;
 	mutable PCG::FLock CachedGridInfoLock;
 
-	/** Sets whether this graph is marked as editor-only; note that the IsEditorOnly call depends on the local graph value and the value in all subgraphs, recursively. */
-	UPROPERTY(EditAnywhere, BlueprintReadWrite, Category = Cooking, meta=(EditCondition="!IsStandaloneGraph() && !IsEmbeddedSubgraph()", EditConditionHides))
-	bool bIsEditorOnly = false;
+	/**
+	 * Controls whether this PCG graph is included when cooking for each platform.
+	 * If disabled, the graph and all of its compiled data (compute graphs, kernel sources, pins, nodes, edges)
+	 * will be stripped from cooked builds. Any PCG components or graph instances that reference this graph
+	 * will also be excluded from cooking.
+	 */
+	UPROPERTY(EditAnywhere, Category = Cooking, meta=(EditCondition="!IsStandaloneGraph() && !IsEmbeddedSubgraph()", EditConditionHides))
+	FPerPlatformBool ShouldCook = true;
 
 #if WITH_EDITORONLY_DATA
 	/** When true the Debug flag in the graph editor will display debug information contextually for the selected debug object. Otherwise
@@ -792,6 +802,9 @@ public:
 	double GetGridCleanupRadiusFromGrid(EPCGHiGenGrid Grid) const { return GetGridCleanupRadiusFromGrid(PCGHiGenGrid::GridToGridSize(Grid)); }
 	UE_DEPRECATED(5.8, "Use GetGridSizeMultiplier instead")
 	uint32 GetGridExponential() const { return FMath::RoundToInt(FMath::Log2(FMath::Max(HiGenGridSizeMultiplier, UE_KINDA_SMALL_NUMBER))); }
+	UE_DEPRECATED(5.8, "Use ShouldCook instead")
+	UPROPERTY()
+	bool bIsEditorOnly = false;
 protected:
 	UE_DEPRECATED(5.8, "Use HiGenGridSizeMultiplier instead")
 	UPROPERTY()
@@ -815,6 +828,7 @@ public:
 	UE_API virtual void PostEditImport() override;
 	UE_API virtual void BeginDestroy() override;
 	virtual bool IsEditorOnly() const override { return GetGraph() && GetGraph()->IsEditorOnly(); }
+	UE_API virtual bool NeedsLoadForTargetPlatform(const ITargetPlatform* TargetPlatform) const override;
 
 #if WITH_EDITOR
 	UE_API virtual void PreEditChange(FProperty* InProperty) override;
