[PCG] Move CreateExecutionSource to IPCGBaseSubsystem as a static call
- This avoid creating a source in the wrong subsystem, the source once created will return the proper Subsystem to use

#Synced-CL 51758684

diff --git a/Engine/Plugins/Experimental/ProceduralVegetationEditor/Source/ProceduralVegetationEditor/Private/PVEditor.cpp b/Engine/Plugins/Experimental/ProceduralVegetationEditor/Source/ProceduralVegetationEditor/Private/PVEditor.cpp
--- a/Engine/Plugins/Experimental/ProceduralVegetationEditor/Source/ProceduralVegetationEditor/Private/PVEditor.cpp
+++ b/Engine/Plugins/Experimental/ProceduralVegetationEditor/Source/ProceduralVegetationEditor/Private/PVEditor.cpp
@@ -90,17 +90,12 @@ void FPVEditor::Initialize(
 
 	GetInspectionDataManager().OnInspectedStackChangedDelegate.AddRaw(this, &FPVEditor::OnNodeInspected);
 
-	UPCGEngineSubsystem* Subsystem = GEngine->GetEngineSubsystem<UPCGEngineSubsystem>();
+	ExecutionSource = IPCGBaseSubsystem::CreateExecutionSource<UPCGDefaultExecutionSource>({ .GraphInterface = ProceduralVegetationGraph });
+
+	IPCGBaseSubsystem* Subsystem = ExecutionSource->GetExecutionState().GetSubsystem();
 	check(Subsystem);
 
 	Subsystem->GetOnPCGSourceGenerationDone().AddRaw(this, &FPVEditor::OnGraphExecuted);
-	
-	const FPCGDefaultExecutionSourceParams Params
-	{
-		.GraphInterface = ProceduralVegetationGraph
-	};
-
-	ExecutionSource = Subsystem->CreateExecutionSource(Params);
 
 	// Set the stack
 	FPCGStack Stack;
@@ -190,7 +185,7 @@ void FPVEditor::OnClose()
 
 	GetInspectionDataManager().OnInspectedStackChangedDelegate.RemoveAll(this);
 
-	if (UPCGEngineSubsystem* Subsystem = GEngine->GetEngineSubsystem<UPCGEngineSubsystem>())
+	if (IPCGBaseSubsystem* Subsystem = ExecutionSource ? ExecutionSource->GetExecutionState().GetSubsystem() : nullptr)
 	{
 		Subsystem->GetOnPCGSourceGenerationDone().RemoveAll(this);
 	}
diff --git a/Engine/Plugins/Experimental/ProceduralVegetationEditor/Source/ProceduralVegetationEditor/Private/Tests/ScenarioTests/PVScenarioTests.cpp b/Engine/Plugins/Experimental/ProceduralVegetationEditor/Source/ProceduralVegetationEditor/Private/Tests/ScenarioTests/PVScenarioTests.cpp
--- a/Engine/Plugins/Experimental/ProceduralVegetationEditor/Source/ProceduralVegetationEditor/Private/Tests/ScenarioTests/PVScenarioTests.cpp
+++ b/Engine/Plugins/Experimental/ProceduralVegetationEditor/Source/ProceduralVegetationEditor/Private/Tests/ScenarioTests/PVScenarioTests.cpp
@@ -315,9 +315,6 @@ namespace PVScenarioTests
 
 	static TSharedRef<FStrongExecutionSourcePtrArray> CreateScenariosPCGExecutionSources(const TArray<TWeakObjectPtr<UPVTestScenario>>& InTestScenarios)
 	{
-		UPCGEngineSubsystem* PCGSubsystem = UPCGEngineSubsystem::Get();
-		check(PCGSubsystem);
-
 		TArray<UPCGDefaultExecutionSource*> OutExecutionSources;
 		OutExecutionSources.SetNum(InTestScenarios.Num());
 
@@ -341,11 +338,7 @@ namespace PVScenarioTests
 				continue;
 			}
 
-			const FPCGDefaultExecutionSourceParams Params
-			{
-				.GraphInterface = Graph
-			};
-			OutExecutionSources[i] = PCGSubsystem->CreateExecutionSource(Params);
+			OutExecutionSources[i] = IPCGBaseSubsystem::CreateExecutionSource<UPCGDefaultExecutionSource>({ .GraphInterface = Graph });
 			OutExecutionSources[i]->GetExecutionState().GetInspection().EnableInspection();
 		}
 
diff --git a/Engine/Plugins/PCG/Source/PCG/Private/PCGDefaultActorExecutionSource.cpp b/Engine/Plugins/PCG/Source/PCG/Private/PCGDefaultActorExecutionSource.cpp
--- a/Engine/Plugins/PCG/Source/PCG/Private/PCGDefaultActorExecutionSource.cpp
+++ b/Engine/Plugins/PCG/Source/PCG/Private/PCGDefaultActorExecutionSource.cpp
@@ -149,7 +149,7 @@ UPCGDefaultActorExecutionSource::UPCGDefaultActorExecutionSource()
 	State = MakeUnique<FPCGDefaultActorExecutionState>(this);
 }
 
-void UPCGDefaultActorExecutionSource::Initialize(const FPCGDefaultActorExecutionSourceParams& InParams)
+void UPCGDefaultActorExecutionSource::Initialize(const ParamsType& InParams)
 {
 	UPCGDefaultExecutionSource::Initialize(InParams);
 	SetActor(InParams.Actor);
diff --git a/Engine/Plugins/PCG/Source/PCG/Private/PCGDefaultExecutionSource.cpp b/Engine/Plugins/PCG/Source/PCG/Private/PCGDefaultExecutionSource.cpp
--- a/Engine/Plugins/PCG/Source/PCG/Private/PCGDefaultExecutionSource.cpp
+++ b/Engine/Plugins/PCG/Source/PCG/Private/PCGDefaultExecutionSource.cpp
@@ -263,7 +263,7 @@ UPCGDefaultExecutionSource::~UPCGDefaultExecutionSource()
 	Sunset();
 }
 
-void UPCGDefaultExecutionSource::Initialize(const FPCGDefaultExecutionSourceParams& InParams)
+void UPCGDefaultExecutionSource::Initialize(const ParamsType& InParams)
 {
 	SetGraphInterface(InParams.GraphInterface);
 	SetSeed(InParams.Seed);
diff --git a/Engine/Plugins/PCG/Source/PCG/Private/Subsystems/PCGEngineSubsystem.cpp b/Engine/Plugins/PCG/Source/PCG/Private/Subsystems/PCGEngineSubsystem.cpp
--- a/Engine/Plugins/PCG/Source/PCG/Private/Subsystems/PCGEngineSubsystem.cpp
+++ b/Engine/Plugins/PCG/Source/PCG/Private/Subsystems/PCGEngineSubsystem.cpp
@@ -73,7 +73,13 @@ void UPCGEngineSubsystem::GenerateGraph(const FPCGGenerateGraphParams& InParams)
 		return;
 	}
 
-	CreateExecutionSource({ InParams.Graph, InParams.Seed, /*bFireAndForget=*/true, InParams.GenerationCallback });
+	IPCGBaseSubsystem::CreateExecutionSource<UPCGDefaultExecutionSource>({ InParams.Graph, InParams.Seed, /*bFireAndForget=*/true, InParams.GenerationCallback });
+}
+
+// deprecated
+UPCGDefaultExecutionSource* UPCGEngineSubsystem::CreateExecutionSource(const FPCGDefaultExecutionSourceParams& InParams)
+{
+	return IPCGBaseSubsystem::CreateExecutionSource<UPCGDefaultExecutionSource>(InParams);
 }
 
 #endif // WITH_EDITOR
diff --git a/Engine/Plugins/PCG/Source/PCG/Public/PCGDefaultActorExecutionSource.h b/Engine/Plugins/PCG/Source/PCG/Public/PCGDefaultActorExecutionSource.h
--- a/Engine/Plugins/PCG/Source/PCG/Public/PCGDefaultActorExecutionSource.h
+++ b/Engine/Plugins/PCG/Source/PCG/Public/PCGDefaultActorExecutionSource.h
@@ -22,8 +22,10 @@ class UPCGDefaultActorExecutionSource : public UPCGDefaultWorldObjectExecutionSo
 	GENERATED_BODY()
 
 public:
+	using ParamsType = FPCGDefaultActorExecutionSourceParams;
+
 	PCG_API UPCGDefaultActorExecutionSource();
-	PCG_API void Initialize(const FPCGDefaultActorExecutionSourceParams& InParams);
+	PCG_API void Initialize(const ParamsType& InParams);
 	PCG_API AActor* GetActor();
 	PCG_API void SetActor(AActor* InActor);
 };
diff --git a/Engine/Plugins/PCG/Source/PCG/Public/PCGDefaultExecutionSource.h b/Engine/Plugins/PCG/Source/PCG/Public/PCGDefaultExecutionSource.h
--- a/Engine/Plugins/PCG/Source/PCG/Public/PCGDefaultExecutionSource.h
+++ b/Engine/Plugins/PCG/Source/PCG/Public/PCGDefaultExecutionSource.h
@@ -83,10 +83,12 @@ class UPCGDefaultExecutionSource : public UObject, public IPCGGraphExecutionSour
 	GENERATED_BODY()
 	
 public:
+	using ParamsType = FPCGDefaultExecutionSourceParams;
+
 	PCG_API UPCGDefaultExecutionSource();
 	PCG_API ~UPCGDefaultExecutionSource();
 
-	PCG_API void Initialize(const FPCGDefaultExecutionSourceParams& InParams);
+	PCG_API void Initialize(const ParamsType& InParams);
 
 	PCG_API static void AddReferencedObjects(UObject* InThis, FReferenceCollector& Collector);
 	
diff --git a/Engine/Plugins/PCG/Source/PCG/Public/Subsystems/IPCGBaseSubsystem.h b/Engine/Plugins/PCG/Source/PCG/Public/Subsystems/IPCGBaseSubsystem.h
--- a/Engine/Plugins/PCG/Source/PCG/Public/Subsystems/IPCGBaseSubsystem.h
+++ b/Engine/Plugins/PCG/Source/PCG/Public/Subsystems/IPCGBaseSubsystem.h
@@ -28,7 +28,6 @@ class UWorld;
 struct FPCGDataCollection;
 struct FPCGScheduleGenericParams;
 struct FPCGStack;
-struct FPCGDefaultExecutionSourceParams;
 
 #if WITH_EDITOR
 DECLARE_MULTICAST_DELEGATE_ThreeParams(FPCGOnPCGSourceGenerationDone, IPCGBaseSubsystem*, IPCGGraphExecutionSource*, EPCGGenerationStatus);
@@ -169,21 +168,24 @@ public:
 	FPCGOnPCGSourceUnregistered& GetOnPCGSourceUnregistered() { return OnPCGSourceUnregisteredDelegate; }
 
 	/** Creates an execution source from the params and optionally tracks its lifetime and generation. Note this call is NOT THREADSAFE and as such this should be called only from the game thread. */
-	template<typename SourceType, typename ParamsType>
+	template<typename SourceType>
 		requires std::is_base_of_v<IPCGGraphExecutionSource, SourceType>
-	SourceType* CreateExecutionSource(const ParamsType& InParams, UObject* Outer)
+	static SourceType* CreateExecutionSource(const typename SourceType::ParamsType& InParams, UObject* Outer = nullptr)
 	{
 		check(IsInGameThread());
-		if (!InParams.GraphInterface || !Outer)
+		if (!InParams.GraphInterface)
 		{
 			return nullptr;
 		}
 
-		SourceType* NewSource = NewObject<SourceType>(Outer);
+		SourceType* NewSource = NewObject<SourceType>(Outer ? Outer : GetTransientPackageAsObject());
 		NewSource->Initialize(InParams);
 
 		if (InParams.bFireAndForgetExecution)
 		{
+			IPCGBaseSubsystem* Subsystem = NewSource->GetExecutionState().GetSubsystem();
+			check(Subsystem);
+			
 			FPCGOnPCGSourceGenerationDone GenerationCallbackWrapper;
 			if (InParams.GenerationCallback.IsBound())
 			{
@@ -193,7 +195,7 @@ public:
 				});
 			}
 
-			GraphExecutions.Emplace(NewSource, GenerationCallbackWrapper);
+			Subsystem->GraphExecutions.Emplace(NewSource, GenerationCallbackWrapper);
 			NewSource->Generate();
 		}
 
diff --git a/Engine/Plugins/PCG/Source/PCG/Public/Subsystems/PCGEngineSubsystem.h b/Engine/Plugins/PCG/Source/PCG/Public/Subsystems/PCGEngineSubsystem.h
--- a/Engine/Plugins/PCG/Source/PCG/Public/Subsystems/PCGEngineSubsystem.h
+++ b/Engine/Plugins/PCG/Source/PCG/Public/Subsystems/PCGEngineSubsystem.h
@@ -45,17 +45,12 @@ public:
 	PCG_API virtual void Initialize(FSubsystemCollectionBase& Collection) override;
 	PCG_API virtual void Deinitialize() override;
 
-	/** Creates an execution source from the params and optionally tracks its lifetime and generation. Note this call is NOT THREADSAFE and as such this should be called only from the game thread. */
-	template<typename SourceType = UPCGDefaultExecutionSource, typename ParamsType = FPCGDefaultExecutionSourceParams> 
-		requires std::is_base_of_v<IPCGGraphExecutionSource, SourceType> && std::is_base_of_v<FPCGDefaultExecutionSourceParams, ParamsType>
-	SourceType* CreateExecutionSource(const ParamsType& InParams)
-	{
-		return IPCGBaseSubsystem::CreateExecutionSource<SourceType, ParamsType>(InParams, this);
-	}
-
 	PCG_API void GenerateGraph(const FPCGGenerateGraphParams& InParams);
 
 	PCG_API virtual void NotifyGraphChanged(UPCGGraph* InGraph, EPCGChangeType ChangeType) override;
+
+	UE_DEPRECATED(5.8, "Use static IPCGBaseSubsystem version instead")
+	PCG_API UPCGDefaultExecutionSource* CreateExecutionSource(const FPCGDefaultExecutionSourceParams& InParams);
 #endif // WITH_EDITOR
 
 protected:
diff --git a/Engine/Plugins/PCG/Source/PCGEditor/Private/PCGEditor.cpp b/Engine/Plugins/PCG/Source/PCGEditor/Private/PCGEditor.cpp
--- a/Engine/Plugins/PCG/Source/PCGEditor/Private/PCGEditor.cpp
+++ b/Engine/Plugins/PCG/Source/PCGEditor/Private/PCGEditor.cpp
@@ -4229,11 +4229,8 @@ void FPCGEditor::UpdateDefaultExecutionSource()
 {
 	if (MainGraph->IsStandaloneGraph() && !PCGDefaultExecutionSource)
 	{
-		if (UPCGEngineSubsystem* EngineSubsystem = UPCGEngineSubsystem::Get())
-		{
-			PCGDefaultExecutionSource = EngineSubsystem->CreateExecutionSource({ MainGraph });
-			RefreshViews();
-		}
+		PCGDefaultExecutionSource = IPCGBaseSubsystem::CreateExecutionSource<UPCGDefaultExecutionSource>({ .GraphInterface = MainGraph });
+		RefreshViews();
 	}
 	else if (!MainGraph->IsStandaloneGraph() && PCGDefaultExecutionSource)
 	{
