[PCG] Tool Mode - Draw Spline: No requirement for graph for easier automation use

#Synced-CL 51835013

diff --git a/Engine/Plugins/PCG/Source/PCGEditor/Private/EditorMode/PCGEdMode.cpp b/Engine/Plugins/PCG/Source/PCGEditor/Private/EditorMode/PCGEdMode.cpp
--- a/Engine/Plugins/PCG/Source/PCGEditor/Private/EditorMode/PCGEdMode.cpp
+++ b/Engine/Plugins/PCG/Source/PCGEditor/Private/EditorMode/PCGEdMode.cpp
@@ -83,12 +83,12 @@ namespace PCGEditorMode
 			UPCGInteractiveToolSettings* ToolSettings = EditorMode ? Cast<UPCGInteractiveToolSettings>(EditorMode->GetCurrentToolSettings()) : nullptr;
 			if (ToolSettings && !Args.IsEmpty())
 			{
-				FSoftObjectPath GraphPath(Args[0]);
-				UObject* GraphObject = GraphPath.TryLoad();
-				if (UPCGGraphInterface* Graph = Cast<UPCGGraphInterface>(GraphObject))
-				{
-					ToolSettings->SetToolGraph(Graph);
-				}
+				// If the input path is empty, or not pointing to a graph, the tool graph will be set to nullptr.
+				const FString& GraphPathString = Args[0];
+				const FSoftObjectPath GraphPath(GraphPathString);
+				UPCGGraphInterface* Graph = Cast<UPCGGraphInterface>(GraphPath.TryLoad());
+				ToolSettings->SetToolGraph(Graph);
+
 				if (Args.Num() > 1)
 				{
 					ToolSettings->SetActorLabel(FName(Args[1]));
