[PCG] Allow auto layout of graph nodes to use the current node selection.

#Synced-CL 51889095

diff --git a/Engine/Plugins/PCG/Source/PCGEditor/Private/PCGEditor.cpp b/Engine/Plugins/PCG/Source/PCGEditor/Private/PCGEditor.cpp
--- a/Engine/Plugins/PCG/Source/PCGEditor/Private/PCGEditor.cpp
+++ b/Engine/Plugins/PCG/Source/PCGEditor/Private/PCGEditor.cpp
@@ -824,7 +824,6 @@ void FPCGEditor::RegisterToolbarInternal(FToolBarBuilder& ToolbarBuilder) const
 	RegisterToolbarButton(ToolbarBuilder, EPCGToolbarButtons::Find);
 	RegisterToolbarButton(ToolbarBuilder, EPCGToolbarButtons::PauseRegen);
 	RegisterToolbarButton(ToolbarBuilder, EPCGToolbarButtons::ForceRegen);
-	RegisterToolbarButton(ToolbarBuilder, EPCGToolbarButtons::AutoLayoutNodes);
 	RegisterToolbarButton(ToolbarBuilder, EPCGToolbarButtons::CancelExecution);
 	RegisterToolbarButton(ToolbarBuilder, EPCGToolbarButtons::OpenDebugObjectTreeTab);
 
@@ -833,6 +832,7 @@ void FPCGEditor::RegisterToolbarInternal(FToolBarBuilder& ToolbarBuilder) const
 	RegisterToolbarButton(ToolbarBuilder, EPCGToolbarButtons::Graph);
 	RegisterToolbarButton(ToolbarBuilder, EPCGToolbarButtons::GraphParams);
 	RegisterToolbarButton(ToolbarBuilder, EPCGToolbarButtons::GraphSettings);
+	RegisterToolbarButton(ToolbarBuilder, EPCGToolbarButtons::AutoLayoutNodes);
 }
 
 void FPCGEditor::RegisterToolbarButton(FToolBarBuilder& ToolbarBuilder, EPCGToolbarButtons Button) const
@@ -1204,12 +1204,7 @@ void FPCGEditor::OnAutoLayoutNodes_Clicked()
 	{
 		FScopedTransaction NodeMoveTransaction(LOCTEXT("MoveNodesAction", "Auto-Layout Nodes"));
 		
-		if (AutoLayoutFullGraph())
-		{
-			// Force a refresh
-			MainGraph->GetEditorGraph()->ReconstructGraph();
-		}
-		else
+		if (!AutoLayoutFullGraph())
 		{
 			// Do not spam the undo buffer if the graph was already well positioned.
 			NodeMoveTransaction.Cancel();
@@ -3649,18 +3644,46 @@ TSharedPtr<SPCGEditorViewport> FPCGEditor::GetViewportWidget(int32 Index) const
 
 bool FPCGEditor::AutoLayoutFullGraph()
 {
+	TSharedPtr<SGraphEditor> FocusedGraphWidget = GetFocusedGraphEditorWidget();
+	check(FocusedGraphWidget.IsValid());
+
 	UPCGNode* InputNode = MainGraph->InputNode;
 	check(InputNode);
 	UPCGNode* OutputNode = MainGraph->OutputNode;
 	check(OutputNode);
 
-	TArray<UPCGNode*> AllNodes = MainGraph->GetNodes();
-	AllNodes.Add(InputNode);
-	AllNodes.Add(OutputNode);
+	// Check if some nodes are selected first. If so, we only layout these nodes.
+	TArray<UPCGNode*> NodesToLayout;
+	for (UObject* Object : FocusedGraphWidget->GetSelectedNodes())
+	{
+		if (UPCGEditorGraphNode * EditorNode = Cast<UPCGEditorGraphNode>(Object))
+		{
+			if (EditorNode->GetPCGNode())
+			{
+				NodesToLayout.Add(EditorNode->GetPCGNode());
+			}
+		}
+	}
+
+	bool bInputShouldBeLayout = false;
+	bool bOutputShouldBeLayout = false;
+	if (NodesToLayout.IsEmpty())
+	{
+		NodesToLayout = MainGraph->GetNodes();
+		NodesToLayout.Add(InputNode);
+		NodesToLayout.Add(OutputNode);
+		bInputShouldBeLayout = true;
+		bOutputShouldBeLayout = true;
+	}
+	else
+	{
+		bInputShouldBeLayout = NodesToLayout.Contains(InputNode);
+		bOutputShouldBeLayout = NodesToLayout.Contains(OutputNode);
+	}
 	
 	// Precompute predecessor since we use it more than once
 	TMap<UPCGNode*, TSet<UPCGNode*>> PredecessorNodes;
-	for (UPCGNode* Node : AllNodes)
+	for (UPCGNode* Node : NodesToLayout)
 	{
 		TSet<UPCGNode*>& CurPred = PredecessorNodes.Add(Node);
 		for (UPCGPin* Pin : Node->GetInputPins())
@@ -3681,8 +3704,8 @@ bool FPCGEditor::AutoLayoutFullGraph()
 	};
 	
 	// Use the min position as the base position to keep the graph approx. where it is
-	FIntVector2 BasePosition{AllNodes[0]->PositionX, AllNodes[0]->PositionY};
-	for (const UPCGNode* Node : AllNodes)
+	FIntVector2 BasePosition{NodesToLayout[0]->PositionX, NodesToLayout[0]->PositionY};
+	for (const UPCGNode* Node : NodesToLayout)
 	{
 		FIntVector2 CurPosition{Node->PositionX, Node->PositionY};
 		BasePosition.X = FMath::Min(BasePosition.X, CurPosition.X);
@@ -3690,7 +3713,7 @@ bool FPCGEditor::AutoLayoutFullGraph()
 	}
 	
 	bool bNodeMovementOccurred = false;
-	auto SetNodePosition = [&bNodeMovementOccurred](UPCGNode* Node, int32 NewPositionX, int32 NewPositionY)
+	auto SetNodePosition = [&bNodeMovementOccurred, this](UPCGNode* Node, int32 NewPositionX, int32 NewPositionY)
 	{
 		int32 CurPositionX = 0;
 		int32 CurPositionY = 0;
@@ -3700,6 +3723,11 @@ bool FPCGEditor::AutoLayoutFullGraph()
 			bNodeMovementOccurred = true;
 			Node->Modify();
 			Node->SetNodePosition(NewPositionX, NewPositionY);
+			if (UPCGEditorGraphNodeBase* EditorNode = GetEditorNode(Node))
+			{
+				EditorNode->SetNodePosX(NewPositionX);
+				EditorNode->SetNodePosY(NewPositionY);
+			}
 		}
 	};
 	
@@ -3711,38 +3739,41 @@ bool FPCGEditor::AutoLayoutFullGraph()
 	
 	TArray<UPCGNode*> LonelyNodes;
 	TArray<UPCGNode*> SortedAttachedNodes;
-	for (UPCGNode* Node : MainGraph->GetNodes())
+	for (UPCGNode* Node : NodesToLayout)
 	{
-		if (IsLonely(Node))
-		{
-			LonelyNodes.Add(Node);
-		}
-		else
+		if (Node != InputNode && Node != OutputNode)
 		{
-			SortedAttachedNodes.Add(Node);
+			if (IsLonely(Node))
+			{
+				LonelyNodes.Add(Node);
+			}
+			else
+			{
+				SortedAttachedNodes.Add(Node);
+			}
 		}
 	}
 	
 	const bool bInputIsLonely = IsLonely(InputNode);
 	const bool bOutputIsLonely = IsLonely(OutputNode);
 	
-	if (!bInputIsLonely)
+	if (!bInputIsLonely && NodesToLayout.Contains(InputNode))
 	{
 		SortedAttachedNodes.Add(InputNode);
 	}
 	
-	if (!bOutputIsLonely)
+	if (!bOutputIsLonely && NodesToLayout.Contains(OutputNode))
 	{
 		SortedAttachedNodes.Add(OutputNode);
 	}
-	
-	Algo::TopologicalSort(SortedAttachedNodes, GetNodePredecessor);
-	
+
 	// Give space for the input node
-	if (bInputIsLonely)
+	if (bInputIsLonely && bInputShouldBeLayout)
 	{
 		BasePosition.X += (FPCGEditor_private::DefaultFixedNodeWidth + FPCGEditor_private::NodeLayoutHorizontalGap);
 	}
+
+	Algo::TopologicalSort(SortedAttachedNodes, GetNodePredecessor);
 	
 	// Associate a layer to each node
 	int MaxLayerIndex = 0;
@@ -3757,7 +3788,10 @@ bool FPCGEditor::AutoLayoutFullGraph()
 		int MaxLayerOfPred = 0;
 		for (UPCGNode* PredNode: GetNodePredecessor(Node))
 		{
-			MaxLayerOfPred = FMath::Max(MaxLayerOfPred, NodeToLayerIndex[PredNode] + 1);
+			if (NodesToLayout.Contains(PredNode))
+			{
+				MaxLayerOfPred = FMath::Max(MaxLayerOfPred, NodeToLayerIndex[PredNode] + 1);
+			}
 		}
 		NodeToLayerIndex[Node] = MaxLayerOfPred;
 		MaxLayerIndex = FMath::Max(MaxLayerIndex, MaxLayerOfPred);
@@ -3797,12 +3831,12 @@ bool FPCGEditor::AutoLayoutFullGraph()
 		CurrentLayerPosX += MaxNodeWidthForLayer + FPCGEditor_private::NodeLayoutHorizontalGap;
 	}
 	
-	if (bInputIsLonely)
+	if (bInputIsLonely && bInputShouldBeLayout)
 	{
 		SetNodePosition(InputNode, BasePosition.X - (FPCGEditor_private::DefaultFixedNodeWidth + FPCGEditor_private::NodeLayoutHorizontalGap), BasePosition.Y);
 	}
 	
-	if (bOutputIsLonely)
+	if (bOutputIsLonely && bOutputShouldBeLayout)
 	{
 		SetNodePosition(OutputNode, CurrentLayerPosX, BasePosition.Y);
 	}
