[PCG] Added filtering support for subgraphs, blueprint elements and settings, which all require cached pins.

#Synced-CL 51854852

diff --git a/Engine/Plugins/PCG/Source/PCG/Private/Elements/Blueprint/PCGBlueprintBaseElement.cpp b/Engine/Plugins/PCG/Source/PCG/Private/Elements/Blueprint/PCGBlueprintBaseElement.cpp
--- a/Engine/Plugins/PCG/Source/PCG/Private/Elements/Blueprint/PCGBlueprintBaseElement.cpp
+++ b/Engine/Plugins/PCG/Source/PCG/Private/Elements/Blueprint/PCGBlueprintBaseElement.cpp
@@ -12,6 +12,10 @@
 #include "Helpers/PCGHelpers.h"
 #include "Math/RandomStream.h"
 
+#if WITH_EDITOR
+#include "UObject/ObjectSaveContext.h"
+#endif // WITH_EDITOR
+
 #include UE_INLINE_GENERATED_CPP_BY_NAME(PCGBlueprintBaseElement)
 
 #define LOCTEXT_NAMESPACE "PCGBlueprintBaseElement"
@@ -63,6 +67,20 @@ void UPCGBlueprintBaseElement::BeginDestroy()
 	Super::BeginDestroy();
 }
 
+#if WITH_EDITOR
+void UPCGBlueprintBaseElement::PreSave(FObjectPreSaveContext ObjectSaveContext)
+{
+	Super::PreSave(ObjectSaveContext);
+
+	// Cache pins for asset retrieval
+	{
+		CachedPins.InputPins = GetInputPins();
+		CachedPins.OutputPins = GetOutputPins();
+		CachedPins.bHasCachedPins = true;
+	}
+}
+#endif // WITH_EDITOR
+
 void UPCGBlueprintBaseElement::Execute_Implementation(const FPCGDataCollection& Input, FPCGDataCollection& Output)
 {
 }
diff --git a/Engine/Plugins/PCG/Source/PCG/Private/Helpers/PCGAssetHelpers.cpp b/Engine/Plugins/PCG/Source/PCG/Private/Helpers/PCGAssetHelpers.cpp
new file mode 100644
--- /dev/null
+++ b/Engine/Plugins/PCG/Source/PCG/Private/Helpers/PCGAssetHelpers.cpp
@@ -0,0 +1,145 @@
+// Copyright Epic Games, Inc. All Rights Reserved.
+
+#include "Helpers/PCGAssetHelpers.h"
+
+#include "PCGGraph.h"
+#include "Elements/Blueprint/PCGBlueprintBaseElement.h"
+
+#include "AssetRegistry/AssetData.h"
+#include "AssetRegistry/AssetRegistryModule.h"
+#include "Blueprint/BlueprintSupport.h"
+
+#if WITH_EDITOR
+
+bool FPCGAssetHelpers::GetGraphAssetRegistryData(const FAssetData& AssetData, FPCGAssetHelpers::FGraphAssetOutput& Output)
+{
+	FAssetRegistryModule& AssetRegistryModule = FModuleManager::Get().LoadModuleChecked<FAssetRegistryModule>(TEXT("AssetRegistry"));
+
+	// Only exposing the instance if it's parent graph is defined for instances. Otherwise, it is not interesting.
+	// Also march up the hierarchy to find overrides for category and description.
+	// We don't look for override in titles because if we have no override for an instance but the parent has an override we will have 2 times the same entry in the palette.
+	auto GetRecursiveData = [&AssetRegistryModule, &Output](const FAssetData& InAssetData, auto&& Recurse)
+	{
+		if (InAssetData.IsInstanceOf<UPCGGraphInstance>())
+		{
+			const FSoftObjectPath ParentGraph(InAssetData.GetTagValueRef<FString>(GET_MEMBER_NAME_CHECKED(UPCGGraphInstance, Graph)));
+			if (ParentGraph.IsNull())
+			{
+				return false;
+			}
+
+			if(Output.Category.IsEmpty())
+			{
+				Output.Category = InAssetData.GetTagValueRef<bool>(GET_MEMBER_NAME_CHECKED(UPCGGraphInstance, bOverrideCategory))
+					? InAssetData.GetTagValueRef<FText>(GET_MEMBER_NAME_CHECKED(UPCGGraphInstance, Category))
+					: FText();
+			}
+
+			if(Output.Description.IsEmpty())
+			{
+				Output.Description = InAssetData.GetTagValueRef<bool>(GET_MEMBER_NAME_CHECKED(UPCGGraphInstance, bOverrideDescription))
+					? InAssetData.GetTagValueRef<FText>(GET_MEMBER_NAME_CHECKED(UPCGGraphInstance, Description))
+					: FText();
+			}
+
+			// Asset data are not big so that should not be that big of a deal, but they are copied all the time.
+			// If we ever have performances issues, might be good to have a cache.
+			const FAssetData ParentAssetData = AssetRegistryModule.Get().GetAssetByObjectPath(ParentGraph);
+			return ParentAssetData.IsValid() && Recurse(ParentAssetData, Recurse);
+		}
+		else
+		{
+			if (Output.Category.IsEmpty())
+			{
+				Output.Category = InAssetData.GetTagValueRef<FText>(GET_MEMBER_NAME_CHECKED(UPCGGraph, Category));
+			}
+
+			if (Output.Description.IsEmpty())
+			{
+				Output.Description = InAssetData.GetTagValueRef<FText>(GET_MEMBER_NAME_CHECKED(UPCGGraph, Description));
+			}
+
+			// Implementation note: pins are only present at the graph level
+			FString TempString;
+			InAssetData.GetTagValue<FString>("CachedPins", TempString); // can't use GET_MEMBER_NAME_CHECKED because CachedPins are private.
+			if (!TempString.IsEmpty())
+			{
+				FPCGAssetCachedPins::StaticStruct()->ImportText(*TempString, &Output.CachedPins, nullptr, 0, nullptr, FString{});
+			}
+
+			return true;
+		}
+	};
+
+	return GetRecursiveData(AssetData, GetRecursiveData);
+}
+
+bool FPCGAssetHelpers::GetBlueprintAssetRegistryData(const FAssetData& AssetData, FPCGAssetHelpers::FBlueprintAssetOutput& Output)
+{
+	//@todo_pcg: validate that the asset is indeed of a BP class. See PCGEditorUtils::ForEachPCGBlueprintAssetData?
+	Output.GeneratedClass = FSoftClassPath(AssetData.GetTagValueRef<FString>(FBlueprintTags::GeneratedClassPath));
+	Output.Category = AssetData.GetTagValueRef<FText>(GET_MEMBER_NAME_CHECKED(UPCGBlueprintBaseElement, Category));
+	Output.Description = AssetData.GetTagValueRef<FText>(GET_MEMBER_NAME_CHECKED(UPCGBlueprintBaseElement, Description));
+	Output.bOnlyExposePreconfiguredSettings = AssetData.GetTagValueRef<bool>(GET_MEMBER_NAME_CHECKED(UPCGBlueprintBaseElement, bOnlyExposePreconfiguredSettings));
+	Output.bEnabledPreconfiguredSettings = AssetData.GetTagValueRef<bool>(GET_MEMBER_NAME_CHECKED(UPCGBlueprintBaseElement, bEnablePreconfiguredSettings));
+
+	FString TempString;
+	AssetData.GetTagValue<FString>("CachedPins", TempString); // note can't use GET_MEMBER_NAME_CHECKED because field is private
+	if (!TempString.IsEmpty())
+	{
+		FPCGAssetCachedPins::StaticStruct()->ImportText(*TempString, &Output.CachedPins, nullptr, 0, nullptr, FString{});
+	}
+
+	return true;
+}
+
+bool FPCGAssetHelpers::GetSettingsAssetRegistryData(const FAssetData& AssetData, FPCGAssetHelpers::FSettingsAssetOutput& Output)
+{
+	FAssetRegistryModule& AssetRegistryModule = FModuleManager::Get().LoadModuleChecked<FAssetRegistryModule>(TEXT("AssetRegistry"));
+
+	// Implementation note - this supports settings + instances of settings.
+	// However it will not support instance of instances, but the instance class isn't built this way yet.
+	auto ProcessSettings = [&AssetRegistryModule, &Output](const FAssetData& InAssetData)
+	{
+		if(InAssetData.IsInstanceOf<UPCGSettings>())
+		{
+			Output.Category = InAssetData.GetTagValueRef<FText>(GET_MEMBER_NAME_CHECKED(UPCGSettings, Category));
+			Output.Description = InAssetData.GetTagValueRef<FText>(GET_MEMBER_NAME_CHECKED(UPCGSettings, Description));
+
+			FString TempString;
+			InAssetData.GetTagValue<FString>("CachedPins", TempString); // can't use GET_MEMBER_NAME_CHECKED because CachedPins are private.
+			if (!TempString.IsEmpty())
+			{
+				FPCGAssetCachedPins::StaticStruct()->ImportText(*TempString, &Output.CachedPins, nullptr, 0, nullptr, FString{});
+			}
+
+			return true;
+		}
+		else
+		{
+			return false;
+		}
+	};
+
+	if (AssetData.IsInstanceOf<UPCGSettings>())
+	{
+		return ProcessSettings(AssetData);
+	}
+	else if (AssetData.IsInstanceOf<UPCGSettingsInstance>())
+	{
+		const FSoftObjectPath OriginalSettingsPath(AssetData.GetTagValueRef<FString>(GET_MEMBER_NAME_CHECKED(UPCGSettingsInstance, Settings)));
+		if (OriginalSettingsPath.IsNull())
+		{
+			return false; // can't do anything here unfortunately
+		}
+
+		const FAssetData OriginalAssetData = AssetRegistryModule.Get().GetAssetByObjectPath(OriginalSettingsPath);
+		return OriginalAssetData.IsValid() && ProcessSettings(OriginalAssetData);
+	}
+	else
+	{
+		return false;
+	}
+}
+
+#endif // WITH_EDITOR
\ No newline at end of file
diff --git a/Engine/Plugins/PCG/Source/PCG/Private/PCGGraph.cpp b/Engine/Plugins/PCG/Source/PCG/Private/PCGGraph.cpp
--- a/Engine/Plugins/PCG/Source/PCG/Private/PCGGraph.cpp
+++ b/Engine/Plugins/PCG/Source/PCG/Private/PCGGraph.cpp
@@ -838,6 +838,15 @@ void UPCGGraph::PreSave(FObjectPreSaveContext ObjectSaveContext)
 		PCGEditorModule->OnGraphPreSave(this, ObjectSaveContext);
 	}
 
+	// Cache pins for asset retrieval
+	if(IsAsset())
+	{
+		CachedPins.InputPins = InputNode ? CastChecked<UPCGGraphInputOutputSettings>(InputNode->GetSettings())->OutputPinProperties() : TArray<FPCGPinProperties>();
+		CachedPins.OutputPins = OutputNode ? CastChecked<UPCGGraphInputOutputSettings>(OutputNode->GetSettings())->InputPinProperties() : TArray<FPCGPinProperties>();
+		CachedPins.bHasCachedPins = true;
+	}
+
+	// Prepare cook data
 	if (ObjectSaveContext.IsCooking() && !IsEditorOnly())
 	{
 		FPCGGraphCompiler GraphCompiler(/*bIsCooking=*/true);
diff --git a/Engine/Plugins/PCG/Source/PCG/Private/PCGSettings.cpp b/Engine/Plugins/PCG/Source/PCG/Private/PCGSettings.cpp
--- a/Engine/Plugins/PCG/Source/PCG/Private/PCGSettings.cpp
+++ b/Engine/Plugins/PCG/Source/PCG/Private/PCGSettings.cpp
@@ -320,6 +320,25 @@ void UPCGSettings::PostEditUndo()
 
 	Super::PostEditUndo();
 }
+
+void UPCGSettings::PreSave(FObjectPreSaveContext ObjectSaveContext)
+{
+	Super::PreSave(ObjectSaveContext);
+
+	// Update cached pins - note that we need to cache pins only for asset-level settings. Otherwise, make them always empty.
+	if(IsAsset())
+	{
+		CachedPins.InputPins = InputPinProperties();
+		CachedPins.OutputPins = OutputPinProperties();
+		CachedPins.bHasCachedPins = true;
+	}
+	else
+	{
+		CachedPins.InputPins.Reset();
+		CachedPins.OutputPins.Reset();
+		CachedPins.bHasCachedPins = false;
+	}
+}
 #endif // WITH_EDITOR
 
 void UPCGSettings::AddToCrc(FArchiveCrc32& Ar, bool bFullDataCrc) const
diff --git a/Engine/Plugins/PCG/Source/PCG/Public/Elements/Blueprint/PCGBlueprintBaseElement.h b/Engine/Plugins/PCG/Source/PCG/Public/Elements/Blueprint/PCGBlueprintBaseElement.h
--- a/Engine/Plugins/PCG/Source/PCG/Public/Elements/Blueprint/PCGBlueprintBaseElement.h
+++ b/Engine/Plugins/PCG/Source/PCG/Public/Elements/Blueprint/PCGBlueprintBaseElement.h
@@ -4,6 +4,7 @@
 
 #include "PCGContext.h"
 #include "PCGSettings.h"
+#include "Helpers/PCGAssetHelpers.h"
 
 #include "PCGBlueprintBaseElement.generated.h"
 
@@ -29,6 +30,9 @@ public:
 	// ~Begin UObject interface
 	PCG_API virtual void PostLoad() override;
 	PCG_API virtual void BeginDestroy() override;
+#if WITH_EDITOR
+	PCG_API virtual void PreSave(FObjectPreSaveContext ObjectSaveContext) override;
+#endif
 	// ~End UObject interface
 
 	/**
@@ -178,6 +182,11 @@ public:
 
 	UPROPERTY(BlueprintReadWrite, EditDefaultsOnly, AdvancedDisplay, Category = "Settings")
 	int32 DependencyParsingDepth = 1;
+
+private:
+	/** Cached pin data to get type information without loading. */
+	UPROPERTY(AssetRegistrySearchable)
+	FPCGAssetCachedPins CachedPins;
 #endif
 
 protected:
diff --git a/Engine/Plugins/PCG/Source/PCG/Public/Helpers/PCGAssetHelpers.h b/Engine/Plugins/PCG/Source/PCG/Public/Helpers/PCGAssetHelpers.h
new file mode 100644
--- /dev/null
+++ b/Engine/Plugins/PCG/Source/PCG/Public/Helpers/PCGAssetHelpers.h
@@ -0,0 +1,63 @@
+// Copyright Epic Games, Inc. All Rights Reserved.
+
+#pragma once
+
+#include "PCGPin.h"
+
+#include "PCGAssetHelpers.generated.h"
+
+struct FAssetData;
+
+// Struct to hold asset registry searchable data (w/ deprecation support) to be used in graphs and blueprints.
+// Modeled after the pattern in PCG Data Assets.
+USTRUCT(MinimalAPI)
+struct FPCGAssetCachedPins
+{
+	GENERATED_BODY()
+
+	/** Cached input pins to be able to know if a graph is compatible with a specific type without loading it first. */
+	UPROPERTY()
+	TArray<FPCGPinProperties> InputPins;
+
+	/** Cached output pins to be able to know if a graph is compatible with a specific type without loading it first. */
+	UPROPERTY()
+	TArray<FPCGPinProperties> OutputPins;
+
+	/** For deprecation, to know if the asset has cached pins. */
+	UPROPERTY()
+	bool bHasCachedPins = false;
+};
+
+#if WITH_EDITOR
+class FPCGAssetHelpers
+{
+public:
+	struct FGraphAssetOutput
+	{
+		FText Category;
+		FText Description;
+		FPCGAssetCachedPins CachedPins;
+	};
+
+	struct FBlueprintAssetOutput
+	{
+		FText Category;
+		FText Description;
+		FPCGAssetCachedPins CachedPins;
+		FSoftClassPath GeneratedClass;
+		bool bOnlyExposePreconfiguredSettings = false;
+		bool bEnabledPreconfiguredSettings = false;
+	};
+
+	struct FSettingsAssetOutput
+	{
+		FText Category;
+		FText Description;
+		FPCGAssetCachedPins CachedPins;
+	};
+
+	static PCG_API bool GetGraphAssetRegistryData(const FAssetData& AssetData, FGraphAssetOutput& Output);
+	static PCG_API bool GetBlueprintAssetRegistryData(const FAssetData& AssetData, FBlueprintAssetOutput& Output);
+	static PCG_API bool GetSettingsAssetRegistryData(const FAssetData& AssetData, FSettingsAssetOutput& Output);
+};
+#endif
\ No newline at end of file
diff --git a/Engine/Plugins/PCG/Source/PCG/Public/PCGGraph.h b/Engine/Plugins/PCG/Source/PCG/Public/PCGGraph.h
--- a/Engine/Plugins/PCG/Source/PCG/Public/PCGGraph.h
+++ b/Engine/Plugins/PCG/Source/PCG/Public/PCGGraph.h
@@ -9,6 +9,7 @@
 #include "Data/PCGToolData.h"
 #include "Editor/PCGGraphCustomization.h"
 #include "Graph/PCGStackContext.h"
+#include "Helpers/PCGAssetHelpers.h"
 #include "Helpers/PCGGraphParameterExtension.h"
 
 #if WITH_EDITOR
@@ -729,6 +730,13 @@ protected:
 	UPROPERTY()
 	TObjectPtr<UPCGGraphCompilationData> CookedCompilationData = nullptr;
 
+private:
+#if WITH_EDITORONLY_DATA
+	/** Cached pin data to get type information without loading. */
+	UPROPERTY(AssetRegistrySearchable)
+	FPCGAssetCachedPins CachedPins;
+#endif // WITH_EDITORONLY_DATA
+
 public:
 	virtual const FInstancedPropertyBag* GetUserParametersStruct() const override { return &UserParameters; }
 
diff --git a/Engine/Plugins/PCG/Source/PCG/Public/PCGInputOutputSettings.h b/Engine/Plugins/PCG/Source/PCG/Public/PCGInputOutputSettings.h
--- a/Engine/Plugins/PCG/Source/PCG/Public/PCGInputOutputSettings.h
+++ b/Engine/Plugins/PCG/Source/PCG/Public/PCGInputOutputSettings.h
@@ -30,6 +30,10 @@ class UPCGGraphInputOutputSettings : public UPCGSettings
 {
 	GENERATED_BODY()
 
+	// @todo_pcg: This is needed because we need the pin properties when caching them in the graph. However, if we remove the override on the Default versions
+	// this wouldn't be needed anymore.
+	friend class UPCGGraph; 
+
 public:
 	// ~Begin UObject interface
 	virtual void PostLoad() override;
diff --git a/Engine/Plugins/PCG/Source/PCG/Public/PCGSettings.h b/Engine/Plugins/PCG/Source/PCG/Public/PCGSettings.h
--- a/Engine/Plugins/PCG/Source/PCG/Public/PCGSettings.h
+++ b/Engine/Plugins/PCG/Source/PCG/Public/PCGSettings.h
@@ -8,6 +8,7 @@
 #include "PCGPin.h"
 #include "Data/Registry/PCGDataType.h"
 #include "Elements/PCGActorSelector.h"
+#include "Helpers/PCGAssetHelpers.h"
 #include "Utils/PCGPreconfiguration.h"
 #include "Tests/Determinism/PCGDeterminismSettings.h"
 
@@ -263,6 +264,7 @@ public:
 	UE_API virtual void PostInitProperties() override;
 #if WITH_EDITOR
 	UE_API virtual void PostEditUndo() override;
+	UE_API virtual void PreSave(FObjectPreSaveContext ObjectSaveContext) override;
 #endif // WITH_EDITOR
 	//~End UObject interface
 
@@ -665,6 +667,10 @@ private:
 #if WITH_EDITORONLY_DATA
 	/** The cached Crc for these settings. */
 	FPCGCrc CachedCrc;
+
+	/** Cached pins when settings are used as assets. Will never be used otherwise. */
+	UPROPERTY(AssetRegistrySearchable)
+	FPCGAssetCachedPins CachedPins;
 #endif
 
 public:
@@ -732,7 +738,7 @@ protected:
 #endif
 
 public:
-	UPROPERTY(VisibleAnywhere, BlueprintReadOnly, Category = Instance, meta = (EditInline))
+	UPROPERTY(VisibleAnywhere, BlueprintReadOnly, Category = Instance, AssetRegistrySearchable, meta = (EditInline))
 	TObjectPtr<UPCGSettings> Settings = nullptr;
 };
 
diff --git a/Engine/Plugins/PCG/Source/PCGEditor/Private/Schema/PCGEditorGraphSchema.cpp b/Engine/Plugins/PCG/Source/PCGEditor/Private/Schema/PCGEditorGraphSchema.cpp
--- a/Engine/Plugins/PCG/Source/PCGEditor/Private/Schema/PCGEditorGraphSchema.cpp
+++ b/Engine/Plugins/PCG/Source/PCGEditor/Private/Schema/PCGEditorGraphSchema.cpp
@@ -16,6 +16,7 @@
 #include "Elements/PCGMakeConcreteElement.h"
 #include "Elements/PCGReroute.h"
 #include "Elements/PCGUserParameterGet.h"
+#include "Helpers/PCGAssetHelpers.h"
 #include "Helpers/PCGHelpers.h"
 #include "Helpers/PCGPropertyHelpers.h"
 #include "Helpers/PCGSubgraphHelpers.h"
@@ -84,13 +85,15 @@ namespace PCGActionsHelpers
 
 		if (Pin->IsOutputPin())
 		{
-			bool bHasRequiredPins = Algo::AnyOf(PinProps, [](const FPCGPinProperties& Props) { return Props.IsRequiredPin(); });
+			const bool bHasRequiredPins = Algo::AnyOf(PinProps, [](const FPCGPinProperties& Props) { return Props.IsRequiredPin(); });
+			const bool bHasNormalPins = Algo::AnyOf(PinProps, [](const FPCGPinProperties& Props) { return Props.IsNormalPin(); });
 
-			return Algo::AnyOf(PinProps, [Pin, bHasRequiredPins, CompatibilityMask](const FPCGPinProperties& Props)
+			return Algo::AnyOf(PinProps, [Pin, bHasRequiredPins, bHasNormalPins, CompatibilityMask](const FPCGPinProperties& Props)
 				{
-					return (!Props.IsAdvancedPin() &&
-						(!bHasRequiredPins || Props.IsRequiredPin()) &&
-						IsValid(FPCGModule::GetConstDataTypeRegistry().IsCompatible(Pin->GetCurrentTypesID(), Props.AllowedTypes), CompatibilityMask));
+					return (Props.IsRequiredPin() ||
+						(!bHasRequiredPins && Props.IsNormalPin()) ||
+						(!bHasRequiredPins && !bHasNormalPins && Props.IsAdvancedPin())) &&
+						IsValid(FPCGModule::GetConstDataTypeRegistry().IsCompatible(Pin->GetCurrentTypesID(), Props.AllowedTypes), CompatibilityMask);
 				});
 		}
 		else
@@ -583,38 +586,55 @@ void UPCGEditorGraphSchema::GetNativeElementActions(FGraphActionMenuBuilder& Act
 
 void UPCGEditorGraphSchema::GetBlueprintElementActions(FGraphActionMenuBuilder& ActionMenuBuilder, const FPCGActionsFilter& ActionFilter) const
 {
-	PCGEditorUtils::ForEachPCGBlueprintAssetData([&ActionMenuBuilder, &ActionFilter](const FAssetData& AssetData)
+	const UEdGraphPin* Pin = ActionMenuBuilder.FromPin;
+	const UPCGPin* PCGPin = nullptr;
+	PCGActionsHelpers::ECompatibilityMask CompatibilityMask = PCGActionsHelpers::GetPinAndCompatibilityMask(Pin, ActionFilter.Graph, PCGPin);
+
+	PCGEditorUtils::ForEachPCGBlueprintAssetData([&ActionMenuBuilder, &ActionFilter, PCGPin, CompatibilityMask](const FAssetData& AssetData)
 	{
 		const bool bExposeToLibrary = AssetData.GetTagValueRef<bool>(GET_MEMBER_NAME_CHECKED(UPCGBlueprintBaseElement, bExposeToLibrary));
-		const bool bOnlyExposePreconfiguredSettings = AssetData.GetTagValueRef<bool>(GET_MEMBER_NAME_CHECKED(UPCGBlueprintBaseElement, bOnlyExposePreconfiguredSettings));
-
 		if (bExposeToLibrary)
 		{
-			const FText MenuDesc = FText::FromString(FName::NameToDisplayString(AssetData.AssetName.ToString(), false));
-			const FText Category = AssetData.GetTagValueRef<FText>(GET_MEMBER_NAME_CHECKED(UPCGBlueprintBaseElement, Category));
-			const FText Description = AssetData.GetTagValueRef<FText>(GET_MEMBER_NAME_CHECKED(UPCGBlueprintBaseElement, Description));
+			FPCGAssetHelpers::FBlueprintAssetOutput AssetOutput;
+			if (!FPCGAssetHelpers::GetBlueprintAssetRegistryData(AssetData, AssetOutput))
+			{
+				return true;
+			}
+
+			const FText& Category = AssetOutput.Category;
+			const FText& Description = AssetOutput.Description;
 
 			if (!ActionFilter.Accepts(Category))
 			{
 				return true;
 			}
 
-			const FSoftClassPath GeneratedClass = FSoftClassPath(AssetData.GetTagValueRef<FString>(FBlueprintTags::GeneratedClassPath));
+			// @todo_pcg: Consider preconfigured info for pins?
+			if (AssetOutput.CachedPins.bHasCachedPins && PCGPin)
+			{
+				const TArray<FPCGPinProperties>& BlueprintPins = PCGPin->IsOutputPin() ? AssetOutput.CachedPins.InputPins : AssetOutput.CachedPins.OutputPins;
+				if (!PCGActionsHelpers::HasCompatiblePins(PCGPin, BlueprintPins, CompatibilityMask))
+				{
+					return true;
+				}
+			}
 
 			// Only load the class if we have enabled preconfigured settings.
 			TArray<FPCGPreConfiguredSettingsInfo> AllPreconfiguredInfo;
-			if (AssetData.GetTagValueRef<bool>(GET_MEMBER_NAME_CHECKED(UPCGBlueprintBaseElement, bEnablePreconfiguredSettings)))
+			if(AssetOutput.bEnabledPreconfiguredSettings)
 			{
-				TSubclassOf<UPCGBlueprintBaseElement> BlueprintClass = GeneratedClass.TryLoadClass<UPCGBlueprintBaseElement>();
+				TSubclassOf<UPCGBlueprintBaseElement> BlueprintClass = AssetOutput.GeneratedClass.TryLoadClass<UPCGBlueprintBaseElement>();
 				const UPCGBlueprintBaseElement* BlueprintElement = BlueprintClass ? Cast<UPCGBlueprintBaseElement>(BlueprintClass->GetDefaultObject()) : nullptr;
 
 				AllPreconfiguredInfo = BlueprintElement ? BlueprintElement->PreconfiguredInfo : TArray<FPCGPreConfiguredSettingsInfo>{};
 			}
 
-			if (AllPreconfiguredInfo.IsEmpty() || !bOnlyExposePreconfiguredSettings)
+			const FText MenuDesc = FText::FromString(FName::NameToDisplayString(AssetData.AssetName.ToString(), false));
+
+			if (AllPreconfiguredInfo.IsEmpty() || !AssetOutput.bOnlyExposePreconfiguredSettings)
 			{
 				TSharedPtr<FPCGEditorGraphSchemaAction_NewBlueprintElement> NewBlueprintAction(new FPCGEditorGraphSchemaAction_NewBlueprintElement(Category, MenuDesc, Description, 0));
-				NewBlueprintAction->BlueprintClassPath = GeneratedClass;
+				NewBlueprintAction->BlueprintClassPath = AssetOutput.GeneratedClass;
 				ActionMenuBuilder.AddAction(NewBlueprintAction);
 			}
 
@@ -629,7 +649,7 @@ void UPCGEditorGraphSchema::GetBlueprintElementActions(FGraphActionMenuBuilder&
 			for (FPCGPreConfiguredSettingsInfo PreconfiguredInfo : AllPreconfiguredInfo)
 			{
 				TSharedPtr<FPCGEditorGraphSchemaAction_NewBlueprintElement> NewPreconfiguredAction(new FPCGEditorGraphSchemaAction_NewBlueprintElement(NewCategory, PreconfiguredInfo.Label, PreconfiguredInfo.Tooltip.IsEmpty() ? Description : PreconfiguredInfo.Tooltip, 0));
-				NewPreconfiguredAction->BlueprintClassPath = GeneratedClass;
+				NewPreconfiguredAction->BlueprintClassPath = AssetOutput.GeneratedClass;
 				NewPreconfiguredAction->PreconfiguredInfo = std::move(PreconfiguredInfo);
 				ActionMenuBuilder.AddAction(NewPreconfiguredAction);
 			}
@@ -641,23 +661,49 @@ void UPCGEditorGraphSchema::GetBlueprintElementActions(FGraphActionMenuBuilder&
 
 void UPCGEditorGraphSchema::GetSettingsElementActions(FGraphActionMenuBuilder& ActionMenuBuilder, const FPCGActionsFilter& ActionFilter, bool bIsContextual) const
 {
-	PCGEditorUtils::ForEachPCGSettingsAssetData([&ActionMenuBuilder, &ActionFilter, bIsContextual](const FAssetData& AssetData)
+	const FText UncategorizedAssetsCategory = LOCTEXT("AssetMenuDescriptionUncategorized", "Uncategorized Assets");
+	const UEdGraphPin* Pin = ActionMenuBuilder.FromPin;
+	const UPCGPin* PCGPin = nullptr;
+	PCGActionsHelpers::ECompatibilityMask CompatibilityMask = PCGActionsHelpers::GetPinAndCompatibilityMask(Pin, ActionFilter.Graph, PCGPin);
+
+	//@todo_pcg: should add settings instances too if they are assets
+	PCGEditorUtils::ForEachPCGSettingsAssetData([&ActionMenuBuilder, &ActionFilter, bIsContextual, &UncategorizedAssetsCategory, PCGPin, CompatibilityMask](const FAssetData& AssetData)
 	{
 		const bool bExposeToLibrary = AssetData.GetTagValueRef<bool>(TEXT("bExposeToLibrary"));
 		if (bExposeToLibrary)
 		{
+			FPCGAssetHelpers::FSettingsAssetOutput AssetOutput;
+			if (!FPCGAssetHelpers::GetSettingsAssetRegistryData(AssetData, AssetOutput))
+			{
+				return true;
+			}
+
+			const FText& Category = AssetOutput.Category.IsEmptyOrWhitespace() ? UncategorizedAssetsCategory : AssetOutput.Category;
+
+			if (!ActionFilter.Accepts(Category))
+			{
+				return true;
+			}
+
+			if (AssetOutput.CachedPins.bHasCachedPins && PCGPin)
+			{
+				const TArray<FPCGPinProperties>& SettingsPins = PCGPin->IsOutputPin() ? AssetOutput.CachedPins.InputPins : AssetOutput.CachedPins.OutputPins;
+				if (!PCGActionsHelpers::HasCompatiblePins(PCGPin, SettingsPins, CompatibilityMask))
+				{
+					return true;
+				}
+			}
+
+			const FText Description = FText::Format(INVTEXT("{0}\n{1}"), FText::FromString(AssetData.GetObjectPathString()), AssetOutput.Description);
 			const FText MenuDesc = FText::FromString(FName::NameToDisplayString(AssetData.AssetName.ToString(), false));
-			const FText AssetCategory = AssetData.GetTagValueRef<FText>(TEXT("Category"));
-			const FText Category = AssetCategory.IsEmptyOrWhitespace() ? LOCTEXT("AssetMenuDescriptionUncategorized", "Uncategorized Assets") : AssetCategory;
-			const FText Description = FText::Format(INVTEXT("{0}\n{1}"), FText::FromString(AssetData.GetObjectPathString()), AssetData.GetTagValueRef<FText>(TEXT("Description")));
 
-			if (!bIsContextual && ActionFilter.Accepts(Category))
+			if (!bIsContextual)
 			{
 				TSharedPtr<FPCGEditorGraphSchemaAction_NewSettingsElement> NewSettingsAction(new FPCGEditorGraphSchemaAction_NewSettingsElement(Category, MenuDesc, Description, 0));
 				NewSettingsAction->SettingsObjectPath = AssetData.GetSoftObjectPath();
 				ActionMenuBuilder.AddAction(NewSettingsAction);
 			}
-			else if(bIsContextual && ActionFilter.Accepts(Category))
+			else
 			{
 				TSharedPtr<FPCGEditorGraphSchemaAction_NewSettingsElement> NewSettingsActionCopy(new FPCGEditorGraphSchemaAction_NewSettingsElement(Category, MenuDesc, Description, 0));
 				NewSettingsActionCopy->SettingsObjectPath = AssetData.GetSoftObjectPath();
@@ -677,74 +723,45 @@ void UPCGEditorGraphSchema::GetSettingsElementActions(FGraphActionMenuBuilder& A
 
 void UPCGEditorGraphSchema::GetSubgraphElementActions(FGraphActionMenuBuilder& ActionMenuBuilder, const FPCGActionsFilter& ActionFilter) const
 {
+	const UEdGraphPin* Pin = ActionMenuBuilder.FromPin;
+	const UPCGPin* PCGPin = nullptr;
+	PCGActionsHelpers::ECompatibilityMask CompatibilityMask = PCGActionsHelpers::GetPinAndCompatibilityMask(Pin, ActionFilter.Graph, PCGPin);
+
 	FAssetRegistryModule& AssetRegistryModule = FModuleManager::Get().LoadModuleChecked<FAssetRegistryModule>(TEXT("AssetRegistry"));
 
-	PCGEditorUtils::ForEachPCGGraphAssetData([&ActionMenuBuilder, &ActionFilter, &AssetRegistryModule](const FAssetData& AssetData) -> bool
+	PCGEditorUtils::ForEachPCGGraphAssetData([&ActionMenuBuilder, &ActionFilter, &AssetRegistryModule, PCGPin, CompatibilityMask](const FAssetData& AssetData) -> bool
 	{
 		const bool bExposeToLibrary = AssetData.GetTagValueRef<bool>(GET_MEMBER_NAME_CHECKED(UPCGGraphInterface, bExposeToLibrary));
 		if (bExposeToLibrary)
 		{
-			// Only exposing the instance if it's parent graph is defined for instances. Otherwise, it is not interesting.
-			// Also march up the hierarchy to find overrides for category and description.
-			// We don't look for override in titles because if we have no override for an instance but the parent has an override we will have 2 times the same entry in the palette.
-			auto GetRecursiveTextsAndIsValid = [&AssetRegistryModule](const FAssetData& InAssetData, FText& OutCategory, FText& OutDescription, auto&& Recurse)
+			FPCGAssetHelpers::FGraphAssetOutput AssetOutput;
+			if (!FPCGAssetHelpers::GetGraphAssetRegistryData(AssetData, AssetOutput))
 			{
-				if (InAssetData.IsInstanceOf<UPCGGraphInstance>())
-				{
-					const FSoftObjectPath ParentGraph(InAssetData.GetTagValueRef<FString>(GET_MEMBER_NAME_CHECKED(UPCGGraphInstance, Graph)));
-					if (ParentGraph.IsNull())
-					{
-						return false;
-					}
-
-					if (OutCategory.IsEmpty())
-					{
-						OutCategory = InAssetData.GetTagValueRef<bool>(GET_MEMBER_NAME_CHECKED(UPCGGraphInstance, bOverrideCategory)) 
-							? InAssetData.GetTagValueRef<FText>(GET_MEMBER_NAME_CHECKED(UPCGGraphInstance, Category)) 
-							: FText();
-					}
+				return true;
+			}
 
-					if (OutDescription.IsEmpty())
-					{
-						OutDescription = InAssetData.GetTagValueRef<bool>(GET_MEMBER_NAME_CHECKED(UPCGGraphInstance, bOverrideDescription)) 
-							? InAssetData.GetTagValueRef<FText>(GET_MEMBER_NAME_CHECKED(UPCGGraphInstance, Description)) 
-							: FText();
-					}
+			if (!ActionFilter.Accepts(AssetOutput.Category))
+			{
+				return true;
+			}
 
-					// Asset data are not big so that should not be that big of a deal, but they are copied all the time.
-					// If we ever have performances issues, might be good to have a cache.
-					const FAssetData ParentAssetData = AssetRegistryModule.Get().GetAssetByObjectPath(ParentGraph);
-					return ParentAssetData.IsValid() ? Recurse(ParentAssetData, OutCategory, OutDescription, Recurse) : false;
-				}
-				else
+			if (AssetOutput.CachedPins.bHasCachedPins && PCGPin)
+			{
+				const TArray<FPCGPinProperties>& SubgraphPins = PCGPin->IsOutputPin() ? AssetOutput.CachedPins.InputPins : AssetOutput.CachedPins.OutputPins;
+				if (!PCGActionsHelpers::HasCompatiblePins(PCGPin, SubgraphPins, CompatibilityMask))
 				{
-					if (OutCategory.IsEmpty())
-					{
-						OutCategory = InAssetData.GetTagValueRef<FText>(GET_MEMBER_NAME_CHECKED(UPCGGraph, Category));
-					}
-
-					if (OutDescription.IsEmpty())
-					{
-						OutDescription = InAssetData.GetTagValueRef<FText>(GET_MEMBER_NAME_CHECKED(UPCGGraph, Description));
-					}
-
 					return true;
 				}
-			};
+			}
 
-			FText Category, Description;
+			// As stated above, we either have an override and we take it, or we use the asset name, to differentiate all possible instances of the same graph.
+			const FText MenuDesc = AssetData.GetTagValueRef<bool>(GET_MEMBER_NAME_CHECKED(UPCGGraphInterface, bOverrideTitle))
+				? AssetData.GetTagValueRef<FText>(GET_MEMBER_NAME_CHECKED(UPCGGraphInterface, Title))
+				: FText::FromString(FName::NameToDisplayString(AssetData.AssetName.ToString(), false));
 
-			if (GetRecursiveTextsAndIsValid(AssetData, Category, Description, GetRecursiveTextsAndIsValid) && ActionFilter.Accepts(Category))
-			{
-				// As stated above, we either have an override and we take it, or we use the asset name, to differentiate all possible instances of the same graph.
-				const FText MenuDesc = AssetData.GetTagValueRef<bool>(GET_MEMBER_NAME_CHECKED(UPCGGraphInterface, bOverrideTitle))
-					? AssetData.GetTagValueRef<FText>(GET_MEMBER_NAME_CHECKED(UPCGGraphInterface, Title))
-					: FText::FromString(FName::NameToDisplayString(AssetData.AssetName.ToString(), false));
-
-				TSharedPtr<FPCGEditorGraphSchemaAction_NewSubgraphElement> NewSubgraphAction(new FPCGEditorGraphSchemaAction_NewSubgraphElement(Category, MenuDesc, Description, 0));
-				NewSubgraphAction->SubgraphObjectPath = AssetData.GetSoftObjectPath();
-				ActionMenuBuilder.AddAction(NewSubgraphAction);
-			}
+			TSharedPtr<FPCGEditorGraphSchemaAction_NewSubgraphElement> NewSubgraphAction(new FPCGEditorGraphSchemaAction_NewSubgraphElement(AssetOutput.Category, MenuDesc, AssetOutput.Description, 0));
+			NewSubgraphAction->SubgraphObjectPath = AssetData.GetSoftObjectPath();
+			ActionMenuBuilder.AddAction(NewSubgraphAction);
 		}
 
 		return true;
@@ -764,12 +781,16 @@ void UPCGEditorGraphSchema::GetSubgraphElementActions(FGraphActionMenuBuilder& A
 void UPCGEditorGraphSchema::GetExtraElementActions(FGraphActionMenuBuilder& ActionMenuBuilder, const FPCGActionsFilter& ActionFilter) const
 {
 	const FText NoCategory;
-	// Comment action
-	const FText CommentMenuDesc = LOCTEXT("PCGAddComment", "Add Comment...");
-	const FText CommentDescription = LOCTEXT("PCGAddCommentTooltip", "Create a resizable comment box.");
 
-	const TSharedPtr<FPCGEditorGraphSchemaAction_NewComment> NewCommentAction(new FPCGEditorGraphSchemaAction_NewComment(NoCategory, CommentMenuDesc, CommentDescription, 0));
-	ActionMenuBuilder.AddAction(NewCommentAction);
+	// Comment action (only if not contextual)
+	if (!ActionMenuBuilder.FromPin)
+	{
+		const FText CommentMenuDesc = LOCTEXT("PCGAddComment", "Add Comment...");
+		const FText CommentDescription = LOCTEXT("PCGAddCommentTooltip", "Create a resizable comment box.");
+
+		const TSharedPtr<FPCGEditorGraphSchemaAction_NewComment> NewCommentAction(new FPCGEditorGraphSchemaAction_NewComment(NoCategory, CommentMenuDesc, CommentDescription, 0));
+		ActionMenuBuilder.AddAction(NewCommentAction);
+	}
 
 	// Reroute action
 	const FText RerouteMenuDesc = LOCTEXT("PCGAddRerouteNode", "Add Reroute Node");
@@ -788,6 +809,12 @@ void UPCGEditorGraphSchema::GetExtraElementActions(FGraphActionMenuBuilder& Acti
 
 void UPCGEditorGraphSchema::GetNamedRerouteUsageActions(FGraphActionMenuBuilder& ActionMenuBuilder, const FPCGActionsFilter& ActionFilter) const
 {
+	// Usage by definition implies that a pin, if one is provided, has to be an input pin.
+	if (ActionMenuBuilder.FromPin && ActionMenuBuilder.FromPin->Direction == EGPD_Output)
+	{
+		return;
+	}
+
 	const UEdGraph* CurrentGraph = ActionFilter.Graph;
 	const UPCGEditorGraph* Graph = Cast<UPCGEditorGraph>(CurrentGraph);
 
@@ -803,10 +830,21 @@ void UPCGEditorGraphSchema::GetNamedRerouteUsageActions(FGraphActionMenuBuilder&
 		return;
 	}
 
+	const UEdGraphPin* Pin = ActionMenuBuilder.FromPin;
+	const UPCGPin* PCGPin = nullptr;
+	PCGActionsHelpers::ECompatibilityMask CompatibilityMask = PCGActionsHelpers::GetPinAndCompatibilityMask(Pin, ActionFilter.Graph, PCGPin);
+
 	for (const UPCGNode* Node : PCGGraph->GetNodes())
 	{
 		if (const UPCGNamedRerouteDeclarationSettings* RerouteDeclaration = Cast<UPCGNamedRerouteDeclarationSettings>(Node->GetSettings()))
 		{
+			// Skip incompatible (typewise) reroute nodes.
+			// Implementation note: we can use AllOutputPinProperties here since reroutes have always a single output pin.
+			if (!PCGActionsHelpers::HasCompatiblePins(PCGPin, RerouteDeclaration->AllOutputPinProperties(), CompatibilityMask))
+			{
+				continue;
+			}
+
 			static const FText Category = LOCTEXT("NamedRerouteCategory", "Named Reroutes");
 			const FText Name = FText::FromName(Node->NodeTitle);
 			const FText Tooltip = FText::Format(LOCTEXT("NamedRerouteTooltip", "Add a usage of '{0}' here."), Name);
