[PCG] Allowed editor-time volume spawning from PCG (defaults cube).

#Synced-CL 51908749

diff --git a/Engine/Plugins/PCG/Source/PCG/Private/Helpers/PCGActorHelpers.cpp b/Engine/Plugins/PCG/Source/PCG/Private/Helpers/PCGActorHelpers.cpp
--- a/Engine/Plugins/PCG/Source/PCG/Private/Helpers/PCGActorHelpers.cpp
+++ b/Engine/Plugins/PCG/Source/PCG/Private/Helpers/PCGActorHelpers.cpp
@@ -26,6 +26,9 @@
 #include "Editor.h"
 #include "Selection.h"
 #include "UnrealEdGlobals.h"
+#include "ActorFactories/ActorFactory.h"
+#include "Builders/CubeBuilder.h"
+#include "Components/BrushComponent.h"
 #include "Engine/Level.h"
 #include "Editor/IPCGEditorModule.h"
 #include "Editor/UnrealEdEngine.h"
@@ -814,6 +817,13 @@ AActor* UPCGActorHelpers::SpawnDefaultActor(const FSpawnDefaultActorParams& Para
 				DataLayerInstance->AddActor(NewActor);
 			}
 		}
+
+		// If the spawed actor is a volume, we potentially need to fix its brush/brush setup/brush component
+		if (AVolume* NewVolume = Cast<AVolume>(NewActor))
+		{
+			UCubeBuilder* DefaultBrushBuilder = NewObject<UCubeBuilder>();
+			UActorFactory::CreateBrushForVolumeActor(NewVolume, DefaultBrushBuilder);
+		}
 	}
 #endif // WITH_EDITOR
 
