[PCG] Fix crash on SceneProxy being nullptr.

It can happen in at least two situations - if primitive is recreating its proxy, or if the primitive was destroyed after the ticker was created.

[RN] [PCG] Fixed rare crash when using GPU primitive spawning due to primitive proxy being null.
#Synced-CL 51942803

diff --git a/Engine/Plugins/Experimental/PCGInterops/PCGFastGeoInterop/Source/PCGFastGeoInterop/Private/Compute/PrimitiveFactories/PCGPrimitiveFactoryFastGeoPISMC.h b/Engine/Plugins/Experimental/PCGInterops/PCGFastGeoInterop/Source/PCGFastGeoInterop/Private/Compute/PrimitiveFactories/PCGPrimitiveFactoryFastGeoPISMC.h
--- a/Engine/Plugins/Experimental/PCGInterops/PCGFastGeoInterop/Source/PCGFastGeoInterop/Private/Compute/PrimitiveFactories/PCGPrimitiveFactoryFastGeoPISMC.h
+++ b/Engine/Plugins/Experimental/PCGInterops/PCGFastGeoInterop/Source/PCGFastGeoInterop/Private/Compute/PrimitiveFactories/PCGPrimitiveFactoryFastGeoPISMC.h
@@ -19,6 +19,7 @@ public:
 	virtual bool IsRenderStateCreated() const override;
 	virtual int32 GetNumPrimitives() const override { return Components.Num(); }
 	virtual FPrimitiveSceneProxy* GetSceneProxy(int32 InPrimitiveIndex) const override;
+	virtual bool IsPrimitiveValid(int32 InPrimitiveIndex) const override { return Components.IsValidIndex(InPrimitiveIndex) && Components[InPrimitiveIndex].Get() != nullptr; }
 	virtual int32 GetNumInstances(int32 InPrimitiveIndex, int32 InCellID) const override;
 	virtual int32 GetNumInstancesTotal(int32 InPrimitiveIndex) const override;
 	//~ End IPCGRuntimePrimitiveFactory interface
diff --git a/Engine/Plugins/PCG/Source/PCG/Private/Compute/DataInterfaces/PCGInstanceDataInterface.cpp b/Engine/Plugins/PCG/Source/PCG/Private/Compute/DataInterfaces/PCGInstanceDataInterface.cpp
--- a/Engine/Plugins/PCG/Source/PCG/Private/Compute/DataInterfaces/PCGInstanceDataInterface.cpp
+++ b/Engine/Plugins/PCG/Source/PCG/Private/Compute/DataInterfaces/PCGInstanceDataInterface.cpp
@@ -407,22 +407,6 @@ void FPCGInstanceDataProviderProxy::PostSubmit(FComputeContext& Context) const
 
 		TSharedPtr<IPCGRuntimePrimitiveFactory> Factory = DataProviderInner->PrimitiveFactory;
 
-#if WITH_EDITOR
-		// When launching Editor in background/unfocused, it appears the render setup does not make progress and can be not ready here.
-		// Does not seem to repro in standalone build.
-		for (int32 PrimitiveIndex = 0; PrimitiveIndex < Factory->GetNumPrimitives(); ++PrimitiveIndex)
-		{
-			FPrimitiveSceneProxy* SceneProxy = Factory->GetSceneProxy(PrimitiveIndex);
-			check(SceneProxy);
-
-			if (SceneProxy->GetPrimitiveSceneInfo()->GetInstanceSceneDataOffset() == -1)
-			{
-				UE_LOG(LogPCG, Verbose, TEXT("Primitive instances not allocated in GPU scene, retrying next frame."));
-				return true;
-			}
-		}
-#endif // WITH_EDITOR
-
 		for (int32 PrimitiveIndex = 0; PrimitiveIndex < Factory->GetNumPrimitives(); ++PrimitiveIndex)
 		{
 			const int32 PrimitiveNumInstances = Factory->GetNumInstancesTotal(PrimitiveIndex);
@@ -434,12 +418,23 @@ void FPCGInstanceDataProviderProxy::PostSubmit(FComputeContext& Context) const
 			}
 
 			FPrimitiveSceneProxy* SceneProxy = Factory->GetSceneProxy(PrimitiveIndex);
-			check(SceneProxy);
+			if (!SceneProxy)
+			{
+				if (!Factory->IsPrimitiveValid(PrimitiveIndex))
+				{
+					UE_LOG(LogPCG, Verbose, TEXT("Primitive component was destroyed, aborting instance write."));
+					return false;
+				}
 
+				UE_LOG(LogPCG, Verbose, TEXT("Scene proxy not ready, retrying next frame."));
+				return true;
+			}
+
+			// todo_pcg: This is technically a race with the RT and we should try remove these conditions and get to the bottom of why it was needed.
 			if (!ensure(SceneProxy->GetPrimitiveSceneInfo()->GetInstanceSceneDataOffset() >= 0))
 			{
-				UE_LOG(LogPCG, Error, TEXT("FPCGInstanceDataProviderProxy: Scene proxy for primitive %d not ready (InstanceSceneDataOffset is negative), this should not happen."), PrimitiveIndex);
-				continue;
+				UE_LOG(LogPCG, Verbose, TEXT("Primitive instances not allocated in GPU scene, retrying next frame."));
+				return true;
 			}
 
 #if UE_BUILD_SHIPPING
diff --git a/Engine/Plugins/PCG/Source/PCG/Private/Compute/PrimitiveFactories/PCGPrimitiveFactoryISKMC.h b/Engine/Plugins/PCG/Source/PCG/Private/Compute/PrimitiveFactories/PCGPrimitiveFactoryISKMC.h
--- a/Engine/Plugins/PCG/Source/PCG/Private/Compute/PrimitiveFactories/PCGPrimitiveFactoryISKMC.h
+++ b/Engine/Plugins/PCG/Source/PCG/Private/Compute/PrimitiveFactories/PCGPrimitiveFactoryISKMC.h
@@ -4,6 +4,8 @@
 
 #include "Compute/PrimitiveFactories/IPCGRuntimePrimitiveFactory.h"
 
+#include "UObject/WeakObjectPtr.h"
+
 class AActor;
 class UPrimitiveComponent;
 struct FPCGContext;
@@ -16,6 +18,7 @@ public:
 	//~ Begin IPCGRuntimePrimitiveFactory interface
 	virtual int32 GetNumPrimitives() const override { return Components.Num(); }
 	virtual FPrimitiveSceneProxy* GetSceneProxy(int32 InPrimitiveIndex) const override;
+	virtual bool IsPrimitiveValid(int32 InPrimitiveIndex) const override { return Components.IsValidIndex(InPrimitiveIndex) && Components[InPrimitiveIndex].IsValid(); }
 	virtual int32 GetNumInstances(int32 InPrimitiveIndex, int32 InCellID) const override;
 	virtual int32 GetNumInstancesTotal(int32 InPrimitiveIndex) const override;
 	virtual bool IsRenderStateCreated() const override;
diff --git a/Engine/Plugins/PCG/Source/PCG/Private/Compute/PrimitiveFactories/PCGPrimitiveFactoryPISMC.h b/Engine/Plugins/PCG/Source/PCG/Private/Compute/PrimitiveFactories/PCGPrimitiveFactoryPISMC.h
--- a/Engine/Plugins/PCG/Source/PCG/Private/Compute/PrimitiveFactories/PCGPrimitiveFactoryPISMC.h
+++ b/Engine/Plugins/PCG/Source/PCG/Private/Compute/PrimitiveFactories/PCGPrimitiveFactoryPISMC.h
@@ -4,6 +4,8 @@
 
 #include "Compute/PrimitiveFactories/PCGPrimitiveFactoryISMBase.h"
 
+#include "UObject/WeakObjectPtr.h"
+
 class AActor;
 class UPrimitiveComponent;
 struct FPCGContext;
@@ -15,6 +17,7 @@ public:
 	//~ Begin IPCGRuntimePrimitiveFactory interface
 	virtual int32 GetNumPrimitives() const override { return Components.Num(); }
 	virtual FPrimitiveSceneProxy* GetSceneProxy(int32 InPrimitiveIndex) const override;
+	virtual bool IsPrimitiveValid(int32 InPrimitiveIndex) const override { return Components.IsValidIndex(InPrimitiveIndex) && Components[InPrimitiveIndex].IsValid(); }
 	virtual int32 GetNumInstances(int32 InPrimitiveIndex, int32 InCellID) const override;
 	virtual int32 GetNumInstancesTotal(int32 InPrimitiveIndex) const override;
 	virtual bool IsRenderStateCreated() const override;
diff --git a/Engine/Plugins/PCG/Source/PCG/Public/Compute/PrimitiveFactories/IPCGRuntimePrimitiveFactory.h b/Engine/Plugins/PCG/Source/PCG/Public/Compute/PrimitiveFactories/IPCGRuntimePrimitiveFactory.h
--- a/Engine/Plugins/PCG/Source/PCG/Public/Compute/PrimitiveFactories/IPCGRuntimePrimitiveFactory.h
+++ b/Engine/Plugins/PCG/Source/PCG/Public/Compute/PrimitiveFactories/IPCGRuntimePrimitiveFactory.h
@@ -14,6 +14,7 @@ public:
 
 	virtual int32 GetNumPrimitives() const = 0;
 	virtual FPrimitiveSceneProxy* GetSceneProxy(int32 InPrimitiveIndex) const = 0;
+	virtual bool IsPrimitiveValid(int32 InPrimitiveIndex) const { ensureMsgf(false, TEXT("IsPrimitiveValid() should be overridden by classes implementing the IPCGRuntimePrimitiveFactory interface.")); return false; }
 	PCG_API virtual int32 GetNumInstances(int32 InPrimitiveIndex, int32 InCellID) const;
 	virtual int32 GetNumInstancesTotal(int32 InPrimitiveIndex) const = 0;
 	virtual bool IsRenderStateCreated() const = 0;
