[PCG] Remove OtherParents in Metadata

* Was not used anywhere
* Might induce crashes on GC in rare cases.

#rnx
#Synced-CL 51771932

diff --git a/Engine/Plugins/PCG/Source/PCG/Private/Metadata/PCGMetadata.cpp b/Engine/Plugins/PCG/Source/PCG/Private/Metadata/PCGMetadata.cpp
--- a/Engine/Plugins/PCG/Source/PCG/Private/Metadata/PCGMetadata.cpp
+++ b/Engine/Plugins/PCG/Source/PCG/Private/Metadata/PCGMetadata.cpp
@@ -595,7 +595,6 @@ void UPCGMetadata::InitializeAsCopy(const FPCGMetadataInitializeParams& InParams
 	if (bShouldParent)
 	{
 		Parent = InParams.Parent;
-		OtherParents = InParams.Parent->OtherParents;
 	}
 }
 
@@ -637,11 +636,6 @@ void UPCGMetadata::AddAttributes(const FPCGMetadataInitializeParams& InParams)
 
 		FindFixOrCreateDomainInitializeParams(InParams, DomainID, OtherMetadataDomain, std::move(InitializeFunc));
 	});
-
-	if (bAddSucceeded && InParams.Parent != Parent)
-	{
-		OtherParents.Add(InParams.Parent);
-	}
 }
 
 void UPCGMetadata::BP_AddAttribute(const UPCGMetadata* InOther, FName AttributeName)
@@ -656,12 +650,7 @@ void UPCGMetadata::AddAttribute(const UPCGMetadata* InOther, FPCGAttributeIdenti
 		return;
 	}
 
-	const bool bAttributeAdded = CopyAttribute(InOther->GetConstAttribute(AttributeName), AttributeName, /*bKeepParent=*/InOther == Parent, /*bCopyEntries=*/false, /*bCopyValues=*/false) != nullptr;
-
-	if (InOther != Parent && bAttributeAdded)
-	{
-		OtherParents.Add(InOther);
-	}
+	CopyAttribute(InOther->GetConstAttribute(AttributeName), AttributeName, /*bKeepParent=*/InOther == Parent, /*bCopyEntries=*/false, /*bCopyValues=*/false);
 }
 
 void UPCGMetadata::CopyAttributes(const UPCGMetadata* InOther)
diff --git a/Engine/Plugins/PCG/Source/PCG/Public/Metadata/PCGMetadata.h b/Engine/Plugins/PCG/Source/PCG/Public/Metadata/PCGMetadata.h
--- a/Engine/Plugins/PCG/Source/PCG/Public/Metadata/PCGMetadata.h
+++ b/Engine/Plugins/PCG/Source/PCG/Public/Metadata/PCGMetadata.h
@@ -477,11 +477,6 @@ protected:
 	UPROPERTY()
 	TObjectPtr<const UPCGMetadata> Parent;
 
-	// Set of parents kept for streams relationship and GC collection
-	// But otherwise not used directly
-	UPROPERTY()
-	TSet<TWeakObjectPtr<const UPCGMetadata>> OtherParents;
-
 	FPCGMetadataDomainID DefaultDomain = PCGMetadataDomainID::Default;
 	TMap<FPCGMetadataDomainID, TSharedPtr<FPCGMetadataDomain>> MetadataDomains;
 	PCG::FLock MetadataDomainsLock;
