[PCG] Added embedded subgraph filtering and also used live pin data instead of relying on saved assets

#Synced-CL 51854852

diff --git a/Engine/Plugins/PCG/Source/PCG/Private/Elements/Blueprint/PCGBlueprintBaseElement.cpp b/Engine/Plugins/PCG/Source/PCG/Private/Elements/Blueprint/PCGBlueprintBaseElement.cpp
--- a/Engine/Plugins/PCG/Source/PCG/Private/Elements/Blueprint/PCGBlueprintBaseElement.cpp
+++ b/Engine/Plugins/PCG/Source/PCG/Private/Elements/Blueprint/PCGBlueprintBaseElement.cpp
@@ -68,16 +68,22 @@ void UPCGBlueprintBaseElement::BeginDestroy()
 }
 
 #if WITH_EDITOR
+FPCGAssetCachedPins UPCGBlueprintBaseElement::BuildCachedPins() const
+{
+	FPCGAssetCachedPins Pins;
+	Pins.bHasCachedPins = true;
+	Pins.InputPins = GetInputPins();
+	Pins.OutputPins = GetOutputPins();
+
+	return Pins;
+}
+
 void UPCGBlueprintBaseElement::PreSave(FObjectPreSaveContext ObjectSaveContext)
 {
 	Super::PreSave(ObjectSaveContext);
 
 	// Cache pins for asset retrieval
-	{
-		CachedPins.InputPins = GetInputPins();
-		CachedPins.OutputPins = GetOutputPins();
-		CachedPins.bHasCachedPins = true;
-	}
+	CachedPins = BuildCachedPins();
 }
 #endif // WITH_EDITOR
 
diff --git a/Engine/Plugins/PCG/Source/PCG/Private/Helpers/PCGAssetHelpers.cpp b/Engine/Plugins/PCG/Source/PCG/Private/Helpers/PCGAssetHelpers.cpp
--- a/Engine/Plugins/PCG/Source/PCG/Private/Helpers/PCGAssetHelpers.cpp
+++ b/Engine/Plugins/PCG/Source/PCG/Private/Helpers/PCGAssetHelpers.cpp
@@ -60,11 +60,18 @@ bool FPCGAssetHelpers::GetGraphAssetRegistryData(const FAssetData& AssetData, FP
 			}
 
 			// Implementation note: pins are only present at the graph level
-			FString TempString;
-			InAssetData.GetTagValue<FString>("CachedPins", TempString); // can't use GET_MEMBER_NAME_CHECKED because CachedPins are private.
-			if (!TempString.IsEmpty())
+			if (UPCGGraph* Graph = Cast<UPCGGraph>(InAssetData.FastGetAsset(/*bLoad=*/false)))
 			{
-				FPCGAssetCachedPins::StaticStruct()->ImportText(*TempString, &Output.CachedPins, nullptr, 0, nullptr, FString{});
+				Output.CachedPins = Graph->BuildCachedPins();
+			}
+			else
+			{
+				FString TempString;
+				InAssetData.GetTagValue<FString>(GET_MEMBER_NAME_CHECKED(UPCGGraph, CachedPins), TempString);
+				if (!TempString.IsEmpty())
+				{
+					FPCGAssetCachedPins::StaticStruct()->ImportText(*TempString, &Output.CachedPins, nullptr, 0, nullptr, FString{});
+				}
 			}
 
 			return true;
@@ -83,11 +90,18 @@ bool FPCGAssetHelpers::GetBlueprintAssetRegistryData(const FAssetData& AssetData
 	Output.bOnlyExposePreconfiguredSettings = AssetData.GetTagValueRef<bool>(GET_MEMBER_NAME_CHECKED(UPCGBlueprintBaseElement, bOnlyExposePreconfiguredSettings));
 	Output.bEnabledPreconfiguredSettings = AssetData.GetTagValueRef<bool>(GET_MEMBER_NAME_CHECKED(UPCGBlueprintBaseElement, bEnablePreconfiguredSettings));
 
-	FString TempString;
-	AssetData.GetTagValue<FString>("CachedPins", TempString); // note can't use GET_MEMBER_NAME_CHECKED because field is private
-	if (!TempString.IsEmpty())
+	if (UPCGBlueprintBaseElement* BPElement = Cast<UPCGBlueprintBaseElement>(AssetData.FastGetAsset(/*bLoad=*/false)))
 	{
-		FPCGAssetCachedPins::StaticStruct()->ImportText(*TempString, &Output.CachedPins, nullptr, 0, nullptr, FString{});
+		Output.CachedPins = BPElement->BuildCachedPins();
+	}
+	else
+	{
+		FString TempString;
+		AssetData.GetTagValue<FString>(GET_MEMBER_NAME_CHECKED(UPCGBlueprintBaseElement, CachedPins), TempString);
+		if (!TempString.IsEmpty())
+		{
+			FPCGAssetCachedPins::StaticStruct()->ImportText(*TempString, &Output.CachedPins, nullptr, 0, nullptr, FString{});
+		}
 	}
 
 	return true;
@@ -106,11 +120,19 @@ bool FPCGAssetHelpers::GetSettingsAssetRegistryData(const FAssetData& AssetData,
 			Output.Category = InAssetData.GetTagValueRef<FText>(GET_MEMBER_NAME_CHECKED(UPCGSettings, Category));
 			Output.Description = InAssetData.GetTagValueRef<FText>(GET_MEMBER_NAME_CHECKED(UPCGSettings, Description));
 
-			FString TempString;
-			InAssetData.GetTagValue<FString>("CachedPins", TempString); // can't use GET_MEMBER_NAME_CHECKED because CachedPins are private.
-			if (!TempString.IsEmpty())
+			// Asset is already loaded, we'll use the live values instead
+			if (UPCGSettings* Settings = Cast<UPCGSettings>(InAssetData.FastGetAsset(/*bLoad=*/false)))
+			{
+				Output.CachedPins = Settings->BuildCachedPins();
+			}
+			else
 			{
-				FPCGAssetCachedPins::StaticStruct()->ImportText(*TempString, &Output.CachedPins, nullptr, 0, nullptr, FString{});
+				FString TempString;
+				InAssetData.GetTagValue<FString>(GET_MEMBER_NAME_CHECKED(UPCGSettings, CachedPins), TempString);
+				if (!TempString.IsEmpty())
+				{
+					FPCGAssetCachedPins::StaticStruct()->ImportText(*TempString, &Output.CachedPins, nullptr, 0, nullptr, FString{});
+				}
 			}
 
 			return true;
diff --git a/Engine/Plugins/PCG/Source/PCG/Private/PCGGraph.cpp b/Engine/Plugins/PCG/Source/PCG/Private/PCGGraph.cpp
--- a/Engine/Plugins/PCG/Source/PCG/Private/PCGGraph.cpp
+++ b/Engine/Plugins/PCG/Source/PCG/Private/PCGGraph.cpp
@@ -850,6 +850,16 @@ void UPCGGraph::SetupEditorData()
 	}
 }
 
+FPCGAssetCachedPins UPCGGraph::BuildCachedPins() const
+{
+	FPCGAssetCachedPins Pins;
+	Pins.bHasCachedPins = true;
+	Pins.InputPins = DefaultInputPinProperties();
+	Pins.OutputPins = DefaultOutputPinProperties();
+
+	return Pins;
+}
+
 void UPCGGraph::PreSave(FObjectPreSaveContext ObjectSaveContext)
 {
 	Super::PreSave(ObjectSaveContext);
@@ -862,9 +872,7 @@ void UPCGGraph::PreSave(FObjectPreSaveContext ObjectSaveContext)
 	// Cache pins for asset retrieval
 	if(IsAsset())
 	{
-		CachedPins.InputPins = InputNode ? CastChecked<UPCGGraphInputOutputSettings>(InputNode->GetSettings())->OutputPinProperties() : TArray<FPCGPinProperties>();
-		CachedPins.OutputPins = OutputNode ? CastChecked<UPCGGraphInputOutputSettings>(OutputNode->GetSettings())->InputPinProperties() : TArray<FPCGPinProperties>();
-		CachedPins.bHasCachedPins = true;
+		CachedPins = BuildCachedPins();
 	}
 
 	// Prepare cook data
@@ -1390,6 +1398,16 @@ bool UPCGGraph::AddLabeledEdge(UPCGNode* From, const FName& FromPinLabel, UPCGNo
 	return bToPinBrokeOtherEdges;
 }
 
+TArray<FPCGPinProperties> UPCGGraph::DefaultInputPinProperties() const
+{
+	return InputNode ? CastChecked<UPCGGraphInputOutputSettings>(InputNode->GetSettings())->OutputPinProperties() : TArray<FPCGPinProperties>();
+}
+
+TArray<FPCGPinProperties> UPCGGraph::DefaultOutputPinProperties() const
+{
+	return OutputNode ? CastChecked<UPCGGraphInputOutputSettings>(OutputNode->GetSettings())->InputPinProperties() : TArray<FPCGPinProperties>();
+}
+
 TArray<UPCGEdge*> UPCGGraph::GetAllEdges() const
 {
 	TArray<UPCGEdge*> AllEdges;
diff --git a/Engine/Plugins/PCG/Source/PCG/Private/PCGSettings.cpp b/Engine/Plugins/PCG/Source/PCG/Private/PCGSettings.cpp
--- a/Engine/Plugins/PCG/Source/PCG/Private/PCGSettings.cpp
+++ b/Engine/Plugins/PCG/Source/PCG/Private/PCGSettings.cpp
@@ -321,23 +321,22 @@ void UPCGSettings::PostEditUndo()
 	Super::PostEditUndo();
 }
 
+FPCGAssetCachedPins UPCGSettings::BuildCachedPins() const
+{
+	FPCGAssetCachedPins Pins;
+	Pins.bHasCachedPins = true;
+	Pins.InputPins = InputPinProperties();
+	Pins.OutputPins = OutputPinProperties();
+
+	return Pins;
+}
+
 void UPCGSettings::PreSave(FObjectPreSaveContext ObjectSaveContext)
 {
 	Super::PreSave(ObjectSaveContext);
 
 	// Update cached pins - note that we need to cache pins only for asset-level settings. Otherwise, make them always empty.
-	if(IsAsset())
-	{
-		CachedPins.InputPins = InputPinProperties();
-		CachedPins.OutputPins = OutputPinProperties();
-		CachedPins.bHasCachedPins = true;
-	}
-	else
-	{
-		CachedPins.InputPins.Reset();
-		CachedPins.OutputPins.Reset();
-		CachedPins.bHasCachedPins = false;
-	}
+	CachedPins = IsAsset() ? BuildCachedPins() : FPCGAssetCachedPins();
 }
 #endif // WITH_EDITOR
 
diff --git a/Engine/Plugins/PCG/Source/PCG/Public/Elements/Blueprint/PCGBlueprintBaseElement.h b/Engine/Plugins/PCG/Source/PCG/Public/Elements/Blueprint/PCGBlueprintBaseElement.h
--- a/Engine/Plugins/PCG/Source/PCG/Public/Elements/Blueprint/PCGBlueprintBaseElement.h
+++ b/Engine/Plugins/PCG/Source/PCG/Public/Elements/Blueprint/PCGBlueprintBaseElement.h
@@ -26,6 +26,10 @@ class UPCGBlueprintBaseElement : public UObject
 {
 	GENERATED_BODY()
 
+#if WITH_EDITOR
+	friend class FPCGAssetHelpers;
+#endif
+
 public:
 	// ~Begin UObject interface
 	PCG_API virtual void PostLoad() override;
@@ -187,6 +191,9 @@ private:
 	/** Cached pin data to get type information without loading. */
 	UPROPERTY(AssetRegistrySearchable)
 	FPCGAssetCachedPins CachedPins;
+
+	/** Builds and returns list of pins as visible for graph creation purposes. */
+	FPCGAssetCachedPins BuildCachedPins() const;
 #endif
 
 protected:
diff --git a/Engine/Plugins/PCG/Source/PCG/Public/PCGGraph.h b/Engine/Plugins/PCG/Source/PCG/Public/PCGGraph.h
--- a/Engine/Plugins/PCG/Source/PCG/Public/PCGGraph.h
+++ b/Engine/Plugins/PCG/Source/PCG/Public/PCGGraph.h
@@ -4,6 +4,7 @@
 
 #include "PCGCommon.h"
 #include "PCGNode.h"
+#include "PCGPin.h"
 #include "PCGSettings.h"
 #include "Concepts/UEnum.h"
 #include "Data/PCGToolData.h"
@@ -319,6 +320,7 @@ class UPCGGraph : public UPCGGraphInterface
 	friend class FPCGEditor;
 	friend class FPCGEditorModule;
 	friend class FPCGSubgraphHelpers;
+	friend class FPCGAssetHelpers;
 #endif // WITH_EDITOR
 	friend class UPCGSubgraphSettings;
 
@@ -460,6 +462,12 @@ public:
 	UFUNCTION(BlueprintCallable, Category = Graph)
 	UPCGNode* GetOutputNode() const { return OutputNode; }
 
+	/** Returns the default graph input pins (from the input node), which exclude the parameter pins, execution dependency pin, etc. */
+	UE_API TArray<FPCGPinProperties> DefaultInputPinProperties() const;
+
+	/** Returns the default graph output pins (from the input node), which exclude the parameter pins, execution dependency pin, etc. */
+	UE_API TArray<FPCGPinProperties> DefaultOutputPinProperties() const;
+
 	/** Returns all the edges of the graph */
 	UFUNCTION(BlueprintCallable, Category = Graph)
 	UE_API TArray<UPCGEdge*> GetAllEdges() const;
@@ -745,6 +753,9 @@ private:
 	/** Cached pin data to get type information without loading. */
 	UPROPERTY(AssetRegistrySearchable)
 	FPCGAssetCachedPins CachedPins;
+
+	/** Builds and returns list of pins as visible for graph creation purposes. */
+	FPCGAssetCachedPins BuildCachedPins() const;
 #endif // WITH_EDITORONLY_DATA
 
 public:
diff --git a/Engine/Plugins/PCG/Source/PCG/Public/PCGSettings.h b/Engine/Plugins/PCG/Source/PCG/Public/PCGSettings.h
--- a/Engine/Plugins/PCG/Source/PCG/Public/PCGSettings.h
+++ b/Engine/Plugins/PCG/Source/PCG/Public/PCGSettings.h
@@ -247,6 +247,10 @@ class UPCGSettings : public UPCGSettingsInterface
 	friend class UPCGSettingsInterface;
 	friend struct FPCGContext;
 
+#if WITH_EDITOR
+	friend class FPCGAssetHelpers;
+#endif
+
 public:
 	// ~Begin UPCGData interface
 	UE_API virtual void AddToCrc(FArchiveCrc32& Ar, bool bFullDataCrc) const override;
@@ -671,6 +675,9 @@ private:
 	/** Cached pins when settings are used as assets. Will never be used otherwise. */
 	UPROPERTY(AssetRegistrySearchable)
 	FPCGAssetCachedPins CachedPins;
+
+	/** Builds and returns list of pins as visible for graph creation purposes. */
+	FPCGAssetCachedPins BuildCachedPins() const;
 #endif
 
 public:
diff --git a/Engine/Plugins/PCG/Source/PCGEditor/Private/Schema/PCGEditorGraphSchema.cpp b/Engine/Plugins/PCG/Source/PCGEditor/Private/Schema/PCGEditorGraphSchema.cpp
--- a/Engine/Plugins/PCG/Source/PCGEditor/Private/Schema/PCGEditorGraphSchema.cpp
+++ b/Engine/Plugins/PCG/Source/PCGEditor/Private/Schema/PCGEditorGraphSchema.cpp
@@ -771,6 +771,11 @@ void UPCGEditorGraphSchema::GetSubgraphElementActions(FGraphActionMenuBuilder& A
 	{
 		for (UPCGGraph* EmbeddedSubgraph : ActionFilter.Graph->GetPCGGraph()->GetEmbeddedSubgraphs())
 		{
+			if (PCGPin && !PCGActionsHelpers::HasCompatiblePins(PCGPin, PCGPin->IsOutputPin() ? EmbeddedSubgraph->DefaultInputPinProperties() : EmbeddedSubgraph->DefaultOutputPinProperties(), CompatibilityMask))
+			{
+				continue;
+			}
+
 			TSharedPtr<FPCGEditorGraphSchemaAction_NewSubgraphElement> NewSubgraphAction(new FPCGEditorGraphSchemaAction_NewSubgraphElement(EmbeddedSubgraph->Category, EmbeddedSubgraph->Title, EmbeddedSubgraph->Description, 0));
 			NewSubgraphAction->SubgraphObjectPath = EmbeddedSubgraph;
 			ActionMenuBuilder.AddAction(NewSubgraphAction);
