[PCG] Generalized standalone graphs to allow different workflow types (asset vs level).

#Synced-CL 51854852

diff --git a/Engine/Plugins/Experimental/ProceduralVegetationEditor/Source/ProceduralVegetation/Public/ProceduralVegetation.h b/Engine/Plugins/Experimental/ProceduralVegetationEditor/Source/ProceduralVegetation/Public/ProceduralVegetation.h
--- a/Engine/Plugins/Experimental/ProceduralVegetationEditor/Source/ProceduralVegetation/Public/ProceduralVegetation.h
+++ b/Engine/Plugins/Experimental/ProceduralVegetationEditor/Source/ProceduralVegetation/Public/ProceduralVegetation.h
@@ -17,7 +17,9 @@ public:
 		: UPCGGraph(ObjectInitializer)
 	{
 #if WITH_EDITORONLY_DATA
+PRAGMA_DISABLE_DEPRECATION_WARNINGS
 		bIsStandaloneGraph = false;
+PRAGMA_ENABLE_DEPRECATION_WARNINGS
 		bExposeGenerationInAssetExplorer = false;
 #endif
 		
diff --git a/Engine/Plugins/PCG/Source/PCG/Private/PCGGraph.cpp b/Engine/Plugins/PCG/Source/PCG/Private/PCGGraph.cpp
--- a/Engine/Plugins/PCG/Source/PCG/Private/PCGGraph.cpp
+++ b/Engine/Plugins/PCG/Source/PCG/Private/PCGGraph.cpp
@@ -184,7 +184,12 @@ FText UPCGGraphInterface::GetDisplayName() const
 #if WITH_EDITOR
 bool UPCGGraphInterface::IsStandaloneGraph() const
 {
-	return GetGraph() && GetGraph()->bIsStandaloneGraph;
+	return GetGraphUsage() != EPCGGraphUsage::Standard;
+}
+
+EPCGGraphUsage UPCGGraphInterface::GetGraphUsage() const
+{
+	return GetGraph() ? GetGraph()->GraphUsageContext : EPCGGraphUsage::Standard;
 }
 
 UObject* UPCGGraphInterface::GetUserParameterHierarchyRoot() const
@@ -193,15 +198,34 @@ UObject* UPCGGraphInterface::GetUserParameterHierarchyRoot() const
 }
 
 bool UPCGGraphInterface::IsStandaloneGraphAsset(const FAssetData& InAssetData)
+{
+	return GetGraphUsage(InAssetData) != EPCGGraphUsage::Standard;
+}
+
+EPCGGraphUsage UPCGGraphInterface::GetGraphUsage(const FAssetData& InAssetData)
 {
 	if (UPCGGraphInterface* GraphInterface = Cast<UPCGGraphInterface>(InAssetData.FastGetAsset(/*bLoad=*/false)))
 	{
-		return GraphInterface->IsStandaloneGraph();
+		return GraphInterface->GetGraphUsage();
 	}
 	else if (InAssetData.IsInstanceOf<UPCGGraph>())
 	{
 		bool bTagValue = false;
-		return InAssetData.GetTagValue(GET_MEMBER_NAME_CHECKED(UPCGGraph, bIsStandaloneGraph), bTagValue) && bTagValue;
+		FString TagStringValue;
+		// Two cases here:
+		// older data -> use bIsStandaloneGraph, if true -> Asset usage.
+		// newer data -> use GraphUsageContext.
+PRAGMA_DISABLE_DEPRECATION_WARNINGS
+		if (InAssetData.GetTagValue(GET_MEMBER_NAME_CHECKED(UPCGGraph, bIsStandaloneGraph), bTagValue) && bTagValue)
+		{
+			return EPCGGraphUsage::Asset;
+		}
+PRAGMA_ENABLE_DEPRECATION_WARNINGS
+		else if (InAssetData.GetTagValue(GET_MEMBER_NAME_CHECKED(UPCGGraph, GraphUsageContext), TagStringValue))
+		{
+			int64 UsageValue = StaticEnum<EPCGGraphUsage>()->GetValueByNameString(TagStringValue);
+			return UsageValue == INDEX_NONE ? EPCGGraphUsage::Standard : static_cast<EPCGGraphUsage>(UsageValue);
+		}
 	}
 	else if (InAssetData.IsInstanceOf<UPCGGraphInstance>())
 	{
@@ -210,11 +234,11 @@ bool UPCGGraphInterface::IsStandaloneGraphAsset(const FAssetData& InAssetData)
 		FAssetData GraphAssetData;
 		if (AssetManager.GetAssetDataForPath(GraphPath, GraphAssetData))
 		{
-			return IsStandaloneGraphAsset(GraphAssetData);
+			return GetGraphUsage(GraphAssetData);
 		}
 	}
 
-	return false;
+	return EPCGGraphUsage::Standard;
 }
 #endif
 
@@ -661,6 +685,16 @@ void UPCGGraph::PostLoad()
 		ShouldCook = false;
 		bIsEditorOnly = false;
 	}
+
+	if (bIsStandaloneGraph)
+	{
+		if (GraphUsageContext == EPCGGraphUsage::Standard)
+		{
+			GraphUsageContext = EPCGGraphUsage::Asset;
+		}
+
+		bIsStandaloneGraph = false;
+	}
 	PRAGMA_ENABLE_DEPRECATION_WARNINGS
 #endif // WITH_EDITOR
 }
@@ -2362,7 +2396,7 @@ void UPCGGraph::PostEditChangeChainProperty(FPropertyChangedChainEvent& Property
 		|| HeadPropertyName == GET_MEMBER_NAME_CHECKED(UPCGGraphInterface, bOverrideTitle)
 		|| HeadPropertyName == GET_MEMBER_NAME_CHECKED(UPCGGraphInterface, Color)
 		|| HeadPropertyName == GET_MEMBER_NAME_CHECKED(UPCGGraphInterface, bOverrideColor)
-		|| HeadPropertyName == GET_MEMBER_NAME_CHECKED(UPCGGraph, bIsStandaloneGraph)
+		|| HeadPropertyName == GET_MEMBER_NAME_CHECKED(UPCGGraph, GraphUsageContext)
 		|| HeadPropertyName == GET_MEMBER_NAME_CHECKED(UPCGGraph, Category)
 		|| HeadPropertyName == GET_MEMBER_NAME_CHECKED(UPCGGraph, Description))
 	{
diff --git a/Engine/Plugins/PCG/Source/PCG/Public/PCGDefaultWorldObjectExecutionSource.h b/Engine/Plugins/PCG/Source/PCG/Public/PCGDefaultWorldObjectExecutionSource.h
--- a/Engine/Plugins/PCG/Source/PCG/Public/PCGDefaultWorldObjectExecutionSource.h
+++ b/Engine/Plugins/PCG/Source/PCG/Public/PCGDefaultWorldObjectExecutionSource.h
@@ -22,6 +22,8 @@ class UPCGDefaultWorldObjectExecutionSource : public UPCGDefaultExecutionSource
 	GENERATED_BODY()
 	
 public:
+	using ParamsType = FPCGDefaultWorldObjectExecutionSourceParams;
+
 	PCG_API UPCGDefaultWorldObjectExecutionSource();
 	PCG_API void Initialize(const FPCGDefaultWorldObjectExecutionSourceParams& InParams);
 	UObject* GetWorldObject() { return WorldObject.Get(); }
diff --git a/Engine/Plugins/PCG/Source/PCG/Public/PCGGraph.h b/Engine/Plugins/PCG/Source/PCG/Public/PCGGraph.h
--- a/Engine/Plugins/PCG/Source/PCG/Public/PCGGraph.h
+++ b/Engine/Plugins/PCG/Source/PCG/Public/PCGGraph.h
@@ -61,6 +61,17 @@ DECLARE_MULTICAST_DELEGATE_ThreeParams(FOnPCGGraphParametersChanged, UPCGGraphIn
 DECLARE_MULTICAST_DELEGATE_TwoParams(FOnPCGNodeSourceCompiled, const UPCGNode*, const FPCGCompilerDiagnostics&);
 #endif // WITH_EDITOR
 
+UENUM()
+enum class EPCGGraphUsage
+{
+	/** Used for graphs running on pcg components. */
+	Standard,
+	/** Used for standalone graphs that generate asset(s) and can be run from the content browser. */
+	Asset,
+	/** Used for standalone graphs that affect the currently loaded level and can be run from the content browser. */
+	Level
+};
+
 /**
 * Extended version of FInstancedPropertyBag, to support overrides and have a custom UI for it
 * Must only be used with PCGGraphInstances.
@@ -186,10 +197,13 @@ public:
 	virtual TOptional<FPCGGraphToolData> GetGraphToolData() const { return {}; }
 	
 	UE_API static bool IsStandaloneGraphAsset(const FAssetData& InAssetData);
+	UE_API static EPCGGraphUsage GetGraphUsage(const FAssetData& InAssetData);
 
 	UFUNCTION()
 	UE_API bool IsStandaloneGraph() const;
 
+	UE_API EPCGGraphUsage GetGraphUsage() const;
+
 	UFUNCTION()
 	UE_API virtual UObject* GetUserParameterHierarchyRoot() const; 
 protected:
@@ -412,8 +426,13 @@ public:
 	UPROPERTY(BlueprintReadWrite, EditAnywhere, Category = AssetInfo, AssetRegistrySearchable, meta=(PCGNoHash, EditCondition="ShowToolData()", EditConditionHides))
 	FPCGGraphToolData ToolData;
 
-	/** When enabled, this graph can be executed outside of the world using an editor execution source. */
+	/** Controls whether a graph can be executed outside of the context of a PCG actor/entity, e.g. run on for asset or level workflows. */
 	UPROPERTY(BlueprintReadWrite, EditAnywhere, Category = AssetInfo, AssetRegistrySearchable, meta = (PCGNoHash, EditCondition = "CanToggleStandaloneGraph()", EditConditionHides))
+	EPCGGraphUsage GraphUsageContext = EPCGGraphUsage::Standard;
+
+	/** When enabled, this graph can be executed outside of the world using an editor execution source. */
+	UE_DEPRECATED(5.8, "Use GraphUsageContext instead.")
+	UPROPERTY(BlueprintReadWrite, Category = AssetInfo, meta = (PCGNoHash, DeprecatedProperty, DeprecationMessage = "Deprecated in favor of GraphUsageContext"))
 	bool bIsStandaloneGraph = false;
 #endif
 
diff --git a/Engine/Plugins/PCG/Source/PCGEditor/Private/AssetDefinitions/AssetDefinition_PCGGraphInterface.cpp b/Engine/Plugins/PCG/Source/PCGEditor/Private/AssetDefinitions/AssetDefinition_PCGGraphInterface.cpp
--- a/Engine/Plugins/PCG/Source/PCGEditor/Private/AssetDefinitions/AssetDefinition_PCGGraphInterface.cpp
+++ b/Engine/Plugins/PCG/Source/PCGEditor/Private/AssetDefinitions/AssetDefinition_PCGGraphInterface.cpp
@@ -5,9 +5,11 @@
 #include "PCGGraph.h"
 
 #include "PCGDefaultExecutionSource.h"
+#include "PCGDefaultWorldObjectExecutionSource.h"
 #include "PCGEditorCommon.h"
 #include "PCGEditorStyle.h"
 #include "PCGGraphFactory.h"
+#include "PCGSubsystem.h"
 #include "Subsystems/PCGEngineSubsystem.h"
 
 #include "ContentBrowserMenuContexts.h"
@@ -85,12 +87,17 @@ namespace MenuExtension_PCGGraphInterface
 			{
 				const UContentBrowserAssetContextMenuContext* CBContext = UContentBrowserAssetContextMenuContext::FindContextWithAssets(InSection);
 
-				const bool bStandaloneGraphSelected = CBContext->SelectedAssets.ContainsByPredicate([](const FAssetData& InAsset)
+				const bool bAssetGraphSelected = CBContext->SelectedAssets.ContainsByPredicate([](const FAssetData& InAsset)
 				{
-					return UPCGGraphInterface::IsStandaloneGraphAsset(InAsset);
+					return UPCGGraphInterface::GetGraphUsage(InAsset) == EPCGGraphUsage::Asset;
 				});
 
-				if (bStandaloneGraphSelected)
+				const bool bLevelGraphSelected = CBContext->SelectedAssets.ContainsByPredicate([](const FAssetData& InAsset)
+				{
+					return UPCGGraphInterface::GetGraphUsage(InAsset) == EPCGGraphUsage::Level;
+				});
+
+				if (bAssetGraphSelected)
 				{
 					const TAttribute<FText> Label = LOCTEXT("GenerateStandaloneGraphsMenuLabel", "Generate Standalone Graph(s)");
 					const TAttribute<FText> ToolTip = LOCTEXT("GenerateStandaloneGraphsMenuToolTip", "Will generate every standalone graph in the selection.");
@@ -106,22 +113,62 @@ namespace MenuExtension_PCGGraphInterface
 							for (const FAssetData& AssetData : CBContext->SelectedAssets)
 							{
 								bool bTagValue = false;
-								const bool bIsStandaloneGraph = UPCGGraphInterface::IsStandaloneGraphAsset(AssetData);
+								const bool bIsAssetGraph = UPCGGraphInterface::GetGraphUsage(AssetData) == EPCGGraphUsage::Asset;
 								const bool bExportAssetActionAvailable = AssetData.IsInstanceOf<UPCGGraphInterface>() && AssetData.GetTagValue(GET_MEMBER_NAME_CHECKED(UPCGGraphInterface, bExposeGenerationInAssetExplorer), bTagValue) && bTagValue;
-								if (!bIsStandaloneGraph || !bExportAssetActionAvailable)
+								if (!bIsAssetGraph || !bExportAssetActionAvailable)
 								{
 									continue;
 								}
 
 								if (UPCGGraphInterface* GraphInterface = Cast<UPCGGraphInterface>(AssetData.GetAsset()))
 								{
-									EngineSubsystem->GenerateGraph({ GraphInterface });
+									IPCGBaseSubsystem::CreateExecutionSource<UPCGDefaultExecutionSource>({ .GraphInterface = GraphInterface, .bFireAndForgetExecution = true });
 								}
 							}
 						}
 					});
 					InSection.AddMenuEntry("GenerateStandaloneGraphsMenu", Label, ToolTip, Icon, UIAction);
 				}
+
+				if (UPCGSubsystem* Subsystem = UPCGSubsystem::GetSubsystemForCurrentWorld(); Subsystem && bLevelGraphSelected)
+				{
+					const TAttribute<FText> Label = LOCTEXT("GenerateGraphInWorldMenuLabel", "Runs Graph in current world.");
+					const TAttribute<FText> ToolTip = LOCTEXT("GenerateGraphInWorldMenuToolTip", "Will generate the graph in the world context.");
+					const FSlateIcon Icon = FSlateIcon(FPCGEditorStyle::Get().GetStyleSetName(), "ClassIcon.PCGGraph");
+
+					FToolUIAction UIAction;
+					UIAction.ExecuteAction.BindLambda([](const FToolMenuContext& InContext)
+					{
+						UPCGSubsystem* Subsystem = UPCGSubsystem::GetSubsystemForCurrentWorld();
+						if (!Subsystem)
+						{
+							return;
+						}
+
+						const UContentBrowserAssetContextMenuContext* CBContext = UContentBrowserAssetContextMenuContext::FindContextWithAssets(InContext);
+						for (const FAssetData& AssetData : CBContext->SelectedAssets)
+						{
+							bool bTagValue = false;
+							const bool bIsLevelGraph = UPCGGraphInterface::GetGraphUsage(AssetData) == EPCGGraphUsage::Level;
+							const bool bExportAssetActionAvailable = AssetData.IsInstanceOf<UPCGGraphInterface>() && AssetData.GetTagValue(GET_MEMBER_NAME_CHECKED(UPCGGraphInterface, bExposeGenerationInAssetExplorer), bTagValue) && bTagValue;
+							if (!bIsLevelGraph || !bExportAssetActionAvailable)
+							{
+								continue;
+							}
+
+							if (UPCGGraphInterface* GraphInterface = Cast<UPCGGraphInterface>(AssetData.GetAsset()))
+							{
+								FPCGDefaultWorldObjectExecutionSourceParams Params;
+								Params.GraphInterface = GraphInterface;
+								Params.WorldObject = Subsystem;
+								Params.bFireAndForgetExecution = true;
+
+								IPCGBaseSubsystem::CreateExecutionSource<UPCGDefaultWorldObjectExecutionSource>(Params);
+							}
+						}
+					});
+					InSection.AddMenuEntry("GenerateGraphInWorldMenu", Label, ToolTip, Icon, UIAction);
+				}
 			}));
 		}));
 	});
diff --git a/Engine/Plugins/PCG/Source/PCGEditor/Private/PCGEditor.cpp b/Engine/Plugins/PCG/Source/PCGEditor/Private/PCGEditor.cpp
--- a/Engine/Plugins/PCG/Source/PCGEditor/Private/PCGEditor.cpp
+++ b/Engine/Plugins/PCG/Source/PCGEditor/Private/PCGEditor.cpp
@@ -4,6 +4,7 @@
 
 #include "PCGComponent.h"
 #include "PCGDefaultExecutionSource.h"
+#include "PCGDefaultWorldObjectExecutionSource.h"
 #include "PCGEdge.h"
 #include "PCGEditorModule.h"
 #include "PCGGraph.h"
@@ -4261,12 +4262,11 @@ void FPCGEditor::RefreshViews()
 
 void FPCGEditor::UpdateDefaultExecutionSource()
 {
-	if (MainGraph->IsStandaloneGraph() && !PCGDefaultExecutionSource)
-	{
-		PCGDefaultExecutionSource = IPCGBaseSubsystem::CreateExecutionSource<UPCGDefaultExecutionSource>({ .GraphInterface = MainGraph });
-		RefreshViews();
-	}
-	else if (!MainGraph->IsStandaloneGraph() && PCGDefaultExecutionSource)
+	const EPCGGraphUsage PreviousUsage = (PCGDefaultExecutionSource ? (Cast<UPCGDefaultWorldObjectExecutionSource>(PCGDefaultExecutionSource) != nullptr ? EPCGGraphUsage::Level : EPCGGraphUsage::Asset) : EPCGGraphUsage::Standard);
+	const EPCGGraphUsage CurrentUsage = MainGraph->GetGraphUsage();
+	bool bNeedsRefreshView = false;
+
+	if (PreviousUsage != CurrentUsage && PCGDefaultExecutionSource)
 	{
 		if (IPCGEditorModule* PCGEditorModule = IPCGEditorModule::Get())
 		{
@@ -4274,6 +4274,28 @@ void FPCGEditor::UpdateDefaultExecutionSource()
 		}
 		PCGDefaultExecutionSource = nullptr;
 		CollectGarbage(GARBAGE_COLLECTION_KEEPFLAGS);
+		bNeedsRefreshView = true;
+	}
+
+	if (CurrentUsage != EPCGGraphUsage::Standard && !PCGDefaultExecutionSource)
+	{
+		if (CurrentUsage == EPCGGraphUsage::Asset)
+		{
+			PCGDefaultExecutionSource = IPCGBaseSubsystem::CreateExecutionSource<UPCGDefaultExecutionSource>({ .GraphInterface = MainGraph });
+		}
+		else if (CurrentUsage == EPCGGraphUsage::Level)
+		{
+			FPCGDefaultWorldObjectExecutionSourceParams Params;
+			Params.GraphInterface = MainGraph;
+			Params.WorldObject = UPCGSubsystem::GetSubsystemForCurrentWorld();
+			PCGDefaultExecutionSource = IPCGBaseSubsystem::CreateExecutionSource<UPCGDefaultWorldObjectExecutionSource>(Params);
+		}
+
+		bNeedsRefreshView = true;
+	}
+
+	if (bNeedsRefreshView)
+	{
 		RefreshViews();
 	}
 }
