[PCG] PR #14547: Fix Dynamic Mesh data not displaying data domain attributes in the ALV

#rnx
#Synced-CL 51771932

diff --git a/Engine/Plugins/PCG/Source/PCG/Private/Data/PCGDynamicMeshData.cpp b/Engine/Plugins/PCG/Source/PCG/Private/Data/PCGDynamicMeshData.cpp
--- a/Engine/Plugins/PCG/Source/PCG/Private/Data/PCGDynamicMeshData.cpp
+++ b/Engine/Plugins/PCG/Source/PCG/Private/Data/PCGDynamicMeshData.cpp
@@ -191,6 +191,7 @@ const UPCGBasePointData* UPCGDynamicMeshData::ToBasePointData(FPCGContext* Conte
 	UPCGBasePointData* Data = FPCGContext::NewObject_AnyThread<UPCGBasePointData>(Context, GetTransientPackage(), PointDataClass);
 	Data->SetNumPoints(Vertices.Num());
 	Data->AllocateProperties(EPCGPointNativeProperties::Transform | EPCGPointNativeProperties::Seed);
+	Data->MutableMetadata()->Initialize(ConstMetadata());
 
 	int32 NumWritten = 0;
 	FPCGPointValueRanges OutRanges(Data, /*bAllocate=*/false);
