[PCG] Make sure we are sanitizing the names in the property

* Add the strict option to the BP helper

#rnx
#Synced-CL 51771932

diff --git a/Engine/Plugins/PCG/Source/PCG/Private/Helpers/PCGBlueprintHelpers.cpp b/Engine/Plugins/PCG/Source/PCG/Private/Helpers/PCGBlueprintHelpers.cpp
--- a/Engine/Plugins/PCG/Source/PCG/Private/Helpers/PCGBlueprintHelpers.cpp
+++ b/Engine/Plugins/PCG/Source/PCG/Private/Helpers/PCGBlueprintHelpers.cpp
@@ -408,6 +408,7 @@ DEFINE_FUNCTION(UPCGBlueprintHelpers::execExtractStructureAsAttributeSet)
 	P_GET_PROPERTY(FNameProperty, PropertyPath);
 	P_GET_PROPERTY(FNameProperty, OutputAttributeName);
 	P_GET_PROPERTY(FBoolProperty, bExtract);
+	P_GET_PROPERTY(FBoolProperty, bStrictlySanitizeNames);
 
 	P_FINISH;
 
@@ -422,6 +423,7 @@ DEFINE_FUNCTION(UPCGBlueprintHelpers::execExtractStructureAsAttributeSet)
 	if (bHasValidStructProp)
 	{
 		PCGPropertyHelpers::FExtractorParameters Params(StructPtr, StructProp->Struct, InputSelector, OutputAttributeName, bExtract, /*bInPropertyNeedsToBeVisible=*/true);
+		Params.bStrictSanitizeOutputAttributeNames = bStrictlySanitizeNames;
 		Result = PCGPropertyHelpers::ExtractPropertyAsAttributeSet(Params);
 		if (!Result)
 		{
diff --git a/Engine/Plugins/PCG/Source/PCG/Public/Helpers/PCGBlueprintHelpers.h b/Engine/Plugins/PCG/Source/PCG/Public/Helpers/PCGBlueprintHelpers.h
--- a/Engine/Plugins/PCG/Source/PCG/Public/Helpers/PCGBlueprintHelpers.h
+++ b/Engine/Plugins/PCG/Source/PCG/Public/Helpers/PCGBlueprintHelpers.h
@@ -181,9 +181,12 @@ private:
 	 * @param PropertyPath Path to the property in the structure to extract. Must be a valid property. If None, will extract the structure as input.
 	 * @param OutputAttributeName If the property to extract is not a struct/object, can provide an override to the output attribute name. Use @Source to use the PropertyName.
 	 * @param bExtract If the property path is not None and reference an object/struct, give the possibility to extract all its members in one go.
+	 * @param bStrictlySanitizeNames Output attribute name are already sanitized by default but it keeps some characters like spaces.
+	 * This is a stricter option and if the name is not following the object name rule (pretty much only alphanumerical and "_"), it will be sanitized.
+	 * Necessary if we have cases with spaces, which doesn't play nicely with PCG attributes.
 	 * @return Extracted param data. Can be null if something went wrong.
 	 */
 	UFUNCTION(BlueprintCallable, CustomThunk, Category = "PCG|Data", meta = (CustomStructureParam = "Struct", AllowAbstract = "false"))
-	static PCG_API UPCGParamData* ExtractStructureAsAttributeSet(UPARAM(DisplayName = "Structure To Convert") const int32& Struct, const FName PropertyPath = NAME_None, const FName OutputAttributeName = "@Source", const bool bExtract = true);
+	static PCG_API UPCGParamData* ExtractStructureAsAttributeSet(UPARAM(DisplayName = "Structure To Convert") const int32& Struct, const FName PropertyPath = NAME_None, const FName OutputAttributeName = "@Source", const bool bExtract = true, const bool bStrictlySanitizeNames = true);
 	DECLARE_FUNCTION(execExtractStructureAsAttributeSet);
 };
