[PCG] Fixed some issues with create spline data, spawn actor and create target actor for level asset workflows.

#Synced-CL 51908749

diff --git a/Engine/Plugins/PCG/Source/PCG/Private/Elements/PCGCreateSpline.cpp b/Engine/Plugins/PCG/Source/PCG/Private/Elements/PCGCreateSpline.cpp
--- a/Engine/Plugins/PCG/Source/PCG/Private/Elements/PCGCreateSpline.cpp
+++ b/Engine/Plugins/PCG/Source/PCG/Private/Elements/PCGCreateSpline.cpp
@@ -89,7 +89,7 @@ bool FPCGCreateSplineElement::ExecuteInternal(FPCGContext* Context) const
 		}
 
 		AActor* TargetActor = Settings->TargetActor.Get() ? Settings->TargetActor.Get() : Context->GetTypedExecutionTarget<AActor>();
-		if (!TargetActor)
+		if (!TargetActor && Settings->Mode != EPCGCreateSplineMode::CreateDataOnly)
 		{
 			PCGE_LOG(Error, GraphAndLog, LOCTEXT("InvalidTargetActor", "Invalid target actor. Ensure TargetActor member is initialized when creating SpatialData."));
 			continue;
@@ -163,7 +163,7 @@ bool FPCGCreateSplineElement::ExecuteInternal(FPCGContext* Context) const
 		SplineEntryKeys.Reserve(NumPoints);
 		bool bHasAValidEntry = false;
 
-		const FTransform SplineActorTransform = SplineActor->GetTransform();
+		const FTransform SplineActorTransform = SplineActor ? SplineActor->GetTransform() : FTransform::Identity;
 		
 		ESplinePointType::Type FixedPointType = ESplinePointType::Curve;
 		if (Settings->bLinear)
@@ -207,6 +207,7 @@ bool FPCGCreateSplineElement::ExecuteInternal(FPCGContext* Context) const
 		if (Settings->Mode != EPCGCreateSplineMode::CreateDataOnly)
 		{
 			check(IsInGameThread());
+			check(SplineActor);
 
 			IPCGGraphExecutionSource* ExecutionSource = Context->ExecutionSource.Get();
 			FPCGManagedResourceContainerHelper ManagedResourcesContainerHelper(ExecutionSource);
@@ -237,11 +238,7 @@ bool FPCGCreateSplineElement::ExecuteInternal(FPCGContext* Context) const
 
 				ManagedResourcesContainerHelper.AddManagedResource(ManagedComponent);
 			}
-		}
 
-		// Execute PostProcess Functions
-		if (SplineActor)
-		{
 			// @todo_pcg: for 5.7 move lock inside UPCGFunctionPrototypes::GetPrototypeWithNoParams().
 			// for 5.6.1 keeping fix in cpp to preserve abi.
 			// This can run off the GameThread and UPCGFunctionPrototypes::GetPrototypeWithNoParams() does a FindObject which can't run while GC is running
@@ -251,6 +248,7 @@ bool FPCGCreateSplineElement::ExecuteInternal(FPCGContext* Context) const
 				FunctionPrototype = UPCGFunctionPrototypes::GetPrototypeWithNoParams();
 			}
 
+			// Execute PostProcess Functions
 			for (UFunction* Function : PCGHelpers::FindUserFunctions(SplineActor->GetClass(), Settings->PostProcessFunctionNames, { FunctionPrototype }, Context))
 			{
 				SplineActor->ProcessEvent(Function, nullptr);
diff --git a/Engine/Plugins/PCG/Source/PCG/Private/Elements/PCGCreateTargetActor.cpp b/Engine/Plugins/PCG/Source/PCG/Private/Elements/PCGCreateTargetActor.cpp
--- a/Engine/Plugins/PCG/Source/PCG/Private/Elements/PCGCreateTargetActor.cpp
+++ b/Engine/Plugins/PCG/Source/PCG/Private/Elements/PCGCreateTargetActor.cpp
@@ -363,9 +363,11 @@ bool FPCGCreateTargetActorElement::ExecuteInternal(FPCGContext* Context) const
 	}
 
 	AActor* TargetActor = Settings->RootActor.Get() ? Settings->RootActor.Get() : Context->GetTypedExecutionTarget<AActor>();
-	if (!TargetActor)
+	UObject* TargetWorldObject = TargetActor ? TargetActor : Context->GetExecutionTarget();
+
+	if (!TargetWorldObject)
 	{
-		PCGE_LOG(Error, GraphAndLog, LOCTEXT("InvalidTargetActor", "Invalid target actor"));
+		PCGE_LOG(Error, GraphAndLog, LOCTEXT("InvalidTargetObject", "Invalid target object. The action will fail because it is required to obtain the world."));
 		return true;
 	}
 
@@ -384,25 +386,31 @@ bool FPCGCreateTargetActorElement::ExecuteInternal(FPCGContext* Context) const
 	FActorSpawnParameters SpawnParams;
 	SpawnParams.Template = TemplateActor;
 	SpawnParams.SpawnCollisionHandlingOverride = ESpawnActorCollisionHandlingMethod::AlwaysSpawn;
-	SpawnParams.OverrideLevel = TargetActor->GetLevel();
+	SpawnParams.OverrideLevel = TargetActor ? TargetActor->GetLevel() : Cast<ULevel>(TargetWorldObject);
 
-	FTransform Transform = TargetActor->GetTransform();
+	FTransform Transform = TargetActor ? TargetActor->GetTransform() : FTransform::Identity;
 	if(Context->IsValueOverriden(GET_MEMBER_NAME_CHECKED(UPCGCreateTargetActor, ActorPivot)))
 	{
 		Transform = Settings->ActorPivot;
 	}
 
-	UPCGActorHelpers::FSpawnDefaultActorParams SpawnDefaultActorParams(TargetActor->GetWorld(), Settings->TemplateActorClass, Transform, SpawnParams);
+	UPCGActorHelpers::FSpawnDefaultActorParams SpawnDefaultActorParams(TargetWorldObject->GetWorld(), Settings->TemplateActorClass, Transform, SpawnParams);
 	SpawnDefaultActorParams.bIsRuntime = (ExecutionSource && ExecutionSource->GetExecutionState().IsManagedByRuntimeGenSystem());
 #if WITH_EDITOR
 	SpawnDefaultActorParams.bIsPreviewActor = (ExecutionSource && ExecutionSource->GetExecutionState().IsInPreviewMode());
 
-	int32 DataLayerCrc;
-	SpawnDefaultActorParams.DataLayerInstances = PCGDataLayerHelpers::GetDataLayerInstancesAndCrc(Context, Settings->DataLayerSettings, TargetActor, DataLayerCrc);
+	int32 DataLayerCrc = 0;
+	if (TargetActor)
+	{
+		SpawnDefaultActorParams.DataLayerInstances = PCGDataLayerHelpers::GetDataLayerInstancesAndCrc(Context, Settings->DataLayerSettings, TargetActor, DataLayerCrc);
+	}
 
-	int32 HLODLayerCrc;
-	AActor* TemplateActorOrDefault = TemplateActor ? TemplateActor : Cast<AActor>(Settings->TemplateActorClass->GetDefaultObject());
-	SpawnDefaultActorParams.HLODLayer = PCGHLODHelpers::GetHLODLayerAndCrc(Context, Settings->HLODSettings, TargetActor, TemplateActorOrDefault, HLODLayerCrc);
+	int32 HLODLayerCrc = 0;
+	if (TargetActor)
+	{
+		AActor* TemplateActorOrDefault = TemplateActor ? TemplateActor : Cast<AActor>(Settings->TemplateActorClass->GetDefaultObject());
+		SpawnDefaultActorParams.HLODLayer = PCGHLODHelpers::GetHLODLayerAndCrc(Context, Settings->HLODSettings, TargetActor, TemplateActorOrDefault, HLODLayerCrc);
+	}
 #endif
 
 	FPCGManagedResourceContainerHelper ManagedResourcesContainerHelper(ExecutionSource);
diff --git a/Engine/Plugins/PCG/Source/PCG/Private/Elements/PCGSpawnActor.cpp b/Engine/Plugins/PCG/Source/PCG/Private/Elements/PCGSpawnActor.cpp
--- a/Engine/Plugins/PCG/Source/PCG/Private/Elements/PCGSpawnActor.cpp
+++ b/Engine/Plugins/PCG/Source/PCG/Private/Elements/PCGSpawnActor.cpp
@@ -22,6 +22,7 @@
 #include "Engine/Blueprint.h"
 #include "Engine/BlueprintGeneratedClass.h"
 #include "Engine/Engine.h"
+#include "Engine/Level.h"
 #include "ISMPartition/ISMComponentDescriptor.h"
 #include "UObject/Package.h"
 
@@ -551,13 +552,20 @@ bool FPCGSpawnActorElement::SpawnAndPrepareSubgraphs(FPCGSubgraphContext* Contex
 		}
 
 		AActor* TargetActor = Settings->RootActor.Get() ? Settings->RootActor.Get() : Context->GetTypedExecutionTarget<AActor>();
+		UObject* TargetWorldObject = TargetActor ? TargetActor : Context->GetExecutionTarget();
 
-		if (!TargetActor)
+		if (!TargetActor && Settings->Option == EPCGSpawnActorOption::CollapseActors)
 		{
 			PCGE_LOG(Error, GraphAndLog, LOCTEXT("InvalidTargetActor", "Invalid target actor. Ensure TargetActor member is initialized when creating SpatialData."));
 			continue;
 		}
 
+		if (!TargetWorldObject)
+		{
+			PCGE_LOG(Error, GraphAndLog, LOCTEXT("InvalidTargetObject", "Invalid target object. Will be unable to spawn actors since we cannot resolve a world."));
+			continue;
+		}
+
 		// First, create target instance transforms
 		const UPCGBasePointData* PointData = SpatialData->ToBasePointData(Context);
 
@@ -573,7 +581,7 @@ bool FPCGSpawnActorElement::SpawnAndPrepareSubgraphs(FPCGSubgraphContext* Contex
 			continue;
 		}
 
-		auto SpawnOrCollapse = [this, bHasAuthority, &Context, &TargetActor, &Settings](TSubclassOf<AActor> TemplateActorClass, AActor* TemplateActor, FPCGTaggedData& Output, const UPCGBasePointData* PointData, UPCGBasePointData* OutPointData)
+		auto SpawnOrCollapse = [this, bHasAuthority, &Context, &TargetActor, &TargetWorldObject, &Settings](TSubclassOf<AActor> TemplateActorClass, AActor* TemplateActor, FPCGTaggedData& Output, const UPCGBasePointData* PointData, UPCGBasePointData* OutPointData)
 		{
 			const bool bSpawnedActorsRequireAuthority = (TemplateActor ? TemplateActor->GetIsReplicated() : CastChecked<AActor>(TemplateActorClass->GetDefaultObject())->GetIsReplicated());
 
@@ -583,7 +591,7 @@ bool FPCGSpawnActorElement::SpawnAndPrepareSubgraphs(FPCGSubgraphContext* Contex
 			}
 			else if (bHasAuthority || !bSpawnedActorsRequireAuthority)
 			{
-				SpawnActors(Context, TargetActor, TemplateActorClass, TemplateActor, Output, PointData, OutPointData);
+				SpawnActors(Context, TargetWorldObject, TemplateActorClass, TemplateActor, Output, PointData, OutPointData);
 			}
 		};
 
@@ -784,10 +792,10 @@ void FPCGSpawnActorElement::CollapseIntoTargetActor(FPCGSubgraphContext* Context
 	}
 }
 
-void FPCGSpawnActorElement::SpawnActors(FPCGSubgraphContext* Context, AActor* TargetActor, TSubclassOf<AActor> InTemplateActorClass, AActor* InTemplateActor, FPCGTaggedData& Output, const UPCGBasePointData* PointData, UPCGBasePointData* OutPointData) const
+void FPCGSpawnActorElement::SpawnActors(FPCGSubgraphContext* Context, UObject* TargetWorldObject, TSubclassOf<AActor> InTemplateActorClass, AActor* InTemplateActor, FPCGTaggedData& Output, const UPCGBasePointData* PointData, UPCGBasePointData* OutPointData) const
 {
 	TRACE_CPUPROFILER_EVENT_SCOPE(FPCGSpawnActorElement::ExecuteInternal::SpawnActors);
-	check(Context && TargetActor && PointData);
+	check(Context && TargetWorldObject && PointData);
 	check(IsInGameThread());
 
 	if (PointData->GetNumPoints() == 0 || !InTemplateActorClass)
@@ -861,7 +869,9 @@ void FPCGSpawnActorElement::SpawnActors(FPCGSubgraphContext* Context, AActor* Ta
 	FActorSpawnParameters SpawnParams;
 	SpawnParams.Template = TemplateActor;
 	SpawnParams.SpawnCollisionHandlingOverride = ESpawnActorCollisionHandlingMethod::AlwaysSpawn;
-	SpawnParams.OverrideLevel = TargetActor->GetLevel();
+
+	AActor* TargetActor = Cast<AActor>(TargetWorldObject);
+	SpawnParams.OverrideLevel = TargetActor ? TargetActor->GetLevel() : Cast<ULevel>(TargetWorldObject);
 
 	IPCGGraphExecutionSource* Source = Context->ExecutionSource.Get();
 
@@ -870,17 +880,24 @@ void FPCGSpawnActorElement::SpawnActors(FPCGSubgraphContext* Context, AActor* Ta
 		return; // won't be able to create resources correctly
 	}
 
-	UPCGActorHelpers::FSpawnDefaultActorParams SpawnDefaultActorParams(TargetActor->GetWorld(), InTemplateActorClass, FTransform::Identity, SpawnParams);
+	UPCGActorHelpers::FSpawnDefaultActorParams SpawnDefaultActorParams(TargetWorldObject->GetWorld(), InTemplateActorClass, FTransform::Identity, SpawnParams);
 	SpawnDefaultActorParams.bForceStaticMobility = false; // Always respect the actor's mobility
 	SpawnDefaultActorParams.bIsRuntime = (Source && Source->GetExecutionState().IsManagedByRuntimeGenSystem());
 #if WITH_EDITOR
 	SpawnDefaultActorParams.bIsPreviewActor = (Source && Source->GetExecutionState().IsInPreviewMode());
 
-	int32 DataLayerCrc;
-	SpawnDefaultActorParams.DataLayerInstances = PCGDataLayerHelpers::GetDataLayerInstancesAndCrc(Context, Settings->DataLayerSettings, TargetActor, DataLayerCrc);
+	int32 DataLayerCrc = 0;
+	if (TargetActor)
+	{
+		SpawnDefaultActorParams.DataLayerInstances = PCGDataLayerHelpers::GetDataLayerInstancesAndCrc(Context, Settings->DataLayerSettings, TargetActor, DataLayerCrc);
+	}
+	
 
-	int32 HLODLayerCrc;
-	SpawnDefaultActorParams.HLODLayer = PCGHLODHelpers::GetHLODLayerAndCrc(Context, Settings->HLODSettings, TargetActor, TemplateActor, HLODLayerCrc);
+	int32 HLODLayerCrc = 0;
+	if (TargetActor)
+	{
+		SpawnDefaultActorParams.HLODLayer = PCGHLODHelpers::GetHLODLayerAndCrc(Context, Settings->HLODSettings, TargetActor, TemplateActor, HLODLayerCrc);
+	}
 #endif
 
 	const bool bForceCallGenerate = (Settings->GenerationTrigger == EPCGSpawnActorGenerationTrigger::ForceGenerate);
@@ -1145,7 +1162,7 @@ TArray<FName> FPCGSpawnActorElement::GetNewActorTags(FPCGContext* Context, AActo
 {
 	TArray<FName> NewActorTags;
 	// Prepare actor tags
-	if (bInheritActorTags)
+	if (bInheritActorTags && TargetActor)
 	{
 		// Special case: if the current target actor is a partition, we'll reach out
 		// and find the original actor tags
diff --git a/Engine/Plugins/PCG/Source/PCG/Private/PCGDefaultWorldObjectExecutionSource.cpp b/Engine/Plugins/PCG/Source/PCG/Private/PCGDefaultWorldObjectExecutionSource.cpp
--- a/Engine/Plugins/PCG/Source/PCG/Private/PCGDefaultWorldObjectExecutionSource.cpp
+++ b/Engine/Plugins/PCG/Source/PCG/Private/PCGDefaultWorldObjectExecutionSource.cpp
@@ -14,6 +14,11 @@ FString FPCGDefaultWorldObjectExecutionState::GetDebugName() const
 	return FString(TEXT("PCGDefaultWorldObjectExecutionSource - ")) + (GetWorldObject() ? GetWorldObject()->GetName() : TEXT("Unknown object"));
 }
 
+UObject* FPCGDefaultWorldObjectExecutionState::GetTarget() const
+{
+	return GetWorldObject();
+}
+
 UWorld* FPCGDefaultWorldObjectExecutionState::GetWorld() const
 {
 	if (UObject* Object = GetWorldObject())
diff --git a/Engine/Plugins/PCG/Source/PCG/Public/Elements/PCGSpawnActor.h b/Engine/Plugins/PCG/Source/PCG/Public/Elements/PCGSpawnActor.h
--- a/Engine/Plugins/PCG/Source/PCG/Public/Elements/PCGSpawnActor.h
+++ b/Engine/Plugins/PCG/Source/PCG/Public/Elements/PCGSpawnActor.h
@@ -204,5 +204,5 @@ private:
 	TArray<FName> GetNewActorTags(FPCGContext* Context, AActor* TargetActor, bool bInheritActorTags, const TArray<FName>& AdditionalTags) const;
 
 	void CollapseIntoTargetActor(FPCGSubgraphContext* Context, AActor* TargetActor, TSubclassOf<AActor> TemplateActorClass, const UPCGBasePointData* PointData) const;
-	void SpawnActors(FPCGSubgraphContext* Context, AActor* TargetActor, TSubclassOf<AActor> TemplateActorClass, AActor* TemplateActor, FPCGTaggedData& Output, const UPCGBasePointData* PointData, UPCGBasePointData* OutPointData) const;
+	void SpawnActors(FPCGSubgraphContext* Context, UObject* TargetWorldObject, TSubclassOf<AActor> TemplateActorClass, AActor* TemplateActor, FPCGTaggedData& Output, const UPCGBasePointData* PointData, UPCGBasePointData* OutPointData) const;
 };
diff --git a/Engine/Plugins/PCG/Source/PCG/Public/PCGDefaultWorldObjectExecutionSource.h b/Engine/Plugins/PCG/Source/PCG/Public/PCGDefaultWorldObjectExecutionSource.h
--- a/Engine/Plugins/PCG/Source/PCG/Public/PCGDefaultWorldObjectExecutionSource.h
+++ b/Engine/Plugins/PCG/Source/PCG/Public/PCGDefaultWorldObjectExecutionSource.h
@@ -41,6 +41,7 @@ public:
 
 	PCG_API virtual FString GetDebugName() const override;
 	PCG_API virtual UWorld* GetWorld() const override;
+	PCG_API virtual UObject* GetTarget() const override;
 	virtual FPCGManagedResourceContainer* GetManagedResourceContainer() override { return &ManagedResourceContainer; }
 	virtual FTransactionallySafeCriticalSection* GetManagedResourceContainerLock() override { return &ManagedResourcesLock; }
 
