[PCG] Fix regression from CL 51376868: Where GetOriginalLocalSpaceBounds was replaced by GetBounds instead of GetLocalSpaceBounds()

#rnx
#Synced-CL 52019349

diff --git a/Engine/Plugins/PCG/Source/PCG/Private/Elements/PCGTextureSampler.cpp b/Engine/Plugins/PCG/Source/PCG/Private/Elements/PCGTextureSampler.cpp
--- a/Engine/Plugins/PCG/Source/PCG/Private/Elements/PCGTextureSampler.cpp
+++ b/Engine/Plugins/PCG/Source/PCG/Private/Elements/PCGTextureSampler.cpp
@@ -238,7 +238,7 @@ bool FPCGTextureSamplerElement::ExecuteInternal(FPCGContext* InContext) const
 			const FTransform Transform = Context->ExecutionSource->GetExecutionState().GetTransform();
 			Context->Transform = Settings->Transform * Transform;
 
-			const FBox Bounds = Context->ExecutionSource->GetExecutionState().GetBounds();
+			const FBox Bounds = Context->ExecutionSource->GetExecutionState().GetLocalSpaceBounds();
 			Context->Transform.SetScale3D(Context->Transform.GetScale3D() * Bounds.GetExtent());
 		}
 
