[PCG] Better error handling on potentially bad execution targets. Removed useless dll exports.

#rnx
#Synced-CL 52060913

diff --git a/Engine/Plugins/PCG/Source/PCG/Private/Elements/PCGCreateTargetActor.cpp b/Engine/Plugins/PCG/Source/PCG/Private/Elements/PCGCreateTargetActor.cpp
--- a/Engine/Plugins/PCG/Source/PCG/Private/Elements/PCGCreateTargetActor.cpp
+++ b/Engine/Plugins/PCG/Source/PCG/Private/Elements/PCGCreateTargetActor.cpp
@@ -365,7 +365,7 @@ bool FPCGCreateTargetActorElement::ExecuteInternal(FPCGContext* Context) const
 	AActor* TargetActor = Settings->RootActor.Get() ? Settings->RootActor.Get() : Context->GetTypedExecutionTarget<AActor>();
 	UObject* TargetWorldObject = TargetActor ? TargetActor : Context->GetExecutionTarget();
 
-	if (!TargetWorldObject)
+	if (!TargetWorldObject || !TargetWorldObject->GetWorld())
 	{
 		PCGE_LOG(Error, GraphAndLog, LOCTEXT("InvalidTargetObject", "Invalid target object. The action will fail because it is required to obtain the world."));
 		return true;
diff --git a/Engine/Plugins/PCG/Source/PCG/Private/Elements/PCGSpawnActor.cpp b/Engine/Plugins/PCG/Source/PCG/Private/Elements/PCGSpawnActor.cpp
--- a/Engine/Plugins/PCG/Source/PCG/Private/Elements/PCGSpawnActor.cpp
+++ b/Engine/Plugins/PCG/Source/PCG/Private/Elements/PCGSpawnActor.cpp
@@ -560,7 +560,7 @@ bool FPCGSpawnActorElement::SpawnAndPrepareSubgraphs(FPCGSubgraphContext* Contex
 			continue;
 		}
 
-		if (!TargetWorldObject)
+		if (!TargetWorldObject || !TargetWorldObject->GetWorld())
 		{
 			PCGE_LOG(Error, GraphAndLog, LOCTEXT("InvalidTargetObject", "Invalid target object. Will be unable to spawn actors since we cannot resolve a world."));
 			continue;
diff --git a/Engine/Plugins/PCG/Source/PCGEditor/Public/PCGEditor.h b/Engine/Plugins/PCG/Source/PCGEditor/Public/PCGEditor.h
--- a/Engine/Plugins/PCG/Source/PCGEditor/Public/PCGEditor.h
+++ b/Engine/Plugins/PCG/Source/PCGEditor/Public/PCGEditor.h
@@ -255,7 +255,7 @@ PRAGMA_ENABLE_DEPRECATION_WARNINGS
 
 	FActionMenuContent OnCreateActionMenu(UEdGraph* InGraph, const FVector2f& Location, const TArray<UEdGraphPin*>& InDraggedPins, bool bAutoExpand, SGraphEditor::FActionMenuClosed OnMenuClosed);
 
-	/** Returns whether contextual menu filtering should be applied. */
+	/** Returns contextual menu filtering mask (no mask, compatible only, accept filtering and/or conversions). */
 	PCGActionsHelpers::ECompatibilityMask GetContextualCompatibilityMask() const;
 
 	/** Can determinism be tested on the current graph */
@@ -295,13 +295,13 @@ protected:
 	PCGEDITOR_API virtual TAttribute<FGraphAppearanceInfo> GetAppearanceInfo() const;
 
 	UE_DEPRECATED(5.8, "Implement version with FToolbarBuilder parameter instead")
-	PCGEDITOR_API virtual void RegisterToolbarInternal(FToolMenuSection& PCGSection) const {}
+	virtual void RegisterToolbarInternal(FToolMenuSection& PCGSection) const {}
 
 	/** Register PCG specific toolbar for the editor */
 	PCGEDITOR_API virtual void RegisterToolbarInternal(FToolBarBuilder& ToolbarBuilder) const;
 
 	UE_DEPRECATED(5.8, "Use version with FToolBarBuilder parameter instead")
-	PCGEDITOR_API void RegisterToolbarButton(FToolMenuSection& Section, EPCGToolbarButtons Button) const {}
+	void RegisterToolbarButton(FToolMenuSection& Section, EPCGToolbarButtons Button) const {}
 
 	/** Register PCG default editor toolbar button */
 	PCGEDITOR_API void RegisterToolbarButton(FToolBarBuilder& ToolbarBuilder, EPCGToolbarButtons Button) const;
