[PCG] Sanitize pin labels for HLSL code generation

Pin labels are used as HLSL identifier prefixes in GPU compute graph compilation. Labels with characters invalid in HLSL (e.g. "End Points") produced broken generated code.

Added SanitizePinLabelForHLSL() which replaces invalid characters with underscores and trims leading non-alpha characters (e.g. "End Points" ??? "End_Points", "3D_Data" ??? "D_Data"). Applied it in the graph compiler and as a proactive fixup on pin label edit in   Custom HLSL nodes. Fixed keyword token parsing in PCGCustomHLSLKernel to correctly handle pin names that now contain underscores.

#rnx

#Synced-CL 51948399

diff --git a/Engine/Plugins/PCG/Source/PCG/Private/Compute/Elements/PCGCustomHLSL.cpp b/Engine/Plugins/PCG/Source/PCG/Private/Compute/Elements/PCGCustomHLSL.cpp
--- a/Engine/Plugins/PCG/Source/PCG/Private/Compute/Elements/PCGCustomHLSL.cpp
+++ b/Engine/Plugins/PCG/Source/PCG/Private/Compute/Elements/PCGCustomHLSL.cpp
@@ -167,6 +167,32 @@ void UPCGCustomHLSLSettings::PostEditChangeProperty(FPropertyChangedEvent& Prope
 		}
 	}
 
+	// Sanitize pin labels to ensure they are valid HLSL identifiers.
+	if (Property == GET_MEMBER_NAME_CHECKED(FPCGPinProperties, Label))
+	{
+		auto SanitizePinLabel = [](FPCGPinProperties& PinProps)
+		{
+			if (!PCGComputeHelpers::IsValidHLSLPinLabel(PinProps.Label))
+			{
+				const FString OriginalLabel = PinProps.Label.ToString();
+				const FString SanitizedLabel = PCGComputeHelpers::SanitizePinLabelForHLSL(PinProps.Label);
+				PinProps.Label = FName(SanitizedLabel);
+
+				UE_LOG(LogPCG, Warning, TEXT("Pin label '%s' contained invalid HLSL characters and was sanitized to '%s'."), *OriginalLabel, *SanitizedLabel);
+			}
+		};
+
+		for (FPCGPinProperties& PinProps : InputPins)
+		{
+			SanitizePinLabel(PinProps);
+		}
+
+		for (FPCGPinPropertiesGPU& PinProps : OutputPins)
+		{
+			SanitizePinLabel(PinProps);
+		}
+	}
+
 	Super::PostEditChangeProperty(PropertyChangedEvent);
 
 	UpdateDeclarations();
@@ -460,7 +486,7 @@ void UPCGCustomHLSLSettings::UpdateInputDeclarations()
 				InputDeclarations += FString::Format(TEXT(
 					"const uint {0}_DataIndex;\n"
 					"const uint {1}_DataIndex;\n"),
-					{ InputPin->Label.ToString(),  OutputPin->Label.ToString() });
+					{ InputPin->Label.ToString(), OutputPin->Label.ToString() });
 			}
 		}
 		else if (IsGeneratorKernel())
diff --git a/Engine/Plugins/PCG/Source/PCG/Private/Compute/Elements/PCGCustomHLSLKernel.cpp b/Engine/Plugins/PCG/Source/PCG/Private/Compute/Elements/PCGCustomHLSLKernel.cpp
--- a/Engine/Plugins/PCG/Source/PCG/Private/Compute/Elements/PCGCustomHLSLKernel.cpp
+++ b/Engine/Plugins/PCG/Source/PCG/Private/Compute/Elements/PCGCustomHLSLKernel.cpp
@@ -864,7 +864,7 @@ void UPCGCustomHLSLKernel::ParseShaderSource()
 
 	for (const FPCGPinProperties& PinProps : CustomHLSLSettings->OutputPins)
 	{
-		const FString PinStr = *PinProps.Label.ToString();
+		const FString PinStr = PinProps.Label.ToString();
 
 		if (!!(PinProps.AllowedTypes & (EPCGDataType::Point | EPCGDataType::Param)))
 		{
@@ -873,7 +873,7 @@ void UPCGCustomHLSLKernel::ParseShaderSource()
 				AttributeKeywords.Add(PinStr + TEXT("_") + PCGCustomHLSLKernel::AttributeFunctionSetKeyword + AttributeTypeString);
 			}
 
-			AttributeKeywords.Add(PinProps.Label.ToString() + TEXT("_") + PCGCustomHLSLKernel::AttributeFunctionAtomicAddKeyword + PCGCustomHLSLKernel::AttributeFunctionAtomicAddType);
+			AttributeKeywords.Add(PinStr + TEXT("_") + PCGCustomHLSLKernel::AttributeFunctionAtomicAddKeyword + PCGCustomHLSLKernel::AttributeFunctionAtomicAddType);
 
 			for (const FPCGPinProperties& InputPinProps : CustomHLSLSettings->InputPins)
 			{
@@ -1037,25 +1037,29 @@ void UPCGCustomHLSLKernel::ParseShaderSource()
 				}
 				else if (CopyElementKeywords.Contains(TokenString))
 				{
-					const int32 FirstDelimiterIndex = TokenString.Find(TEXT("_"));
-					const int32 SecondDelimiterIndex = TokenString.Find(TEXT("_"), ESearchCase::IgnoreCase, ESearchDir::FromEnd);
-					
-					FString TargetPin = TokenString.Left(FirstDelimiterIndex);
-					FString SourcePin = TokenString.RightChop(SecondDelimiterIndex + 1);
+					// Format: TargetPin_CopyElementFrom_SourcePin. Find the known middle keyword to split.
+					const FString CopyKeyword = FString(TEXT("_")) + PCGCustomHLSLKernel::CopyElementFunctionKeyword + TEXT("_");
+					const int32 CopyKeywordIndex = TokenString.Find(CopyKeyword);
+					FString TargetPin = TokenString.Left(CopyKeywordIndex);
+					FString SourcePin = TokenString.RightChop(CopyKeywordIndex + CopyKeyword.Len());
 
 					InOutParsedSource.InitializedOutputPins.Add(TargetPin);
 					InOutParsedSource.CopyElementFunctions.Emplace(MoveTemp(SourcePin), MoveTemp(TargetPin));
 				}
 				else if (InitializeKeywords.Contains(TokenString))
 				{
-					const int32 DelimiterIndex = TokenString.Find(TEXT("_"));
+					// Use the last underscore as a delimiter in case the pin name contains underscores.
+					int32 DelimiterIndex = INDEX_NONE;
+					TokenString.FindLastChar(TEXT('_'), DelimiterIndex);
 					FString PinStr = TokenString.Left(DelimiterIndex);
 
 					InOutParsedSource.InitializedOutputPins.Add(MoveTemp(PinStr));
 				}
 				else if (PointPropertySetterKeywords.Contains(TokenString))
 				{
-					const int32 DelimiterIndex = TokenString.Find(TEXT("_"));
+					// Use the last underscore as a delimiter in case the pin name contains underscores.
+					int32 DelimiterIndex = INDEX_NONE;
+					TokenString.FindLastChar(TEXT('_'), DelimiterIndex);
 					const FString PinStr = TokenString.Left(DelimiterIndex);
 					const FString PropertyStr = TokenString.RightChop(DelimiterIndex + 1);
 					
@@ -1211,7 +1215,7 @@ void UPCGCustomHLSLKernel::CollectDataLabels(const FPCGCustomHLSLParsedSource& I
 		FPCGDataLabels& DataLabels = PinDataLabels.PinToDataLabels.FindOrAdd(InPinLabel);
 
 		// Matches against {PinName}_AnyFunction('{DataLabel}'...
-		const FString Pattern = FString::Format(TEXT("{0}_.*?[\\s]*?\\([\\s]*?'([a-zA-Z0-9_].*?)'"), { *InPinLabel.ToString() });
+		const FString Pattern = FString::Format(TEXT("{0}_.*?[\\s]*?\\([\\s]*?'([a-zA-Z0-9_].*?)'"), { InPinLabel.ToString() });
 
 		// First capture: Data label (supports a - z, A - Z, 0 - 9, and underscores).
 		constexpr int LabelCaptureGroup = 1;
@@ -1679,7 +1683,7 @@ FString UPCGCustomHLSLKernel::ProcessShaderSource(FPCGGPUCompilationContext& InO
 
 			// Matches against {PinName}_AnyFunction('{DataLabel}'...
 			// First capture group is the data label, so that we can find & replace it by index.
-			const FString Pattern = FString::Format(TEXT("{0}_.*?[\\s]*?\\([\\s]*?('{1}')"), { *PinLabel.ToString(), DataLabel });
+			const FString Pattern = FString::Format(TEXT("{0}_.*?[\\s]*?\\([\\s]*?('{1}')"), { PinLabel.ToString(), DataLabel });
 			FRegexMatcher ModuleMatcher(FRegexPattern(Pattern), OutShaderSource);
 
 			while (ModuleMatcher.FindNext())
@@ -1757,7 +1761,7 @@ bool UPCGCustomHLSLKernel::AreAllOutputPinsWritten()
 		}
 
 		const FPCGPinPropertiesGPU& PinProps = CustomHLSLSettings->OutputPins[I];
-		const FString& PinStr = PinProps.Label.ToString();
+		const FString PinStr = PinProps.Label.ToString();
 		bool bInitializedByAnySource = false;
 
 		for (const FPCGCustomHLSLParsedSource& ParsedSource : ParsedSources)
diff --git a/Engine/Plugins/PCG/Source/PCG/Private/Compute/PCGComputeCommon.cpp b/Engine/Plugins/PCG/Source/PCG/Private/Compute/PCGComputeCommon.cpp
--- a/Engine/Plugins/PCG/Source/PCG/Private/Compute/PCGComputeCommon.cpp
+++ b/Engine/Plugins/PCG/Source/PCG/Private/Compute/PCGComputeCommon.cpp
@@ -492,6 +492,40 @@ namespace PCGComputeHelpers
 			InContext->ExecutionSource->GetExecutionState().GetInspection().NotifyCPUToGPUUpload(InNode, InContext->GetStack());
 		});
 	}
+
+	FString SanitizePinLabelForHLSL(FName InLabel)
+	{
+		FString Result = InLabel.ToString();
+
+		// Replace invalid characters in-place.
+		for (TCHAR& Ch : Result)
+		{
+			if (!FChar::IsAlpha(Ch) && !FChar::IsDigit(Ch) && Ch != TEXT('_'))
+			{
+				Ch = TEXT('_');
+			}
+		}
+
+		// Trim leading non-alpha characters — our code generation and parsing requires identifiers that start with a letter.
+		int32 FirstAlpha = 0;
+		while (FirstAlpha < Result.Len() && !FChar::IsAlpha(Result[FirstAlpha]))
+		{
+			++FirstAlpha;
+		}
+
+		if (FirstAlpha > 0)
+		{
+			Result.RightChopInline(FirstAlpha);
+		}
+
+		return Result;
+	}
+
+	bool IsValidHLSLPinLabel(FName InLabel)
+	{
+		return SanitizePinLabelForHLSL(InLabel) == InLabel.ToString();
+	}
+
 #endif
 }
 
diff --git a/Engine/Plugins/PCG/Source/PCG/Private/Graph/PCGGraphCompilerGPU.cpp b/Engine/Plugins/PCG/Source/PCG/Private/Graph/PCGGraphCompilerGPU.cpp
--- a/Engine/Plugins/PCG/Source/PCG/Private/Graph/PCGGraphCompilerGPU.cpp
+++ b/Engine/Plugins/PCG/Source/PCG/Private/Graph/PCGGraphCompilerGPU.cpp
@@ -1729,10 +1729,10 @@ void FPCGGraphCompilerGPU::CompileComputeGraph(
 				if (Binding.Prefix != NAME_None)
 				{
 					TArray<FShaderFunctionDefinition>& DataInterfaceFunctions = bInEdgesAreInputs ? DataInterfaceToInputs.FindChecked(DataInterface) : DataInterfaceToOutputs.FindChecked(DataInterface);
-					
+
 					Edge.BindingFunctionNameOverride = FString::Format(
 						TEXT("{0}_{1}"),
-						{ Binding.Prefix.ToString(), DataInterfaceFunctions[Edge.DataInterfaceBindingIndex].Name }
+						{ PCGComputeHelpers::SanitizePinLabelForHLSL(Binding.Prefix), DataInterfaceFunctions[Edge.DataInterfaceBindingIndex].Name }
 					);
 				}
 
diff --git a/Engine/Plugins/PCG/Source/PCG/Private/Tests/Compute/PCGComputeCommonTest.cpp b/Engine/Plugins/PCG/Source/PCG/Private/Tests/Compute/PCGComputeCommonTest.cpp
new file mode 100644
--- /dev/null
+++ b/Engine/Plugins/PCG/Source/PCG/Private/Tests/Compute/PCGComputeCommonTest.cpp
@@ -0,0 +1,74 @@
+// Copyright Epic Games, Inc. All Rights Reserved.
+
+#if WITH_EDITOR
+
+#include "Tests/PCGTestsCommon.h"
+#include "Compute/PCGComputeCommon.h"
+
+IMPLEMENT_CUSTOM_SIMPLE_AUTOMATION_TEST(
+	FPCGSanitizePinLabelForHLSLTest,
+	FPCGTestBaseClass,
+	"Plugins.PCG.Compute.Common.SanitizePinLabel",
+	PCGTestsCommon::TestFlags)
+
+bool FPCGSanitizePinLabelForHLSLTest::RunTest(const FString& Parameters)
+{
+	// Valid identifiers pass through unchanged.
+	UTEST_EQUAL("Simple alpha", PCGComputeHelpers::SanitizePinLabelForHLSL(FName(TEXT("Out"))), FString(TEXT("Out")));
+	UTEST_EQUAL("AlphaNumeric", PCGComputeHelpers::SanitizePinLabelForHLSL(FName(TEXT("Pin1"))), FString(TEXT("Pin1")));
+	UTEST_EQUAL("Underscore", PCGComputeHelpers::SanitizePinLabelForHLSL(FName(TEXT("My_Pin"))), FString(TEXT("My_Pin")));
+	UTEST_EQUAL("Long valid name", PCGComputeHelpers::SanitizePinLabelForHLSL(FName(TEXT("Output_Data_3"))), FString(TEXT("Output_Data_3")));
+	UTEST_EQUAL("NAME_None", PCGComputeHelpers::SanitizePinLabelForHLSL(FName(NAME_None)), FString(TEXT("None")));
+
+	// Spaces replaced with underscores.
+	UTEST_EQUAL("Single space", PCGComputeHelpers::SanitizePinLabelForHLSL(FName(TEXT("End Points"))), FString(TEXT("End_Points")));
+	UTEST_EQUAL("Multiple spaces", PCGComputeHelpers::SanitizePinLabelForHLSL(FName(TEXT("My Pin Label"))), FString(TEXT("My_Pin_Label")));
+
+	// Special characters replaced with underscores.
+	UTEST_EQUAL("Dot", PCGComputeHelpers::SanitizePinLabelForHLSL(FName(TEXT("My.Pin"))), FString(TEXT("My_Pin")));
+	UTEST_EQUAL("Hyphen", PCGComputeHelpers::SanitizePinLabelForHLSL(FName(TEXT("My-Pin"))), FString(TEXT("My_Pin")));
+	UTEST_EQUAL("Mixed special chars", PCGComputeHelpers::SanitizePinLabelForHLSL(FName(TEXT("A!B@C"))), FString(TEXT("A_B_C")));
+
+	// Leading non-alpha characters trimmed (HLSL identifiers must start with a letter).
+	UTEST_EQUAL("Leading digit", PCGComputeHelpers::SanitizePinLabelForHLSL(FName(TEXT("3D_Data"))), FString(TEXT("D_Data")));
+	UTEST_EQUAL("Multiple leading digits", PCGComputeHelpers::SanitizePinLabelForHLSL(FName(TEXT("123abc"))), FString(TEXT("abc")));
+	UTEST_EQUAL("Leading underscore", PCGComputeHelpers::SanitizePinLabelForHLSL(FName(TEXT("_foo"))), FString(TEXT("foo")));
+	UTEST_EQUAL("Leading underscores and digits", PCGComputeHelpers::SanitizePinLabelForHLSL(FName(TEXT("__2x"))), FString(TEXT("x")));
+
+	// All characters invalid → empty string.
+	UTEST_EQUAL("All digits", PCGComputeHelpers::SanitizePinLabelForHLSL(FName(TEXT("123"))), FString(TEXT("")));
+	UTEST_EQUAL("All special", PCGComputeHelpers::SanitizePinLabelForHLSL(FName(TEXT("!@#"))), FString(TEXT("")));
+
+	// Combined edge cases.
+	UTEST_EQUAL("Leading special + spaces", PCGComputeHelpers::SanitizePinLabelForHLSL(FName(TEXT("#My Pin"))), FString(TEXT("My_Pin")));
+	UTEST_EQUAL("Single letter", PCGComputeHelpers::SanitizePinLabelForHLSL(FName(TEXT("A"))), FString(TEXT("A")));
+	UTEST_EQUAL("Trailing special", PCGComputeHelpers::SanitizePinLabelForHLSL(FName(TEXT("Pin!"))), FString(TEXT("Pin_")));
+
+	return true;
+}
+
+// ---- IsValidHLSLPinLabel ----
+
+IMPLEMENT_CUSTOM_SIMPLE_AUTOMATION_TEST(
+	FPCGIsValidHLSLPinLabelTest,
+	FPCGTestBaseClass,
+	"Plugins.PCG.Compute.Common.IsValidHLSLPinLabel",
+	PCGTestsCommon::TestFlags)
+
+bool FPCGIsValidHLSLPinLabelTest::RunTest(const FString& Parameters)
+{
+	// Valid labels.
+	UTEST_TRUE("Simple label is valid", PCGComputeHelpers::IsValidHLSLPinLabel(FName(TEXT("Out"))));
+	UTEST_TRUE("Underscore label is valid", PCGComputeHelpers::IsValidHLSLPinLabel(FName(TEXT("My_Pin"))));
+	UTEST_TRUE("Alphanumeric label is valid", PCGComputeHelpers::IsValidHLSLPinLabel(FName(TEXT("Pin1"))));
+
+	// Invalid labels.
+	UTEST_FALSE("Space is invalid", PCGComputeHelpers::IsValidHLSLPinLabel(FName(TEXT("End Points"))));
+	UTEST_FALSE("Leading digit is invalid", PCGComputeHelpers::IsValidHLSLPinLabel(FName(TEXT("3D_Data"))));
+	UTEST_FALSE("Dot is invalid", PCGComputeHelpers::IsValidHLSLPinLabel(FName(TEXT("My.Pin"))));
+	UTEST_FALSE("Leading underscore is invalid", PCGComputeHelpers::IsValidHLSLPinLabel(FName(TEXT("_foo"))));
+
+	return true;
+}
+
+#endif // WITH_EDITOR
diff --git a/Engine/Plugins/PCG/Source/PCG/Private/Tests/Compute/PCGCustomHLSLTest.cpp b/Engine/Plugins/PCG/Source/PCG/Private/Tests/Compute/PCGCustomHLSLTest.cpp
--- a/Engine/Plugins/PCG/Source/PCG/Private/Tests/Compute/PCGCustomHLSLTest.cpp
+++ b/Engine/Plugins/PCG/Source/PCG/Private/Tests/Compute/PCGCustomHLSLTest.cpp
@@ -8,6 +8,8 @@
 #include "PCGGraph.h"
 #include "PCGParamData.h"
 #include "Compute/Elements/PCGCustomHLSL.h"
+#include "Compute/PCGComputeCommon.h"
+#include "Compute/PCGPinPropertiesGPU.h"
 #include "Data/PCGBasePointData.h"
 #include "Elements/PCGCreateAttribute.h"
 #include "Elements/PCGCreatePoints.h"
@@ -340,4 +342,108 @@ bool FPCGCustomHLSLAttributeSetStringAttributeTest::RunTest(const FString& Param
 	});
 }
 
+/**
+ * GPU test: CustomHLSL PointProcessor with pin labels containing spaces.
+ *
+ *   CreatePoints (4 known positions) -> CustomHLSL (X += 100) -> Gather
+ *
+ * The CustomHLSL node uses "Input Points" as input pin label and "Output Points" as output pin label.
+ * The graph compiler sanitizes these to "Input_Points" and "Output_Points" for HLSL code generation.
+ * The HLSL source references the sanitized names. Verifies the end-to-end pipeline produces
+ * correct output, confirming that pin label sanitization works in the compiler.
+ */
+IMPLEMENT_CUSTOM_SIMPLE_AUTOMATION_TEST(
+	FPCGCustomHLSLSanitizedPinLabelsTest,
+	FPCGTestBaseClass,
+	"Plugins.PCG.Compute.CustomHLSL.SanitizedPinLabels",
+	PCGTestsCommon::TestFlags | EAutomationTestFlags::NonNullRHI)
+
+bool FPCGCustomHLSLSanitizedPinLabelsTest::RunTest(const FString& Parameters)
+{
+	// Verify our test labels actually need sanitization.
+	const FName OriginalInputLabel(TEXT("Input Points"));
+	const FName OriginalOutputLabel(TEXT("Output Points"));
+	TestFalse(TEXT("Input label requires sanitization"), PCGComputeHelpers::IsValidHLSLPinLabel(OriginalInputLabel));
+	TestFalse(TEXT("Output label requires sanitization"), PCGComputeHelpers::IsValidHLSLPinLabel(OriginalOutputLabel));
+
+	// Sanitize — mimics what PostEditChangeProperty does in the editor when a user types a label with spaces.
+	const FName SanitizedInputLabel(PCGComputeHelpers::SanitizePinLabelForHLSL(OriginalInputLabel));
+	const FName SanitizedOutputLabel(PCGComputeHelpers::SanitizePinLabelForHLSL(OriginalOutputLabel));
+	TestEqual(TEXT("Input label sanitized"), SanitizedInputLabel, FName(TEXT("Input_Points")));
+	TestEqual(TEXT("Output label sanitized"), SanitizedOutputLabel, FName(TEXT("Output_Points")));
+
+	FPCGGPUGraphTestRunner Runner;
+
+	// CreatePoints: 4 points at known positions.
+	UPCGCreatePointsSettings* CreatePointsSettings = nullptr;
+	UPCGNode* CreatePointsNode = Runner.Graph->AddNodeOfType<UPCGCreatePointsSettings>(CreatePointsSettings);
+	CreatePointsSettings->PointsToCreate.Empty();
+
+	const TArray<FVector> InputPositions = {
+		FVector(  0.f,   0.f,   0.f),
+		FVector(100.f, 200.f, 300.f),
+		FVector(-50.f,  75.f, 250.f),
+		FVector( 10.f, -30.f, -80.f),
+	};
+	for (const FVector& Pos : InputPositions)
+	{
+		FPCGPoint& Pt = CreatePointsSettings->PointsToCreate.AddDefaulted_GetRef();
+		Pt.Transform.SetLocation(Pos);
+	}
+	CreatePointsSettings->CoordinateSpace = EPCGCoordinateSpace::World;
+
+	// CustomHLSL PointProcessor with sanitized pin labels (spaces replaced with underscores).
+	UPCGCustomHLSLSettings* HLSLSettings = nullptr;
+	UPCGNode* HLSLNode = Runner.Graph->AddNodeOfType<UPCGCustomHLSLSettings>(HLSLSettings);
+	HLSLSettings->KernelType = EPCGKernelType::PointProcessor;
+
+	// Use the sanitized labels on the pins, same as the editor would after PostEditChangeProperty.
+	HLSLSettings->InputPins = { FPCGPinProperties(SanitizedInputLabel, FPCGDataTypeIdentifier{EPCGDataType::Point}) };
+	HLSLSettings->OutputPins = { FPCGPinPropertiesGPU(SanitizedOutputLabel, FPCGDataTypeIdentifier{EPCGDataType::Point}) };
+
+	HLSLSettings->SetSourceText(
+		TEXT("float3 Pos = Input_Points_GetPosition(Input_Points_DataIndex, ElementIndex);\n")
+		TEXT("Pos.x += 100.0f;\n")
+		TEXT("Output_Points_SetPosition(Output_Points_DataIndex, ElementIndex, Pos);\n")
+	);
+
+	UPCGGatherSettings* GatherSettings = nullptr;
+	UPCGNode* GatherNode = Runner.Graph->AddNodeOfType<UPCGGatherSettings>(GatherSettings);
+
+	Runner.Graph->AddEdge(CreatePointsNode, PCGPinConstants::DefaultOutputLabel, HLSLNode, SanitizedInputLabel);
+	Runner.Graph->AddEdge(HLSLNode, SanitizedOutputLabel, GatherNode, PCGPinConstants::DefaultInputLabel);
+
+	TArray<FPCGDataCollection> OutNodeData;
+	if (!Runner.Execute({ this, { GatherNode } }, OutNodeData))
+	{
+		return false;
+	}
+
+	return PCGGPUTestCommon::ValidateSingleTaggedData(this, OutNodeData[0], [&](const FPCGTaggedData& TaggedData)
+	{
+		const UPCGBasePointData* OutPoints = Cast<UPCGBasePointData>(TaggedData.Data);
+		if (!TestNotNull(TEXT("Readback data is point data"), OutPoints))
+		{
+			return false;
+		}
+
+		bool bTestPassed = TestEqual(TEXT("Output point count matches input"), OutPoints->GetNumPoints(), InputPositions.Num());
+
+		if (OutPoints->GetNumPoints() == InputPositions.Num())
+		{
+			const FConstPCGPointValueRanges Ranges(OutPoints);
+			for (int32 i = 0; i < InputPositions.Num(); ++i)
+			{
+				const FVector ExpectedPos = InputPositions[i] + FVector(100.0, 0.0, 0.0);
+				const FVector ActualPos = Ranges.TransformRange[i].GetLocation();
+				bTestPassed &= TestNearlyEqual(*FString::Printf(TEXT("Point[%d] X (expected %.1f, got %.1f)"), i, ExpectedPos.X, ActualPos.X), ActualPos.X, ExpectedPos.X, 0.1);
+				bTestPassed &= TestNearlyEqual(*FString::Printf(TEXT("Point[%d] Y (expected %.1f, got %.1f)"), i, ExpectedPos.Y, ActualPos.Y), ActualPos.Y, ExpectedPos.Y, 0.1);
+				bTestPassed &= TestNearlyEqual(*FString::Printf(TEXT("Point[%d] Z (expected %.1f, got %.1f)"), i, ExpectedPos.Z, ActualPos.Z), ActualPos.Z, ExpectedPos.Z, 0.1);
+			}
+		}
+
+		return bTestPassed;
+	});
+}
+
 #endif // WITH_EDITOR
diff --git a/Engine/Plugins/PCG/Source/PCG/Public/Compute/PCGComputeCommon.h b/Engine/Plugins/PCG/Source/PCG/Public/Compute/PCGComputeCommon.h
--- a/Engine/Plugins/PCG/Source/PCG/Public/Compute/PCGComputeCommon.h
+++ b/Engine/Plugins/PCG/Source/PCG/Public/Compute/PCGComputeCommon.h
@@ -167,6 +167,12 @@ namespace PCGComputeHelpers
 
 	PCG_API void NotifyGPUToCPUReadback(UPCGDataBinding* InBinding, const UPCGComputeKernel* InKernel, const UPCGSettings* InSettings);
 	PCG_API void NotifyCPUToGPUUpload(UPCGDataBinding* InBinding, const UPCGComputeKernel* InKernel, const UPCGSettings* InSettings);
+
+	/** Sanitizes a pin label for use as a valid HLSL identifier. Replaces invalid characters with underscores and trims leading non-alpha characters. */
+	PCG_API FString SanitizePinLabelForHLSL(FName InLabel);
+
+	/** Returns true if the pin label is already a valid HLSL identifier (i.e. sanitization would not change it). */
+	PCG_API bool IsValidHLSLPinLabel(FName InLabel);
 #endif
 }
 
