[PCG] Remove most of the pin widget dynamic lambdas which can simply be baked onto the pin and updated on refresh.

#was-backout 46474967

[RN] [PCG] Optimized CPU cost of PCG graph editor window
#Synced-CL 51942803

diff --git a/Engine/Plugins/PCG/Source/PCGEditor/Private/Nodes/SPCGEditorGraphNodeCompact.cpp b/Engine/Plugins/PCG/Source/PCGEditor/Private/Nodes/SPCGEditorGraphNodeCompact.cpp
--- a/Engine/Plugins/PCG/Source/PCGEditor/Private/Nodes/SPCGEditorGraphNodeCompact.cpp
+++ b/Engine/Plugins/PCG/Source/PCGEditor/Private/Nodes/SPCGEditorGraphNodeCompact.cpp
@@ -5,6 +5,7 @@
 #include "Nodes/PCGEditorGraphNode.h"
 #include "Nodes/PCGEditorGraphNodeBase.h"
 #include "PCGEditorStyle.h"
+#include "Pins/SPCGEditorGraphNodePin.h"
 
 #include "GraphEditorSettings.h"
 #include "IDocumentation.h"
@@ -256,23 +257,6 @@ void SPCGEditorGraphNodeCompact::UpdateGraphNode()
 
 	CreatePinWidgets();
 
-	// Hide pin labels
-	for (const TSharedRef<SGraphPin>& InputPin : this->InputPins)
-	{
-		if (InputPin->GetPinObj()->ParentPin == nullptr)
-		{
-			InputPin->SetShowLabel(false);
-		}
-	}
-
-	for (const TSharedRef<SGraphPin>& OutputPin : this->OutputPins)
-	{
-		if (OutputPin->GetPinObj()->ParentPin == nullptr)
-		{
-			OutputPin->SetShowLabel(false);
-		}
-	}
-
 	// @todo_pcg: Likely this is the same as non-compact. Could be factored out.
 	// Create comment bubble
 	TSharedPtr<SCommentBubble> CommentBubble;
@@ -303,6 +287,14 @@ void SPCGEditorGraphNodeCompact::UpdateGraphNode()
 	CreateOutputSideAddButton(RightNodeBox);
 }
 
+TSharedPtr<SGraphPin> SPCGEditorGraphNodeCompact::CreatePinWidget(UEdGraphPin* Pin) const
+{
+	// Compact nodes never show pin labels — bake ShowLabel=false at construction so no dynamic attribute is needed.
+	// Note: compact nodes do not use default value widgets (no UPCGSettingsDefaultValueProvider), so the parent
+	// CreatePinWidget logic for specialised pin types is intentionally bypassed here.
+	return SNew(SPCGEditorGraphNodePin, Pin).ShowLabel(false);
+}
+
 TSharedRef<SWidget> SPCGEditorGraphNodeCompact::CreateTitleWidget(TSharedPtr<SNodeTitle> InNodeTitle)
 {
 	check(PCGEditorGraphNode);
diff --git a/Engine/Plugins/PCG/Source/PCGEditor/Private/Nodes/SPCGEditorGraphNodeCompact.h b/Engine/Plugins/PCG/Source/PCGEditor/Private/Nodes/SPCGEditorGraphNodeCompact.h
--- a/Engine/Plugins/PCG/Source/PCGEditor/Private/Nodes/SPCGEditorGraphNodeCompact.h
+++ b/Engine/Plugins/PCG/Source/PCGEditor/Private/Nodes/SPCGEditorGraphNodeCompact.h
@@ -18,6 +18,7 @@ public:
 	virtual void UpdateGraphNode() override;
 	virtual TSharedRef<SWidget> CreateTitleWidget(TSharedPtr<SNodeTitle> InNodeTitle) override;
 	virtual const FSlateBrush* GetShadowBrush(bool bSelected) const override;
+	virtual TSharedPtr<SGraphPin> CreatePinWidget(UEdGraphPin* Pin) const override;
 	//~ End SGraphNode Interface
 
 private:
diff --git a/Engine/Plugins/PCG/Source/PCGEditor/Private/Pins/SPCGEditorGraphNodePin.cpp b/Engine/Plugins/PCG/Source/PCGEditor/Private/Pins/SPCGEditorGraphNodePin.cpp
--- a/Engine/Plugins/PCG/Source/PCGEditor/Private/Pins/SPCGEditorGraphNodePin.cpp
+++ b/Engine/Plugins/PCG/Source/PCGEditor/Private/Pins/SPCGEditorGraphNodePin.cpp
@@ -14,6 +14,7 @@
 #include "Algo/AnyOf.h"
 #include "Schema/PCGEditorGraphSchema.h"
 #include "Widgets/SToolTip.h"
+#include "Widgets/Images/SLayeredImage.h"
 #include "Widgets/Input/SButton.h"
 
 void SPCGEditorGraphNodePin::Construct(const FArguments& InArgs, UEdGraphPin* InPin)
@@ -25,8 +26,6 @@ void SPCGEditorGraphNodePin::Construct(const FArguments& InArgs, UEdGraphPin* In
 	bUsePinColorForText = InArgs._UsePinColorForText;
 	this->SetCursor(EMouseCursor::Default);
 
-	SetVisibility(MakeAttributeSP(this, &SPCGEditorGraphNodePin::GetPinVisiblity));
-
 	GraphPinObj = InPin;
 	check(GraphPinObj != NULL);
 
@@ -40,6 +39,8 @@ void SPCGEditorGraphNodePin::Construct(const FArguments& InArgs, UEdGraphPin* In
 		);
 
 	const bool bIsInput = (GetDirection() == EGPD_Input);
+	const FSlateColor PinColor = GetPinColor();
+	const FSlateColor PinTextColor = GetPinTextColor();
 
 	// A small marker to indicate the pin is required for the node to be executed.
 	TSharedPtr<SImage> RequiredPinIconWidget;
@@ -54,15 +55,14 @@ void SPCGEditorGraphNodePin::Construct(const FArguments& InArgs, UEdGraphPin* In
 		RequiredPinIconWidget =
 			SNew(SImage)
 			.Image(bDisplayPinMarker ? RequiredPinMarkerIcon : FAppStyle::GetNoBrush())
-			.ColorAndOpacity(this, &SPCGEditorGraphNodePin::GetPinColor);
+			.ColorAndOpacity(PinColor);
 	}
 
 	// Create the pin icon widget
-	TSharedRef<SWidget> PinWidgetRef = SPinTypeSelector::ConstructPinTypeImage(
-		MakeAttributeSP(this, &SPCGEditorGraphNodePin::GetPinIcon),
-		MakeAttributeSP(this, &SPCGEditorGraphNodePin::GetPinColor),
-		MakeAttributeSP(this, &SPCGEditorGraphNodePin::GetSecondaryPinIcon),
-		MakeAttributeSP(this, &SPCGEditorGraphNodePin::GetSecondaryPinColor));
+	TSharedRef<SWidget> PinWidgetRef =
+		SNew(SLayeredImage, GetSecondaryPinIcon(), GetSecondaryPinColor())
+		.Image(MakeAttributeSP(this, &SPCGEditorGraphNodePin::GetPinIcon))
+		.ColorAndOpacity(PinColor);
 	PinImage = PinWidgetRef;
 
 	PinWidgetRef->SetCursor(
@@ -84,7 +84,12 @@ void SPCGEditorGraphNodePin::Construct(const FArguments& InArgs, UEdGraphPin* In
 			.Image(this, &SPCGEditorGraphNodePin::GetPinStatusIcon)
 		];
 
-	TSharedRef<SWidget> LabelWidget = GetLabelWidget(GetLabelStyle(InArgs._PinLabelStyle));
+	TSharedRef<SWidget> LabelWidget =
+		SNew(STextBlock)
+		.Text(GetPinLabel())
+		.TextStyle(FAppStyle::Get(), InArgs._PinLabelStyle)
+		.Visibility(InArgs._ShowLabel ? EVisibility::Visible : EVisibility::Collapsed)
+		.ColorAndOpacity(PinTextColor);
 
 	// Create the widget used for the pin body (status indicator, label, and value)
 	LabelAndValue =
@@ -98,7 +103,7 @@ void SPCGEditorGraphNodePin::Construct(const FArguments& InArgs, UEdGraphPin* In
 	{
 		ExtraPinIconWidget = SNew(SImage)
 			.Image(FAppStyle::GetBrush(ExtraPinIcon))
-			.ColorAndOpacity(this, &SPCGEditorGraphNodePin::GetPinTextColor);
+			.ColorAndOpacity(PinTextColor);
 
 		if (!ExtraPinIconTooltip.IsEmpty())
 		{
@@ -165,7 +170,7 @@ void SPCGEditorGraphNodePin::Construct(const FArguments& InArgs, UEdGraphPin* In
 
 			if (!DoesWidgetHandleSettingEditingEnabled())
 			{
-				ValueBox->SetEnabled(TAttribute<bool>(this, &SPCGEditorGraphNodePin::IsEditingEnabled));
+				ValueBox->SetEnabled(IsEditingEnabled());
 			}
 		}
 
@@ -224,18 +229,18 @@ void SPCGEditorGraphNodePin::Construct(const FArguments& InArgs, UEdGraphPin* In
 	// Set up a hover for pins that is tinted the color of the pin.
 
 	SBorder::Construct(SBorder::FArguments()
-		.BorderImage(this, &SPCGEditorGraphNodePin::GetPinBorder)
-		.BorderBackgroundColor(this, &SPCGEditorGraphNodePin::GetHighlightColor)
+		.BorderImage(SPCGEditorGraphNodePin::GetPinBorder())
+		.BorderBackgroundColor(SPCGEditorGraphNodePin::GetHighlightColor())
 		.OnMouseButtonDown(this, &SPCGEditorGraphNodePin::OnPinNameMouseDown)
 		.Padding(0) // NOTE: This is different from base class implementation
 		[
 			SNew(SBorder)
 			.BorderImage(CachedImg_Pin_DiffOutline)
-			.BorderBackgroundColor(this, &SPCGEditorGraphNodePin::GetPinDiffColor)
+			.BorderBackgroundColor(GetPinDiffColor())
 			.Padding(0) // NOTE: This is different from base class implementation
 			[
 				SNew(SLevelOfDetailBranchNode)
-				.UseLowDetailSlot(this, &SPCGEditorGraphNodePin::UseLowDetailPinNames)
+				.UseLowDetailSlot(UseLowDetailPinNames())
 				.LowDetail()
 				[
 					//@TODO: Try creating a pin-colored line replacement that doesn't measure text / call delegates but still renders
@@ -248,8 +253,9 @@ void SPCGEditorGraphNodePin::Construct(const FArguments& InArgs, UEdGraphPin* In
 			]
 		]);
 
-	TSharedPtr<IToolTip> TooltipWidget = SNew(SToolTip)
-		.Text(this, &SPCGEditorGraphNodePin::GetTooltipText);
+	TSharedPtr<IToolTip> TooltipWidget =
+		SNew(SToolTip)
+		.Text(GetTooltipText());
 
 	SetToolTip(TooltipWidget);
 }
diff --git a/Engine/Plugins/PCG/Source/PCGEditor/Private/Pins/SPCGEditorGraphNodePin.h b/Engine/Plugins/PCG/Source/PCGEditor/Private/Pins/SPCGEditorGraphNodePin.h
--- a/Engine/Plugins/PCG/Source/PCGEditor/Private/Pins/SPCGEditorGraphNodePin.h
+++ b/Engine/Plugins/PCG/Source/PCGEditor/Private/Pins/SPCGEditorGraphNodePin.h
@@ -15,11 +15,13 @@ public:
 		: _PinLabelStyle(NAME_DefaultPinLabelStyle)
 		, _SideToSideMargin(5.0f)
 		, _UsePinColorForText(false)
+		, _ShowLabel(true)
 	{}
 
 	SLATE_ARGUMENT(FName, PinLabelStyle)
 	SLATE_ARGUMENT(float, SideToSideMargin)
 	SLATE_ARGUMENT(bool, UsePinColorForText)
+	SLATE_ARGUMENT(bool, ShowLabel)
 
 	SLATE_END_ARGS()
 
