[PCG] Fix TSet re-entrant mutation crash in PCGRuntimeGenScheduler

MoveTemp GeneratedSourcesToRemove into a local before iterating in CleanupDelayedRefreshSources() to prevent crash when CleanupSource() re-entrantly adds entries back via cleanup callbacks.

#rnx

#Synced-CL 51948399

diff --git a/Engine/Plugins/PCG/Source/PCG/Private/RuntimeGen/PCGRuntimeGenScheduler.cpp b/Engine/Plugins/PCG/Source/PCG/Private/RuntimeGen/PCGRuntimeGenScheduler.cpp
--- a/Engine/Plugins/PCG/Source/PCG/Private/RuntimeGen/PCGRuntimeGenScheduler.cpp
+++ b/Engine/Plugins/PCG/Source/PCG/Private/RuntimeGen/PCGRuntimeGenScheduler.cpp
@@ -1914,9 +1914,12 @@ void FPCGRuntimeGenScheduler::CleanupDelayedRefreshSources()
 {
 	TRACE_CPUPROFILER_EVENT_SCOPE(FPCGRuntimeGenScheduler::CleanupDelayedRefreshSources);
 
+	// Move to local to prevent crash if CleanupSource() re-entrantly adds to GeneratedSourcesToRemove.
+	TSet<FGridGenerationKey> SourcesToRemoveLocal = MoveTemp(GeneratedSourcesToRemove);
+
 	// Check that each refreshed local source is still intersecting its original source.
 	// If it is not, it would be leaked instead of refreshed, so we should force a full cleanup.
-	for (const FGridGenerationKey& GenerationKey : GeneratedSourcesToRemove)
+	for (const FGridGenerationKey& GenerationKey : SourcesToRemoveLocal)
 	{
 		if (!GenerationKey.IsValid())
 		{
@@ -1957,11 +1960,10 @@ void FPCGRuntimeGenScheduler::CleanupDelayedRefreshSources()
 		}
 	}
 
-	// Remove any generation keys that have been registered for deferred removal.
-	if (!GeneratedSourcesToRemove.IsEmpty())
+	// Remove cleaned sources from GeneratedSources.
+	if (!SourcesToRemoveLocal.IsEmpty())
 	{
-		GeneratedSources = GeneratedSources.Difference(GeneratedSourcesToRemove);
-		GeneratedSourcesToRemove.Reset();
+		GeneratedSources = GeneratedSources.Difference(SourcesToRemoveLocal);
 	}
 }
 
