[PCG] Added option to hit on ground plane for splines & volumes in the PCG editor tools.

#Synced-CL 52100463

diff --git a/Engine/Plugins/PCG/Source/PCGEditor/Private/EditorMode/Tools/Helpers/PCGEdModeSceneQueryHelpers.cpp b/Engine/Plugins/PCG/Source/PCGEditor/Private/EditorMode/Tools/Helpers/PCGEdModeSceneQueryHelpers.cpp
--- a/Engine/Plugins/PCG/Source/PCGEditor/Private/EditorMode/Tools/Helpers/PCGEdModeSceneQueryHelpers.cpp
+++ b/Engine/Plugins/PCG/Source/PCGEditor/Private/EditorMode/Tools/Helpers/PCGEdModeSceneQueryHelpers.cpp
@@ -106,6 +106,14 @@ namespace UE::PCG::EditorMode::Scene
 		{
 			return FViewHitResult(Result->ImpactPoint, Result->ImpactNormal, Result->Distance);
 		}
+		else if (Params.bHitGroundPlanes != nullptr && *Params.bHitGroundPlanes)
+		{
+			const FVector PlaneOrigin = FVector::ZeroVector;
+			const FVector PlaneNormal = (Params.CameraState && Params.CameraState->bIsOrthographic) ? -WorldRay.Direction : FVector::ZAxisVector;
+
+			// Implementation note - not testing against world because we've already done it & it would create an infinite recursion.
+			return ViewportRayAgainstPlane(WorldRay, Params, FPlane(PlaneOrigin, PlaneNormal), /*bTestAgainstWorld=*/false);
+		}
 		else
 		{
 			return {};
@@ -117,16 +125,27 @@ namespace UE::PCG::EditorMode::Scene
 		return ViewportRay(Ray.WorldRay, Params);
 	}
 
-	TOptional<FViewHitResult> ViewportRayAgainstPlane(const FInputDeviceRay& Ray, const FViewRayParams& Params, const FPlane& Plane)
+	TOptional<FViewHitResult> ViewportRayAgainstPlane(const FInputDeviceRay& Ray, const FViewRayParams& Params, const FPlane& Plane, bool bTestAgainstWorld)
+	{
+		return ViewportRayAgainstPlane(Ray.WorldRay, Params, Plane, bTestAgainstWorld);
+	}
+
+	TOptional<FViewHitResult> ViewportRayAgainstPlane(const FRay& Ray, const FViewRayParams& Params, const FPlane& Plane, bool bTestAgainstWorld)
 	{
 		bool bHitPlane = false;
 		FVector IntersectionPoint;
-		GizmoMath::RayPlaneIntersectionPoint(Plane.GetOrigin(), Plane.GetNormal(), Ray.WorldRay.Origin, Ray.WorldRay.Direction, bHitPlane, IntersectionPoint);
+		GizmoMath::RayPlaneIntersectionPoint(Plane.GetOrigin(), Plane.GetNormal(), Ray.Origin, Ray.Direction, bHitPlane, IntersectionPoint);
+
+		TOptional<FViewHitResult> HitResult;
+		if (bTestAgainstWorld)
+		{
+			HitResult = ViewportRay(Ray, Params);
+		}
 
-		if (const TOptional<FViewHitResult> HitResult = ViewportRay(Ray, Params))
+		if (HitResult)
 		{
 			// If the plane wasn't hit or the hit result was closer, return it.
-			if (!bHitPlane || FVector::DistSquared(Ray.WorldRay.Origin, HitResult->ImpactPosition) < FVector::DistSquared(Ray.WorldRay.Origin, IntersectionPoint))
+			if (!bHitPlane || FVector::DistSquared(Ray.Origin, HitResult->ImpactPosition) < FVector::DistSquared(Ray.Origin, IntersectionPoint))
 			{
 				return HitResult;
 			}
@@ -134,7 +153,7 @@ namespace UE::PCG::EditorMode::Scene
 
 		if (bHitPlane)
 		{
-			return FViewHitResult(IntersectionPoint, Plane.GetNormal(), FVector::Distance(Ray.WorldRay.Origin, IntersectionPoint));
+			return FViewHitResult(IntersectionPoint, Plane.GetNormal(), FVector::Distance(Ray.Origin, IntersectionPoint));
 		}
 
 		return {};
diff --git a/Engine/Plugins/PCG/Source/PCGEditor/Private/EditorMode/Tools/Helpers/PCGEdModeSceneQueryHelpers.h b/Engine/Plugins/PCG/Source/PCGEditor/Private/EditorMode/Tools/Helpers/PCGEdModeSceneQueryHelpers.h
--- a/Engine/Plugins/PCG/Source/PCGEditor/Private/EditorMode/Tools/Helpers/PCGEdModeSceneQueryHelpers.h
+++ b/Engine/Plugins/PCG/Source/PCGEditor/Private/EditorMode/Tools/Helpers/PCGEdModeSceneQueryHelpers.h
@@ -103,6 +103,8 @@ namespace UE::PCG::EditorMode::Scene
 		const TArray<const UPrimitiveComponent*>* ComponentsToIgnore = nullptr;
 		const TArray<const UPrimitiveComponent*>* AllowedInvisibleComponents = nullptr;
 		const FPCGRaycastFilterRuleCollection* FilterRuleCollection = nullptr;
+		const FViewCameraState* CameraState = nullptr;
+		bool* bHitGroundPlanes = nullptr;
 	};
 
 	/** The resulting data from a UI ray cast hit. */
@@ -138,7 +140,8 @@ namespace UE::PCG::EditorMode::Scene
 	TOptional<FViewHitResult> ViewportRay(const FInputDeviceRay& Ray, const FViewRayParams& Params);
 
 	/** Cast a ray from the view to a plane in space. If a hit is recorded, return a condensed view-related hit result. */
-	TOptional<FViewHitResult> ViewportRayAgainstPlane(const FInputDeviceRay& Ray, const FViewRayParams& Params, const FPlane& Plane);
+	TOptional<FViewHitResult> ViewportRayAgainstPlane(const FRay& Ray, const FViewRayParams& Params, const FPlane& Plane, bool bTestAgainstWorld = true);
+	TOptional<FViewHitResult> ViewportRayAgainstPlane(const FInputDeviceRay& Ray, const FViewRayParams& Params, const FPlane& Plane, bool bTestAgainstWorld = true);
 
 	/** Cast a ray from the view to a plane parallel to the user view. If a hit is recorded, return a condensed view-related hit result. */
 	TOptional<FViewHitResult> ViewportRayAgainstCameraPlane(const FInputDeviceRay& Ray, const FViewRayParams& Params, const FViewCameraState& CameraState);
diff --git a/Engine/Plugins/PCG/Source/PCGEditor/Private/EditorMode/Tools/Line/PCGDrawSplineTool.cpp b/Engine/Plugins/PCG/Source/PCGEditor/Private/EditorMode/Tools/Line/PCGDrawSplineTool.cpp
--- a/Engine/Plugins/PCG/Source/PCGEditor/Private/EditorMode/Tools/Line/PCGDrawSplineTool.cpp
+++ b/Engine/Plugins/PCG/Source/PCGEditor/Private/EditorMode/Tools/Line/PCGDrawSplineTool.cpp
@@ -318,7 +318,9 @@ void UPCGDrawSplineToolBase::Setup()
 	CachedRayParams =
 	{
 		.World = GetWorld(),
-		.FilterRuleCollection = &ToolSettings->RaycastRuleCollection
+		.FilterRuleCollection = &ToolSettings->RaycastRuleCollection,
+		.CameraState = &CachedCameraState,
+		.bHitGroundPlanes = &ToolSettings->bHitGroundPlanes
 	};
 }
 
@@ -720,6 +722,8 @@ void UPCGDrawSplineToolBase::OnTick(const float DeltaTime)
 			LOCTEXT("LostWorkingSpline", "The Draw Spline tool must close because the in-progress spline has been unexpectedly deleted."));
 	}
 
+	GetToolManager()->GetContextQueriesAPI()->GetCurrentViewState(CachedCameraState);
+
 	Super::OnTick(DeltaTime);
 }
 
diff --git a/Engine/Plugins/PCG/Source/PCGEditor/Private/EditorMode/Tools/Line/PCGDrawSplineTool.h b/Engine/Plugins/PCG/Source/PCGEditor/Private/EditorMode/Tools/Line/PCGDrawSplineTool.h
--- a/Engine/Plugins/PCG/Source/PCGEditor/Private/EditorMode/Tools/Line/PCGDrawSplineTool.h
+++ b/Engine/Plugins/PCG/Source/PCGEditor/Private/EditorMode/Tools/Line/PCGDrawSplineTool.h
@@ -54,6 +54,10 @@ public:
 	UPROPERTY(EditAnywhere, Category= "Raycast")
 	FPCGRaycastFilterRuleCollection RaycastRuleCollection;
 
+	/** Whether to place spline points on a plane through the origin aligned with the Z axis in perspective views, or facing the camera in othographic views */
+	UPROPERTY(EditAnywhere, Category= "Raycast")
+	bool bHitGroundPlanes = true;
+
 	/** The settings to manipulate raycast results, such as normals or offsets of a spline point. */
 	UPROPERTY(EditAnywhere, Category = "Raycast", meta = (ShowOnlyInnerProperties))
 	FPCGToolRaycastSettings Settings{};
@@ -136,6 +140,7 @@ private:
 	TObjectPtr<USingleClickOrDragInputBehavior> ClickOrDragBehavior = nullptr;
 	
 	UE::PCG::EditorMode::Scene::FViewRayParams CachedRayParams{};
+	FViewCameraState CachedCameraState;
 	bool bDrawTangentForPreviousPoint = false;
 	bool bFreeDrawPlacedPreviewPoint = false;
 	int32 FreeDrawStrokeStartIndex = 0;
diff --git a/Engine/Plugins/PCG/Source/PCGEditor/Private/EditorMode/Tools/Volume/PCGVolumeTool.cpp b/Engine/Plugins/PCG/Source/PCGEditor/Private/EditorMode/Tools/Volume/PCGVolumeTool.cpp
--- a/Engine/Plugins/PCG/Source/PCGEditor/Private/EditorMode/Tools/Volume/PCGVolumeTool.cpp
+++ b/Engine/Plugins/PCG/Source/PCGEditor/Private/EditorMode/Tools/Volume/PCGVolumeTool.cpp
@@ -155,7 +155,8 @@ void UPCGVolumeTool::Setup()
 	CachedRayParams =
 	{
 		.World = GetWorld(),
-		.FilterRuleCollection = &ToolSettings->RaycastRuleCollection
+		.FilterRuleCollection = &ToolSettings->RaycastRuleCollection,
+		.bHitGroundPlanes = &ToolSettings->bHitGroundPlanes
 	};
 }
 
diff --git a/Engine/Plugins/PCG/Source/PCGEditor/Private/EditorMode/Tools/Volume/PCGVolumeTool.h b/Engine/Plugins/PCG/Source/PCGEditor/Private/EditorMode/Tools/Volume/PCGVolumeTool.h
--- a/Engine/Plugins/PCG/Source/PCGEditor/Private/EditorMode/Tools/Volume/PCGVolumeTool.h
+++ b/Engine/Plugins/PCG/Source/PCGEditor/Private/EditorMode/Tools/Volume/PCGVolumeTool.h
@@ -38,6 +38,10 @@ public:
 	UPROPERTY(EditAnywhere, Category = Raycast)
 	FPCGRaycastFilterRuleCollection RaycastRuleCollection;
 
+	/** Whether to place the volume base on a plane through the origin aligned with the Z axis in perspective views. */
+	UPROPERTY(EditAnywhere, Category = "Raycast")
+	bool bHitGroundPlanes = true;
+
 protected:
 	/** This function is called by the builder using UPCGInteractiveToolSettings_Volume default class which prohibits to have override
 	 * implementation of this function by child classes. Mark it final for that reason, to be revisited if we change the tool classes hierarchy. */
