[PCG] Optimize PCG graph editor interactions by optimizing UPCGEditorGraphSchema methods that are called at very high frequency.

#rnx
#Synced-CL 52096374

diff --git a/Engine/Plugins/PCG/Source/PCGEditor/Private/Nodes/PCGEditorGraphNodeBase.cpp b/Engine/Plugins/PCG/Source/PCGEditor/Private/Nodes/PCGEditorGraphNodeBase.cpp
--- a/Engine/Plugins/PCG/Source/PCGEditor/Private/Nodes/PCGEditorGraphNodeBase.cpp
+++ b/Engine/Plugins/PCG/Source/PCGEditor/Private/Nodes/PCGEditorGraphNodeBase.cpp
@@ -185,6 +185,24 @@ void UPCGEditorGraphNodeBase::BeginDestroy()
 	Super::BeginDestroy();
 }
 
+FLinearColor UPCGEditorGraphNodeBase::GetCachedPinColor(const UPCGPin* PCGPin) const
+{
+	if (!PCGPin)
+	{
+		return GetDefault<UPCGEditorSettings>()->DefaultPinColor;
+	}
+
+	if (const FLinearColor* Cached = CachedPinColors.Find(PCGPin))
+	{
+		return *Cached;
+	}
+
+	const FLinearColor Color = FPCGModule::GetConstDataTypeRegistry().GetPinColor(PCGPin->GetCurrentTypesID());
+	CachedPinColors.Add(PCGPin, Color);
+
+	return Color;
+}
+
 void UPCGEditorGraphNodeBase::PostTransacted(const FTransactionObjectEvent& TransactionEvent)
 {
 	Super::PostTransacted(TransactionEvent);
@@ -624,6 +642,11 @@ void UPCGEditorGraphNodeBase::OnNodeChanged(UPCGNode* InNode, EPCGChangeType Cha
 {
 	if (InNode == PCGNode)
 	{
+		// Invalidate pin color cache before reconstruction so that Slate widget construction (which bakes
+		// colors at build time) always uses fresh values. Without this, the cache can contain pre-change
+		// colors that get baked into reconstructed widgets before the graph-change delegate fires.
+		ClearPinColorCache();
+
 		bool bRequiresReconstructNode = false;
 
 		if (!!(ChangeType & EPCGChangeType::Settings))
diff --git a/Engine/Plugins/PCG/Source/PCGEditor/Private/PCGEditorGraph.cpp b/Engine/Plugins/PCG/Source/PCGEditor/Private/PCGEditorGraph.cpp
--- a/Engine/Plugins/PCG/Source/PCGEditor/Private/PCGEditorGraph.cpp
+++ b/Engine/Plugins/PCG/Source/PCGEditor/Private/PCGEditorGraph.cpp
@@ -409,6 +409,20 @@ void UPCGEditorGraph::CreateLinks(UPCGEditorGraphNodeBase* GraphNode, bool bCrea
 
 void UPCGEditorGraph::OnGraphChanged(UPCGGraphInterface* InGraph, EPCGChangeType ChangeType)
 {
+	constexpr EPCGChangeType IgnoredChangeTypes = EPCGChangeType::Cosmetic | EPCGChangeType::Input | EPCGChangeType::ShaderSource | EPCGChangeType::GraphCustomization;
+	if (!!(ChangeType & ~IgnoredChangeTypes))
+	{
+		CachedEdgeUsedByExecution.Reset();
+
+		for (UEdGraphNode* EdGraphNode : Nodes)
+		{
+			if (UPCGEditorGraphNodeBase* PCGEditorNode = Cast<UPCGEditorGraphNodeBase>(EdGraphNode))
+			{
+				PCGEditorNode->ClearPinColorCache();
+			}
+		}
+	}
+
 	if (!!(ChangeType & EPCGChangeType::ExternalModification))
 	{
 		//Changes caused by external tools: Editor graph should be reconstructed.
@@ -514,3 +528,18 @@ FReply UPCGEditorGraph::OnPropertyBagDetailsDropOnGraph(const FPropertyBagProper
 
 	return FReply::Handled();
 }
+
+bool UPCGEditorGraph::IsEdgeUsedByExecution(const UPCGEdge* InEdge) const
+{
+	if (const bool* Cached = CachedEdgeUsedByExecution.Find(InEdge))
+	{
+		return *Cached;
+	}
+
+	// IsUsedByExecution dictated by downstream input pin.
+	const UPCGNode* Node = InEdge && InEdge->OutputPin ? InEdge->OutputPin->Node : nullptr;
+	const bool bUsed = !Node || Node->IsEdgeUsedByNodeExecution(InEdge);
+	CachedEdgeUsedByExecution.Add(InEdge, bUsed);
+
+	return bUsed;
+}
diff --git a/Engine/Plugins/PCG/Source/PCGEditor/Private/PCGEditorGraph.h b/Engine/Plugins/PCG/Source/PCGEditor/Private/PCGEditorGraph.h
--- a/Engine/Plugins/PCG/Source/PCGEditor/Private/PCGEditorGraph.h
+++ b/Engine/Plugins/PCG/Source/PCGEditor/Private/PCGEditorGraph.h
@@ -7,9 +7,12 @@
 
 #include "PCGGraph.h"
 
+#include "UObject/ObjectKey.h"
+
 #include "PCGEditorGraph.generated.h"
 
 class FPCGEditor;
+class UPCGEdge;
 class UPCGNode;
 class UPCGPin;
 class UPCGEditorGraphNodeBase;
@@ -58,6 +61,9 @@ public:
 	/** Returns the PCG editor graph node corresponding to the given PCG node. */
 	const UPCGEditorGraphNodeBase* GetEditorNodeFromPCGNode(const UPCGNode* InPCGNode) const;
 
+	/** Returns whether the edge is used by the downstream node's execution. Cached per graph change. */
+	bool IsEdgeUsedByExecution(const UPCGEdge* InEdge) const;
+
 	/** Returns the PCG Editor Graph Node that should be spawned from any given PCG Settings. */
 	static TSubclassOf<UPCGEditorGraphNodeBase> GetGraphNodeClassFromPCGSettings(const UPCGSettings* Settings);
 
@@ -80,4 +86,6 @@ private:
 	TWeakObjectPtr<UPCGGraph> WeakPCGGraph;
 
 	TWeakPtr<FPCGEditor> PCGEditor = nullptr;
+
+	mutable TMap<TObjectKey<UPCGEdge>, bool> CachedEdgeUsedByExecution;
 };
diff --git a/Engine/Plugins/PCG/Source/PCGEditor/Private/Schema/PCGEditorGraphSchema.cpp b/Engine/Plugins/PCG/Source/PCGEditor/Private/Schema/PCGEditorGraphSchema.cpp
--- a/Engine/Plugins/PCG/Source/PCGEditor/Private/Schema/PCGEditorGraphSchema.cpp
+++ b/Engine/Plugins/PCG/Source/PCGEditor/Private/Schema/PCGEditorGraphSchema.cpp
@@ -1324,7 +1324,7 @@ FLinearColor UPCGEditorGraphSchema::GetPinColor(const UEdGraphPin* InPin) const
 		PCGPin = PCGNode ? PCGNode->GetOutputPin(InPin->GetFName()) : nullptr;
 	}
 	
-	return PCGPin ? FPCGModule::GetConstDataTypeRegistry().GetPinColor(PCGPin->GetCurrentTypesID()) : GetDefault<UPCGEditorSettings>()->DefaultPinColor;
+	return EditorNode ? EditorNode->GetCachedPinColor(PCGPin) : GetDefault<UPCGEditorSettings>()->DefaultPinColor;
 }
 
 TSharedPtr<IAssetReferenceFilter> UPCGEditorGraphSchema::MakeAssetReferenceFilter(const UEdGraph* Graph)
@@ -1460,7 +1460,8 @@ void FPCGEditorConnectionDrawingPolicy::DetermineWiringStyle(UEdGraphPin* Output
 				{
 					return UpstreamEditorNode->GetPCGNode() == ConnectedPCGEdge->InputPin->Node && ConnectedPCGEdge->InputPin->Properties.Label == OutputPin->GetFName();
 				});
-				const bool bDownstreamNodeDoesNotUseData = PCGEdge && !PCGNode->IsEdgeUsedByNodeExecution(*PCGEdge);
+
+				const bool bDownstreamNodeDoesNotUseData = PCGEdge && Graph && !Graph->IsEdgeUsedByExecution(*PCGEdge);
 
 				// If edge found and is not used, gray it out
 				if (bDownstreamNodeForceDisabled || bDownstreamNodeDoesNotUseData)
diff --git a/Engine/Plugins/PCG/Source/PCGEditor/Public/Nodes/PCGEditorGraphNodeBase.h b/Engine/Plugins/PCG/Source/PCGEditor/Public/Nodes/PCGEditorGraphNodeBase.h
--- a/Engine/Plugins/PCG/Source/PCGEditor/Public/Nodes/PCGEditorGraphNodeBase.h
+++ b/Engine/Plugins/PCG/Source/PCGEditor/Public/Nodes/PCGEditorGraphNodeBase.h
@@ -8,6 +8,8 @@
 #include "PCGCommon.h"
 #include "Utils/PCGPreconfiguration.h"
 
+#include "UObject/ObjectKey.h"
+
 #include "PCGEditorGraphNodeBase.generated.h"
 
 enum class EPCGChangeType : uint32;
@@ -18,6 +20,7 @@ struct FPCGStack;
 
 class IPCGSettingsDefaultValueProvider;
 class UPCGGraph;
+class UPCGGraphInterface;
 class UPCGNode;
 class UPCGPin;
 class UPCGSettings;
@@ -153,6 +156,10 @@ public:
 	DECLARE_DELEGATE(FOnPCGEditorGraphNodeChanged);
 	FOnPCGEditorGraphNodeChanged OnNodeChangedDelegate;
 
+	/** Returns the cached wire color for the given pin. Recomputes on first call after any graph change. */
+	FLinearColor GetCachedPinColor(const UPCGPin* PCGPin) const;
+	void ClearPinColorCache() { CachedPinColors.Reset(); }
+
 protected:
 	static FEdGraphPinType GetPinType(const UPCGPin* InPin);
 
@@ -205,6 +212,9 @@ protected:
 	/** The higen grid currently being inspected if any, otherwise UninitializedGridSize(). */
 	uint32 InspectedGridSize = PCGHiGenGrid::UninitializedGridSize();
 
+private:
+	mutable TMap<TObjectKey<UPCGPin>, FLinearColor> CachedPinColors;
+
 	// Deprecated section
 public:
 	UE_DEPRECATED(5.8, "Use GenerationGridSize")
@@ -264,4 +274,4 @@ struct FPCGDeferNodeReconstructScope
 
 private:
 	UPCGEditorGraphNodeBase* Node = nullptr;
-};
\ No newline at end of file
+};
