[PCG] Add support for typed erased add compressed values in Generic Attributes

#rnx
#Synced-CL 52060281

diff --git a/Engine/Plugins/PCG/Source/PCG/Private/Metadata/PCGMetadataAttributeGeneric.cpp b/Engine/Plugins/PCG/Source/PCG/Private/Metadata/PCGMetadataAttributeGeneric.cpp
--- a/Engine/Plugins/PCG/Source/PCG/Private/Metadata/PCGMetadataAttributeGeneric.cpp
+++ b/Engine/Plugins/PCG/Source/PCG/Private/Metadata/PCGMetadataAttributeGeneric.cpp
@@ -622,13 +622,42 @@ bool FPCGMetadataAttributeGeneric::IsMap() const
 	return !CachedDesc.ContainerTypes.IsEmpty() && CachedDesc.ContainerTypes[0] == EPCGMetadataAttributeContainerTypes::Map;
 }
 
-TArray<PCGMetadataValueKey> FPCGMetadataAttributeGeneric::AddValues_Internal(const PCG::Private::FInValues& InValues, int32 Count, int32 ExistingIndex, bool bLockless)
+TArray<PCGMetadataValueKey> FPCGMetadataAttributeGeneric::AddValues_Internal(const PCG::Private::FInValues& InValues, int32 Count, int32 ExistingIndex, bool bLockless, TOptional<TConstArrayView<int32>> OptionalIndices)
 {
 	using namespace PCG::Private;
+	
+	check(!OptionalIndices || Count == OptionalIndices->Num());
+	
+	if (const FInValuesSubset* InValuesSubset = InValues.TryGet<FInValuesSubset>())
+	{
+		if (OptionalIndices)
+		{
+			// If we have multiple layers of subsets, we need to extract the indices.
+			// Note that the order we extract the layers is inner to outer, so Optional indices MUST be always smaller than the indices in the subset.
+			// For example, if we have FirstLayer = Subset(OriginalValues, [0,1,3]) and SecondLayer = Subset(FirstLayer, [0, 2]), the first time we call AddValues_Internal,
+			// we pass `SecondLayer`. We then call the function again with FirstLayer and OptionalIndices = [0, 2].
+			// The second time, we call the function again with OriginalValues and OptionalIndices = [0, 3] (First and Third indices ([0, 2]) of the FirstLayer subset indices [0, 1, 3])
+			check(InValuesSubset->Indices.Num() >= OptionalIndices->Num())
+
+			TArray<int32, TInlineAllocator<256>> IndicesSubset;
+			IndicesSubset.Reserve(OptionalIndices->Num());
+			for (int32 i : *OptionalIndices)
+			{
+				IndicesSubset.Add(InValuesSubset->Indices[i]);
+			}
+			
+			return AddValues_Internal(*static_cast<const FInValues*>(InValuesSubset->InValues), IndicesSubset.Num(), ExistingIndex, bLockless, MakeConstArrayView(IndicesSubset));
+		}
+
+		check(InValuesSubset->Indices.Num() <= Count);
+		
+		// Otherwise, forward the call with the underlying values and its indices.
+		return AddValues_Internal(*static_cast<const FInValues*>(InValuesSubset->InValues), InValuesSubset->Indices.Num(), ExistingIndex, bLockless, MakeConstArrayView(InValuesSubset->Indices));
+	}
 
 	TArray<PCGMetadataValueKey> ValueKeys;
 
-	int32 FirstValueIndex = INDEX_NONE;
+	int32 StartWriteIndex = INDEX_NONE;
 
 	TOptional<FScriptArrayHelper> Helper = MakeArrayHelper();
 	if (!Helper.IsSet() || Count <= 0)
@@ -647,81 +676,110 @@ TArray<PCGMetadataValueKey> FPCGMetadataAttributeGeneric::AddValues_Internal(con
 		if (ExistingIndex == INDEX_NONE)
 		{
 			check(!bLockless);
-			FirstValueIndex = AllocateValues_Unsafe(Count);
+			StartWriteIndex = AllocateValues_Unsafe(Count);
 		}
 		else if (ExistingIndex != ExistingIndexForDefaultValue)
 		{
 			// Expected to be already allocated.
 			check(Helper->Num() >= ExistingIndex + Count);
-			FirstValueIndex = ExistingIndex;
+			StartWriteIndex = ExistingIndex;
 		}
 		else
 		{
 			// Expected the default value to always be allocated.
-			check(DefaultValue.GetRawPtr() != nullptr)
-			FirstValueIndex = ExistingIndex;
+			check(DefaultValue.GetRawPtr() != nullptr);
+			StartWriteIndex = ExistingIndex;
 		}
 
-		check(FirstValueIndex != INDEX_NONE);
+		check(StartWriteIndex != INDEX_NONE);
 
-		auto GetRawPtr = [this, &Helper, FirstValueIndex](const int32 Index)
+		auto GetRawPtr = [this, &Helper, StartWriteIndex](const int32 Index)
 		{
-			return FirstValueIndex == ExistingIndexForDefaultValue ? DefaultValue.GetRawPtr() : Helper->GetRawPtr(Index);
+			return StartWriteIndex == ExistingIndexForDefaultValue ? DefaultValue.GetRawPtr() : Helper->GetRawPtr(Index);
+		};
+		
+		auto Loop = [&OptionalIndices, Count, StartWriteIndex](auto Callback)
+		{
+			if (!OptionalIndices)
+			{
+				for (int32 i = 0; i < Count; ++i)
+				{
+					Callback(/*ReadIndex=*/i, /*WriteIndex=*/StartWriteIndex + i);
+				}
+			}
+			else
+			{
+				for (int32 i = 0; i < Count; ++i)
+				{
+					Callback(/*ReadIndex=*/(*OptionalIndices)[i], /*WriteIndex=*/StartWriteIndex + i);
+				}
+			}
 		};
 
-		Visit([this, &GetRawPtr, FirstValueIndex, Count](auto&& Value)
+		Visit([this, &GetRawPtr, StartWriteIndex, Count, &OptionalIndices, &Loop](auto&& Value)
 			{
 				using T = std::decay_t<decltype(Value)>;
 				if constexpr(std::is_same_v<T, FInValuesByValue>)
 				{
-					const void* SrcPtr = Value.InValues;
-					void* DestPtr = GetRawPtr(FirstValueIndex);
-					UnderlyingProperty->Copy(DestPtr, SrcPtr, Count);
-				}
-				else if constexpr(std::is_same_v<T, FInValuesByValueSubset>)
-				{
-					for (int32 i = 0; i < Value.Indices.Num(); ++i)
+					if (!OptionalIndices)
+					{
+						const void* SrcPtr = Value.InValues;
+						void* DestPtr = GetRawPtr(StartWriteIndex);
+						UnderlyingProperty->Copy(DestPtr, SrcPtr, Count);
+					}
+					else
 					{
-						const void* SrcPtr = UnderlyingProperty->GetPtrInArray(Value.InValues,  Value.Indices[i]);
-						void* DestPtr = GetRawPtr(FirstValueIndex + i);
-						UnderlyingProperty->Copy(DestPtr, SrcPtr, 1);
+						for (int32 i = 0; i < OptionalIndices->Num(); ++i)
+						{
+							const void* SrcPtr = UnderlyingProperty->GetPtrInArray(Value.InValues, (*OptionalIndices)[i]);
+							void* DestPtr = GetRawPtr(StartWriteIndex + i);
+							UnderlyingProperty->Copy(DestPtr, SrcPtr, 1);
+						}
 					}
 				}
+				else if constexpr(std::is_same_v<T, FInValuesSubset>)
+				{
+					// Must be caught earlier
+					checkNoEntry();
+				}
 				else if constexpr (std::is_same_v<T, FInValuesAsSet>)
 				{
 					const FSetProperty* InnerProperty = CastFieldChecked<FSetProperty>(UnderlyingProperty->Property->Inner);
-					for (int32 i = 0; i < Count; ++i)
+
+					Loop([InnerProperty, &GetRawPtr, &Value](int32 ReadIndex, int32 WriteIndex)
 					{
-						FScriptSetHelper OtherHelper(InnerProperty, GetRawPtr(FirstValueIndex + i));
-						const int32 Num = Value.InValues[i].Num();
+						FScriptSetHelper OtherHelper(InnerProperty, GetRawPtr(WriteIndex));
+						const int32 Num = Value.InValues[ReadIndex].Num();
 
 						for (int j = 0; j < Num; ++j)
 						{
-							OtherHelper.AddElement(Value.InValues[i][j]);
-						}
+							OtherHelper.AddElement(Value.InValues[ReadIndex][j]);
 						}
+					});
 				}
 				else if constexpr (std::is_same_v<T, FInValuesAsMap>)
 				{
 					const FMapProperty* InnerProperty = CastFieldChecked<FMapProperty>(UnderlyingProperty->Property->Inner);
-					for (int32 i = 0; i < Count; ++i)
+
+					Loop([InnerProperty, &GetRawPtr, &Value](int32 ReadIndex, int32 WriteIndex)
 					{
-						FScriptMapHelper OtherHelper(InnerProperty, GetRawPtr(FirstValueIndex + i));
-						const int32 Num = Value.InValues[i].Num();
+						FScriptMapHelper OtherHelper(InnerProperty, GetRawPtr(WriteIndex));
+						const int32 Num = Value.InValues[ReadIndex].Num();
 
 						for (int j = 0; j < Num; ++j)
 						{
-							OtherHelper.AddPair(Value.InValues[i][j].Key, Value.InValues[i][j].Value);
-						}
+							OtherHelper.AddPair(Value.InValues[ReadIndex][j].Key, Value.InValues[ReadIndex][j].Value);
 						}
+					});
 				}
 				else if constexpr (std::is_same_v<T, FInValuesAsArray>)
 				{
 					const FArrayProperty* InnerProperty = CastFieldChecked<FArrayProperty>(UnderlyingProperty->Property->Inner);
-					for (int32 i = 0; i < Count; ++i)
+
+					Loop([InnerProperty, &GetRawPtr, &Value](int32 ReadIndex, int32 WriteIndex)
 					{
-						FScriptArrayHelper OtherHelper(InnerProperty, GetRawPtr(FirstValueIndex + i));
-						const int32 Num = Value.InValues[i].template Get<1>();
+						FScriptArrayHelper OtherHelper(InnerProperty, GetRawPtr(WriteIndex));
+						const int32 Num = Value.InValues[ReadIndex].template Get<1>();
 						if (InnerProperty->HasAnyPropertyFlags(CPF_IsPlainOldData))
 						{
 							OtherHelper.AddUninitializedValues(Num);
@@ -731,8 +789,8 @@ TArray<PCGMetadataValueKey> FPCGMetadataAttributeGeneric::AddValues_Internal(con
 							OtherHelper.AddValues(Num);
 						}
 
-						PCG::Private::CopyArray(InnerProperty, OtherHelper.GetRawPtr(), Value.InValues[i].template Get<0>(), Num);
-					}
+						PCG::Private::CopyArray(InnerProperty, OtherHelper.GetRawPtr(), Value.InValues[ReadIndex].template Get<0>(), Num);
+					});
 				}
 				else
 				{
@@ -744,31 +802,45 @@ TArray<PCGMetadataValueKey> FPCGMetadataAttributeGeneric::AddValues_Internal(con
 	ValueKeys.SetNum(Count);
 	for (int ValueIndex = 0; ValueIndex < Count; ++ValueIndex)
 	{
-		ValueKeys[ValueIndex] = ValueKeyOffset + FirstValueIndex + ValueIndex;
+		ValueKeys[ValueIndex] = ValueKeyOffset + StartWriteIndex + ValueIndex;
 	}
 
 	return ValueKeys;
 }
 
-TArray<PCGMetadataValueKey> FPCGMetadataAttributeGeneric::AddCompressedValues(PCG::Private::FInValuesByValue InUniqueValues, TArray<PCGMetadataValueKey>& FoundValueKeys)
+TArray<PCGMetadataValueKey> FPCGMetadataAttributeGeneric::AddCompressedValues(PCG::Private::FInValues InUniqueValues, int32 Count, TArray<PCGMetadataValueKey>& FoundValueKeys)
 {
 	using namespace PCG::Private;
 
-	const bool bHasDuplicateValues = (InUniqueValues.Count != FoundValueKeys.Num());
+	// For now only support FInValuesByValue or FInValuesSubset
+	if (!ensure(InUniqueValues.IsType<FInValuesByValue>() || InUniqueValues.IsType<FInValuesSubset>()))
+	{
+		return {};
+	}
+	
+	// FoundValueKeys should always be bigger than Count (since FoundValueKeys should represent the full range of values to add, while InUniqueValues/Count represent the unique values)
+	if (!ensure(Count <= FoundValueKeys.Num()))
+	{
+		return {};
+	}
+
+	const bool bHasDuplicateValues = (Count != FoundValueKeys.Num());
 
 	TArray<PCGMetadataValueKey> FoundUniqueValueKeys;
 	TArray<PCGMetadataValueKey>& FoundKeys = (bHasDuplicateValues ? FoundUniqueValueKeys : FoundValueKeys);
 
 	// Implementation note: when we don't have any duplicate values, the previously-set values in FoundValueKeys will be wiped out - this is intended
-	const bool bAtLeastOneValueNotFound = !FindValues(FoundKeys, InUniqueValues);
+	const bool bAtLeastOneValueNotFound = !FindValues(FoundKeys, InUniqueValues, Count);
 
 	if (bAtLeastOneValueNotFound)
 	{
-		FInValuesByValueSubset ValueSubset{};
-		ValueSubset.InValues = InUniqueValues.InValues;
-		ValueSubset.Indices.Reserve(InUniqueValues.Count);
+		// If there is at least one value not found, go through all the found keys (from the original unique values)
+		// and keep the indices of all the original values that need to be added.
+		FInValuesSubset ValueSubset{};
+		ValueSubset.InValues = &InUniqueValues;
+		ValueSubset.Indices.Reserve(Count);
 
-		for (int ValueIndex = 0; ValueIndex < InUniqueValues.Count; ++ValueIndex)
+		for (int ValueIndex = 0; ValueIndex < Count; ++ValueIndex)
 		{
 			if (FoundKeys[ValueIndex] == PCGNotFoundValueKey)
 			{
@@ -776,11 +848,11 @@ TArray<PCGMetadataValueKey> FPCGMetadataAttributeGeneric::AddCompressedValues(PC
 			}
 		}
 
-		const int32 Count = ValueSubset.Indices.Num();
-		TArray<PCGMetadataValueKey> NewValueKeys = AddValues_Internal(FInValues{TInPlaceType<FInValuesByValueSubset>{}, MoveTemp(ValueSubset)}, Count);
+		const int32 SubCount = ValueSubset.Indices.Num();
+		TArray<PCGMetadataValueKey> NewValueKeys = AddValues_Internal(FInValues{TInPlaceType<FInValuesSubset>{}, MoveTemp(ValueSubset)}, SubCount);
 
 		int32 NewlyAddedValueKeyIndex = 0;
-		for (int ValueIndex = 0; ValueIndex < InUniqueValues.Count; ++ValueIndex)
+		for (int ValueIndex = 0; ValueIndex < Count; ++ValueIndex)
 		{
 			if (FoundKeys[ValueIndex] == PCGNotFoundValueKey)
 			{
@@ -801,51 +873,200 @@ TArray<PCGMetadataValueKey> FPCGMetadataAttributeGeneric::AddCompressedValues(PC
 	return FoundValueKeys;
 }
 
-bool FPCGMetadataAttributeGeneric::FindValues(TArray<PCGMetadataValueKey>& OutValueKeys, const PCG::Private::FInValuesByValue& InValues) const
+TArray<PCGMetadataValueKey> FPCGMetadataAttributeGeneric::AddCompressedValues(PCG::Private::FInValuesByValue InValues, int32 Count)
 {
-	OutValueKeys.Init(PCGNotFoundValueKey, InValues.Count);
+	using namespace PCG::Private;
+	
+	if (!UnderlyingProperty || !UnderlyingProperty->IsValid())
+	{
+		return {};
+	}
+
+	check(UnderlyingProperty->bCompressData);
+
+	// Since we're getting raw values here, we might have duplicates
+	// so we should aim to remove duplicates here so we preserve our 'compress data' idea, otherwise it will break other foundational blocks (e.g. partition)
+	TArray<int32> UniqueValuesIndices;
+
+	auto Compare = [this, &InValues](int32 LHS, int32 RHS)
+	{	
+		const void* Value1 = UnderlyingProperty->GetPtrInArray(InValues.InValues, LHS);
+		const void* Value2 = UnderlyingProperty->GetPtrInArray(InValues.InValues, RHS);
+			
+		return Value1 != nullptr && Value2 != nullptr && UnderlyingProperty->Property->Inner->Identical(Value1, Value2);
+	};
+
+	// Initially, fill with mapping to unique values so we can remap them at the end if needed
+	TArray<PCGMetadataValueKey> FoundValueKeys;
+	FoundValueKeys.Reserve(Count);
+
+	const FProperty* InnerProperty = UnderlyingProperty->Property->Inner;
+	
+	// Hash-accelerated deduplication: O(n) average case.
+	// Algorithm:
+	//	- For each value in our incoming values
+	//		* Get the value pointer
+	//		* Compute the hash for this value
+	//		* Find the hash in our map of already computed hashes
+	//		* For each unique value with the same hash (collision), do the comparison
+	//		* If it is found, we have our value key
+	//		* If it is not found, we have a new unique value, so add it to our map.
+	//
+	// This reduces drastically the amount of comparisons to do at the cost of extra memory.
+	if (InnerProperty->HasAllPropertyFlags(CPF_HasGetValueTypeHash))
+	{
+		TMap<uint32, TArray<int32, TInlineAllocator<4>>> HashBuckets;
+
+		for (int32 i = 0; i < Count; ++i)
+		{
+			const void* ValuePtr = UnderlyingProperty->GetPtrInArray(InValues.InValues, i);
 
-	return FindValues(MakeArrayView(OutValueKeys), InValues);
+			int32 FoundValueKey = INDEX_NONE;
+
+			if (ValuePtr)
+			{
+				const uint32 Hash = InnerProperty->GetValueTypeHash(ValuePtr);
+
+				if (TArray<int32, TInlineAllocator<4>>* Bucket = HashBuckets.Find(Hash))
+				{
+					// Check within collision bucket using Identical for correctness
+					for (const int32 UniqueIdx : *Bucket)
+					{
+						if (Compare(i, UniqueValuesIndices[UniqueIdx]))
+						{
+							FoundValueKey = UniqueIdx;
+							break;
+						}
+					}
+				}
+
+				if (FoundValueKey == INDEX_NONE)
+				{
+					FoundValueKey = UniqueValuesIndices.Add(i);
+					HashBuckets.FindOrAdd(Hash).Add(FoundValueKey);
+				}
+			}
+			else
+			{
+				// Null pointer - treat as unique (consistent with Compare returning false for nulls)
+				FoundValueKey = UniqueValuesIndices.Add(i);
+			}
+
+			FoundValueKeys.Emplace(FoundValueKey);
+		}
+	}
+	else
+	{
+		// Fallback for types without hash support: O(n*m) where n = Count and m is the number of unique values (worst case: all unique values -> n^2)
+		for (int32 i = 0; i < Count; ++i)
+		{
+			int32 FoundValueKey = INDEX_NONE;
+			for (int32 j = 0; j < UniqueValuesIndices.Num(); ++j)
+			{
+				if (Compare(i, UniqueValuesIndices[j]))
+				{
+					FoundValueKey = j;
+					break;
+				}
+			}
+
+			if (FoundValueKey == INDEX_NONE)
+			{
+				FoundValueKey = UniqueValuesIndices.Add(i);
+			}
+
+			FoundValueKeys.Emplace(FoundValueKey);
+		}
+	}
+
+	FInValues TempValues{TInPlaceType<FInValuesByValue>{}, MoveTemp(InValues)};
+	const int32 UniqueCount = UniqueValuesIndices.Num();
+	return AddCompressedValues(FInValues{TInPlaceType<FInValuesSubset>{}, &TempValues, MoveTemp(UniqueValuesIndices)}, UniqueCount, FoundValueKeys);
+}
+
+bool FPCGMetadataAttributeGeneric::FindValues(TArray<PCGMetadataValueKey>& OutValueKeys, const PCG::Private::FInValues& InValues, int32 Count) const
+{
+	OutValueKeys.Init(PCGNotFoundValueKey, Count);
+
+	return FindValues(MakeArrayView(OutValueKeys), InValues, Count);
 }
 
-bool FPCGMetadataAttributeGeneric::FindValues(TArrayView<PCGMetadataValueKey> OutValueKeys, const PCG::Private::FInValuesByValue& InValues) const
+bool FPCGMetadataAttributeGeneric::FindValues(TArrayView<PCGMetadataValueKey> OutValueKeys, const PCG::Private::FInValues& InValues, int32 Count) const
 {
-	check(OutValueKeys.Num() == InValues.Count);
+	check(OutValueKeys.Num() == Count);
 	int ValueKeysSet = 0;
-	FindValuesInternal(InValues, OutValueKeys, ValueKeysSet, /*bIsRoot=*/true);
+	FindValuesInternal(InValues, Count, OutValueKeys, ValueKeysSet, /*bIsRoot=*/true);
 
-	return (ValueKeysSet == InValues.Count);
+	return (ValueKeysSet == Count);
 }
 
-void FPCGMetadataAttributeGeneric::FindValuesInternal(const PCG::Private::FInValuesByValue& InValues, TArrayView<PCGMetadataValueKey> ValueKeys, int& ValueKeysSet, bool bIsRoot) const
+void FPCGMetadataAttributeGeneric::FindValuesInternal(const PCG::Private::FInValues& InValues, int32 Count, TArrayView<PCGMetadataValueKey> ValueKeys, int& ValueKeysSet, bool bIsRoot) const
 {
 	using namespace PCG::Private;
 
-	check(InValues.Count == ValueKeys.Num());
+	check(Count == ValueKeys.Num());
 
 	if (!ensure(UnderlyingProperty))
 	{
 		return;
 	}
+	
+	// For now only support FInValuesByValue and FInValuesSubset
+	const TArray<int32>* OptionalIndices = nullptr;
+	const void* UnderlyingValues = nullptr;
 
-	auto Compare = [this, &InValues](int32 InValueIndex, PCGMetadataValueKey ValueKey)
+	Visit([this, &OptionalIndices, &UnderlyingValues](auto&& InValue)
 	{
-		const void* Value1 = GetReadAddressForValueKey_Unsafe(ValueKey);
-		const void* Value2 = UnderlyingProperty->GetPtrInArray(InValues.InValues, InValueIndex);
-
-		if (Value1 == nullptr || Value2 == nullptr)
+		using T = std::decay_t<decltype(InValue)>;
+		if constexpr (std::is_same_v<T, FInValuesAsArray> || std::is_same_v<T, FInValuesAsSet> || std::is_same_v<T, FInValuesAsMap>)
 		{
-			return false;
+			// Unsupported
+			return;
+		}
+		else if constexpr (std::is_same_v<T, FInValuesByValue>)
+		{
+			UnderlyingValues = InValue.InValues;
+		}
+		else if constexpr (std::is_same_v<T, FInValuesSubset>)
+		{
+			const FInValues* UnderlyingVariant = static_cast<const FInValues*>(InValue.InValues);
+			
+			if (const FInValuesByValue* UnderlyingValue = UnderlyingVariant->TryGet<FInValuesByValue>())
+			{
+				UnderlyingValues = UnderlyingValue->InValues;
+				OptionalIndices = &InValue.Indices;
+			}
+			else
+			{
+				// Unsupported
+				return;
+			}
 		}
+		else
+		{
+			static_assert(!std::is_same_v<T, T>, "Missing variant case for FindValuesInternal");
+		}
+	}, InValues);
 
-		return UnderlyingProperty->Property->Inner->Identical(Value1, Value2);
+	if (!UnderlyingValues)
+	{
+		// Unsupported
+		return;
+	}
+
+	auto Compare = [this, UnderlyingValues](int32 InValueIndex, PCGMetadataValueKey ValueKey)
+	{
+		const void* Value1 = GetReadAddressForValueKey_Unsafe(ValueKey);
+		const void* Value2 = UnderlyingProperty->GetPtrInArray(UnderlyingValues, InValueIndex);
+
+		return Value1 != nullptr && Value2 != nullptr && UnderlyingProperty->Property->Inner->Identical(Value1, Value2);
 	};
 
 	if (bIsRoot)
 	{
-		for (int ValueIndex = 0; ValueIndex < InValues.Count; ++ValueIndex)
+		for (int ValueIndex = 0; ValueIndex < Count; ++ValueIndex)
 		{
-			if (Compare(ValueIndex, PCGDefaultValueKey))
+			if (Compare(OptionalIndices ? (*OptionalIndices)[ValueIndex] : ValueIndex, PCGDefaultValueKey))
 			{
 				ValueKeys[ValueIndex] = PCGDefaultValueKey;
 				++ValueKeysSet;
@@ -853,19 +1074,19 @@ void FPCGMetadataAttributeGeneric::FindValuesInternal(const PCG::Private::FInVal
 		}
 	}
 
-	if (ValueKeysSet != InValues.Count)
+	if (ValueKeysSet != Count)
 	{
 		const FPCGMetadataAttributeBase* ThisParent = GetParent();
 		if (ThisParent && ThisParent->IsGeneric())
 		{
-			static_cast<const FPCGMetadataAttributeGeneric*>(ThisParent)->FindValuesInternal(InValues, ValueKeys, ValueKeysSet, /*bRoot=*/false);
+			static_cast<const FPCGMetadataAttributeGeneric*>(ThisParent)->FindValuesInternal(InValues, Count, ValueKeys, ValueKeysSet, /*bRoot=*/false);
 		}
 	}
 
-	if (ValueKeysSet != InValues.Count)
+	if (ValueKeysSet != Count)
 	{
 		PCG::TSharedScopeLock ReadLock(ValueLock, /*bShouldLock=*/bIsRoot);
-		for (int ValueIndex = 0; ValueIndex < InValues.Count; ++ValueIndex)
+		for (int ValueIndex = 0; ValueIndex < Count; ++ValueIndex)
 		{
 			if (ValueKeys[ValueIndex] != PCGNotFoundValueKey)
 			{
@@ -875,7 +1096,7 @@ void FPCGMetadataAttributeGeneric::FindValuesInternal(const PCG::Private::FInVal
 			for (PCGMetadataValueKey ValueKey = 0; ValueKey < Values.Num(); ++ValueKey)
 			{
 				const PCGMetadataValueKey OffsetKey = ValueKey + ValueKeyOffset;
-				if (Compare(ValueIndex, OffsetKey))
+				if (Compare(OptionalIndices ? (*OptionalIndices)[ValueIndex] : ValueIndex, OffsetKey))
 				{
 					ValueKeys[ValueIndex] = OffsetKey;
 					++ValueKeysSet;
@@ -1252,7 +1473,7 @@ void FPCGMetadataAttributeGeneric::SetValue(PCGMetadataEntryKey ItemKey, const F
 
 	if (UnderlyingProperty->bCompressData)
 	{
-		const bool bValueKeyFound = FindValues(MakeArrayView(&ValueKeyToAdd, 1), PCG::Private::FInValuesByValue{SrcPtr, 1});
+		const bool bValueKeyFound = FindValues(MakeArrayView(&ValueKeyToAdd, 1), FInValues{TInPlaceType<FInValuesByValue>{}, SrcPtr, 1}, 1);
 		check(!bValueKeyFound || ValueKeyToAdd != PCGNotFoundValueKey);
 	}
 
@@ -1690,13 +1911,13 @@ void FPCGMetadataAttributeGeneric::SetValueFromProperty(PCGMetadataEntryKey Item
 
 		if (UnderlyingProperty->bCompressData && ItemKey != PCGInvalidEntryKey)
 		{
-			const bool bValueKeyFound = FindValues(MakeArrayView(&ValueKeyToAdd, 1), FInValuesByValue{SrcPtr, 1});
+			const bool bValueKeyFound = FindValues(MakeArrayView(&ValueKeyToAdd, 1), FInValues{TInPlaceType<FInValuesByValue>{}, SrcPtr, 1}, 1);
 			check(!bValueKeyFound || ValueKeyToAdd != PCGNotFoundValueKey);
 		}
 
 		if (ValueKeyToAdd == PCGNotFoundValueKey)
 		{
-			SetValuesFromValueKeys({ItemKey}, AddValues_Internal(FInValues{TInPlaceType<FInValuesByValue>{}, SrcPtr, 1}, 1, ExistingIndex), /*bResetDefaultValueOnDefaultValueKey=*/ItemKey==PCGInvalidEntryKey);
+			SetValuesFromValueKeys({ItemKey}, AddValues_Internal(PCG::Private::FInValues{TInPlaceType<PCG::Private::FInValuesByValue>{}, SrcPtr, 1}, 1, ExistingIndex), /*bResetDefaultValueOnDefaultValueKey=*/ItemKey==PCGInvalidEntryKey);
 		}
 		else
 		{
diff --git a/Engine/Plugins/PCG/Source/PCG/Public/Metadata/PCGMetadataAttributeGeneric.h b/Engine/Plugins/PCG/Source/PCG/Public/Metadata/PCGMetadataAttributeGeneric.h
--- a/Engine/Plugins/PCG/Source/PCG/Public/Metadata/PCGMetadataAttributeGeneric.h
+++ b/Engine/Plugins/PCG/Source/PCG/Public/Metadata/PCGMetadataAttributeGeneric.h
@@ -378,7 +378,10 @@ public:
 protected:
 	/** Attributes can't be created individually, they need to be created through a Metadata Domain. */
 	friend FPCGMetadataDomain;
-	
+
+	template <typename T>
+	friend class FPCGAttributeGenericAccessor;
+
 	// New attribute with no parent, name is contained in the InAttributeDesc.
 	PCG_API FPCGMetadataAttributeGeneric(const FPCGMetadataAttributeDesc& InAttributeDesc, FPCGMetadataDomain* InMetadata, bool bInAllowsInterpolation);
 
@@ -461,7 +464,7 @@ protected:
 					FoundValueKeys.Emplace(UniqueValues.AddUnique(InValues[ValueIndex]));
 				}
 
-				return AddCompressedValues(PCG::Private::FInValuesByValue{UniqueValues.GetData(), UniqueValues.Num()}, FoundValueKeys);
+				return AddCompressedValues(PCG::Private::FInValues{TInPlaceType<PCG::Private::FInValuesByValue>{}, UniqueValues.GetData(), UniqueValues.Num()}, UniqueValues.Num(), FoundValueKeys);
 			}
 		}
 
@@ -554,20 +557,23 @@ protected:
 	}
 	
 	// Internal implementation that will do the dispatch using the typed erased variant.
-	PCG_API TArray<PCGMetadataValueKey> AddValues_Internal(const PCG::Private::FInValues& InValues, int32 Count, int32 ExistingIndex = INDEX_NONE, bool bLockless = false);
+	PCG_API TArray<PCGMetadataValueKey> AddValues_Internal(const PCG::Private::FInValues& InValues, int32 Count, int32 ExistingIndex = INDEX_NONE, bool bLockless = false, TOptional<TConstArrayView<int32>> OptionalIndices = {});
 
 	// Internal implementation for compressed data
-	PCG_API TArray<PCGMetadataValueKey> AddCompressedValues(PCG::Private::FInValuesByValue InUniqueValues, TArray<PCGMetadataValueKey>& FoundValueKeys);
+	PCG_API TArray<PCGMetadataValueKey> AddCompressedValues(PCG::Private::FInValues InUniqueValues, int32 Count, TArray<PCGMetadataValueKey>& FoundValueKeys);
+
+	// Internal implementation for typed-erased compressed data. Will create a temporary storage to find all the unique values then will call AddCompressedValues with those.
+	PCG_API TArray<PCGMetadataValueKey> AddCompressedValues(PCG::Private::FInValuesByValue InValues, int32 Count);
 
 	/** 
 	 * Find all the value keys associated to InValues for compressed data. Return true if all the values were found.
 	 */
-	PCG_API bool FindValues(TArray<PCGMetadataValueKey>& OutValueKeys, const PCG::Private::FInValuesByValue& InValues) const;
+	PCG_API bool FindValues(TArray<PCGMetadataValueKey>& OutValueKeys, const PCG::Private::FInValues& InValues, int32 Count) const;
 
 	/** 
 	 * Find all the value keys associated to InValues for compressed data. Return true if all the values were found. Preallocated version.
 	 */
-	PCG_API bool FindValues(TArrayView<PCGMetadataValueKey> OutValueKeys, const PCG::Private::FInValuesByValue& InValues) const;
+	PCG_API bool FindValues(TArrayView<PCGMetadataValueKey> OutValueKeys, const PCG::Private::FInValues& InValues, int32 Count) const;
 
 	template <typename T>
 	void InitForDeprecation(TConstArrayView<T> InValues, const T& InDefaultValue)
@@ -595,7 +601,7 @@ private:
 	void Init();
 	
 	PCG_API void GetValues_Internal(const TArrayView<const PCGMetadataValueKey> ValueKeys, PCG::Private::FOutValues OutValues, TBitArray<>& UnretrievedValues) const;
-	void FindValuesInternal(const PCG::Private::FInValuesByValue& InValues, TArrayView<PCGMetadataValueKey> ValueKeys, int& ValueKeysSet, bool bIsRoot) const;
+	void FindValuesInternal(const PCG::Private::FInValues& InValues, int32 Count, TArrayView<PCGMetadataValueKey> ValueKeys, int& ValueKeysSet, bool bIsRoot) const;
 
 public:
 	PCG_API virtual void Flatten() override;
diff --git a/Engine/Plugins/PCG/Source/PCG/Public/Metadata/PCGMetadataCommon.h b/Engine/Plugins/PCG/Source/PCG/Public/Metadata/PCGMetadataCommon.h
--- a/Engine/Plugins/PCG/Source/PCG/Public/Metadata/PCGMetadataCommon.h
+++ b/Engine/Plugins/PCG/Source/PCG/Public/Metadata/PCGMetadataCommon.h
@@ -341,13 +341,13 @@ namespace PCG::Private
 	//////////////////////
 	struct FInValuesByValue
 	{
-		const void* InValues;
-		int32 Count;
+		const void* InValues = nullptr;
+		int32 Count = 0;
 	};
 
-	struct FInValuesByValueSubset
+	struct FInValuesSubset
 	{
-		const void* InValues;
+		const void* InValues = nullptr; // It is expected to be a FInValues underneath. Anything else will crash.
 		TArray<int32> Indices;
 	};
 
@@ -366,5 +366,5 @@ namespace PCG::Private
 		TArrayView<TArray<TPair<const void*, const void*>>> InValues;
 	};
 
-	using FInValues = TVariant<FInValuesByValue, FInValuesAsArray, FInValuesAsSet, FInValuesAsMap, FInValuesByValueSubset>;
+	using FInValues = TVariant<FInValuesByValue, FInValuesAsArray, FInValuesAsSet, FInValuesAsMap, FInValuesSubset>;
 }
\ No newline at end of file
