[PCG] Include zero partitions error along with missing attribute errors in AttributeGenericPartition

#Synced-CL 52134510

diff --git a/Engine/Plugins/PCG/Source/PCG/Private/Metadata/PCGMetadataPartitionCommon.cpp b/Engine/Plugins/PCG/Source/PCG/Private/Metadata/PCGMetadataPartitionCommon.cpp
--- a/Engine/Plugins/PCG/Source/PCG/Private/Metadata/PCGMetadataPartitionCommon.cpp
+++ b/Engine/Plugins/PCG/Source/PCG/Private/Metadata/PCGMetadataPartitionCommon.cpp
@@ -482,14 +482,14 @@ namespace PCGMetadataPartitionCommon
 	 *  3  b  c  c                         Final Partition (A&B&C)->[0],[1],[2],[3,4]
 	 *  4  b  c  c
 	 */
-	TArray<TArray<int32>> AttributeGenericPartition(const UPCGData* InData, const TArrayView<const FPCGAttributePropertySelector>& InSelectorArrayView, FPCGContext* InOptionalContext, bool bSilenceMissingAttributeErrors)
+	TArray<TArray<int32>> AttributeGenericPartition(const UPCGData* InData, const TArrayView<const FPCGAttributePropertySelector>& InSelectorArrayView, FPCGContext* InOptionalContext, bool bSilenceZeroPartitionErrors)
 	{
 		TRACE_CPUPROFILER_EVENT_SCOPE(PCGMetadataPartitionCommon::AttributeGenericPartition::MultiSelector);
 
 		// Small optimization to partition on a single attribute
 		if (InSelectorArrayView.Num() == 1)
 		{
-			return AttributeGenericPartition<TArray<int32>>(InData, InSelectorArrayView[0], InOptionalContext, bSilenceMissingAttributeErrors);
+			return AttributeGenericPartition<TArray<int32>>(InData, InSelectorArrayView[0], InOptionalContext, bSilenceZeroPartitionErrors);
 		}
 
 		if (!InData || InSelectorArrayView.IsEmpty() || !InData->ConstMetadata())
@@ -517,11 +517,15 @@ namespace PCGMetadataPartitionCommon
 			// Partition per attribute
 			for (int32 I = 0; I < NumAttributes; ++I)
 			{
-				PerAttributePartitions[I] = AttributeGenericPartition<TArray<int32>>(InData, InSelectorArrayView[I], InOptionalContext, bSilenceMissingAttributeErrors);
+				PerAttributePartitions[I] = AttributeGenericPartition<TArray<int32>>(InData, InSelectorArrayView[I], InOptionalContext, bSilenceZeroPartitionErrors);
 				// A partition could be empty if the selector is bad, which would result in a 0 value multiplier later.
 				if (PerAttributePartitions[I].IsEmpty())
 				{
-					PCGLog::LogErrorOnGraph(FText::Format(LOCTEXT("AttributeResultsInZeroPartitions", "Attribute resulted in zero partitions. Check selector {0}"), InSelectorArrayView[I].GetDisplayText()), InOptionalContext);
+					if (!bSilenceZeroPartitionErrors)
+					{
+						PCGLog::LogErrorOnGraph(FText::Format(LOCTEXT("AttributeResultsInZeroPartitions", "Attribute resulted in zero partitions. Check selector {0}"), InSelectorArrayView[I].GetDisplayText()), InOptionalContext);
+					}
+
 					return {};
 				}
 			}
diff --git a/Engine/Plugins/PCG/Source/PCG/Public/Metadata/PCGMetadataPartitionCommon.h b/Engine/Plugins/PCG/Source/PCG/Public/Metadata/PCGMetadataPartitionCommon.h
--- a/Engine/Plugins/PCG/Source/PCG/Public/Metadata/PCGMetadataPartitionCommon.h
+++ b/Engine/Plugins/PCG/Source/PCG/Public/Metadata/PCGMetadataPartitionCommon.h
@@ -50,10 +50,10 @@ namespace PCGMetadataPartitionCommon
 	* @param InData - Data to partition, need to support attributes (spatial data or attribute set).
 	* @param InSelectorArrayView - ArrayView of selectors on the attribute to partition.
 	* @param InOptionalContext - Optional context for logging.
-	* @param bSilenceMissingAttributeErrors - Do not log errors to the context or log.
+	* @param bSilenceZeroPartitionErrors - Do not log errors to the context or log.
 	* @returns Array of bucket indices.
 	*/
-	PCG_API TArray<TArray<int32>> AttributeGenericPartition(const UPCGData* InData, const TArrayView<const FPCGAttributePropertySelector>& InSelectorArrayView, FPCGContext* InOptionalContext = nullptr, bool bSilenceMissingAttributeErrors = false);
+	PCG_API TArray<TArray<int32>> AttributeGenericPartition(const UPCGData* InData, const TArrayView<const FPCGAttributePropertySelector>& InSelectorArrayView, FPCGContext* InOptionalContext = nullptr, bool bSilenceZeroPartitionErrors = false);
 
 	/**
 	* Remove duplicates for the incoming data on the given attribute.
