[PCG] Inline Constant Widgets for SoftObjectPath and SoftClassPath

#Synced-CL 52134510

diff --git a/Engine/Plugins/PCG/Source/PCGEditor/PCGEditor.Build.cs b/Engine/Plugins/PCG/Source/PCGEditor/PCGEditor.Build.cs
--- a/Engine/Plugins/PCG/Source/PCGEditor/PCGEditor.Build.cs
+++ b/Engine/Plugins/PCG/Source/PCGEditor/PCGEditor.Build.cs
@@ -40,6 +40,7 @@ namespace UnrealBuildTool.Rules
 					"AssetRegistry",
 					"AssetTools",
 					"BlueprintGraph",
+					"ClassViewer",
 					"ComponentVisualizers",
 					"ContentBrowser",
 					"ContentBrowserData",
diff --git a/Engine/Plugins/PCG/Source/PCGEditor/Private/Pins/SPCGEditorGraphNodePin.cpp b/Engine/Plugins/PCG/Source/PCGEditor/Private/Pins/SPCGEditorGraphNodePin.cpp
--- a/Engine/Plugins/PCG/Source/PCGEditor/Private/Pins/SPCGEditorGraphNodePin.cpp
+++ b/Engine/Plugins/PCG/Source/PCGEditor/Private/Pins/SPCGEditorGraphNodePin.cpp
@@ -11,6 +11,7 @@
 
 #include "SLevelOfDetailBranchNode.h"
 #include "SPinTypeSelector.h"
+#include "ScopedTransaction.h"
 #include "Algo/AnyOf.h"
 #include "Schema/PCGEditorGraphSchema.h"
 #include "Widgets/SToolTip.h"
@@ -390,4 +391,25 @@ bool SPCGEditorGraphNodePin::ShouldDisplayAsRequiredForExecution() const
 	}
 
 	return false;
+}
+
+bool SPCGEditorGraphNodePin::SetPinDefaultValue(const FString& NewValue)
+{
+	if (!GraphPinObj || GraphPinObj->IsPendingKill() || GraphPinObj->DefaultValue == NewValue)
+	{
+		return false;
+	}
+
+	const UEdGraphSchema* Schema = GraphPinObj->GetSchema();
+	if (!Schema)
+	{
+		return false;
+	}
+
+	const FScopedTransaction Transaction(NSLOCTEXT("PCGGraphEditor", "ChangePinDefaultValue", "Change Pin Default Value"));
+	GraphPinObj->Modify();
+	Schema->TrySetDefaultValue(*GraphPinObj, NewValue);
+	OnModifyDelegate.ExecuteIfBound();
+
+	return true;
 }
\ No newline at end of file
diff --git a/Engine/Plugins/PCG/Source/PCGEditor/Private/Pins/SPCGEditorGraphNodePin.h b/Engine/Plugins/PCG/Source/PCGEditor/Private/Pins/SPCGEditorGraphNodePin.h
--- a/Engine/Plugins/PCG/Source/PCGEditor/Private/Pins/SPCGEditorGraphNodePin.h
+++ b/Engine/Plugins/PCG/Source/PCGEditor/Private/Pins/SPCGEditorGraphNodePin.h
@@ -35,6 +35,12 @@ public:
 	/** Whether pin is required to be connected for execution. */
 	bool ShouldDisplayAsRequiredForExecution() const;
 
+protected:
+	/** Sets pin default value inside a scoped transaction if the value changed. Returns true if the value was set. */
+	bool SetPinDefaultValue(const FString& NewValue);
+
+	FSimpleDelegate OnModifyDelegate;
+
 private:
 	void ApplyUnusedPinStyle(FSlateColor& InOutColor) const;
 	void GetPCGNodeAndPin(const UPCGNode*& OutNode, const UPCGPin*& OutPin) const;
diff --git a/Engine/Plugins/PCG/Source/PCGEditor/Private/Pins/SPCGEditorGraphPinSoftClassPath.h b/Engine/Plugins/PCG/Source/PCGEditor/Private/Pins/SPCGEditorGraphPinSoftClassPath.h
new file mode 100644
--- /dev/null
+++ b/Engine/Plugins/PCG/Source/PCGEditor/Private/Pins/SPCGEditorGraphPinSoftClassPath.h
@@ -0,0 +1,170 @@
+// Copyright Epic Games, Inc. All Rights Reserved.
+
+#pragma once
+
+#include "SPCGEditorGraphNodePin.h"
+
+#include "Editor.h"
+
+#include "ClassViewerModule.h"
+#include "Modules/ModuleManager.h"
+
+#include "Widgets/SBoxPanel.h"
+#include "Widgets/Images/SImage.h"
+#include "Widgets/Input/SButton.h"
+#include "Widgets/Input/SComboButton.h"
+#include "Widgets/Layout/SBorder.h"
+#include "Widgets/Layout/SBox.h"
+#include "Widgets/Text/STextBlock.h"
+
+// @todo_pcg: Propagate MetaClass from UPROPERTY metadata to filter the class viewer.
+/** Note: Based on KismetPins/SGraphPinClass.cpp */
+
+class SPCGEditorGraphPinSoftClassPath final : public SPCGEditorGraphNodePin
+{
+public:
+	SLATE_BEGIN_ARGS(SPCGEditorGraphPinSoftClassPath)
+		{}
+
+	SLATE_END_ARGS()
+
+	void Construct(const FArguments& InArgs, UEdGraphPin* InPin, TDelegate<void()>&& OnModify)
+	{
+		OnModifyDelegate = MoveTemp(OnModify);
+		SPCGEditorGraphNodePin::Construct(SPCGEditorGraphNodePin::FArguments(), InPin);
+	}
+
+protected:
+	virtual TSharedRef<SWidget> GetDefaultValueWidget() override
+	{
+		return SNew(SHorizontalBox)
+			   .Visibility(this, &SPCGEditorGraphNodePin::GetDefaultValueVisibility)
+			   .IsEnabled(this, &SPCGEditorGraphNodePin::GetDefaultValueIsEditable)
+				// Combo button — shows class name, opens class viewer
+				+ SHorizontalBox::Slot()
+			   .AutoWidth()
+			   .Padding(2.0f, 0.0f)
+			   .MaxWidth(200.0f)
+				[
+					SAssignNew(ClassPickerAnchor, SComboButton)
+				   .ButtonStyle(FAppStyle::Get(), "PropertyEditor.AssetComboStyle")
+				   .ContentPadding(FMargin(2.0f, 2.0f, 2.0f, 1.0f))
+				   .MenuPlacement(MenuPlacement_BelowAnchor)
+				   .ButtonContent()
+					[
+						SNew(STextBlock)
+					   .TextStyle(FAppStyle::Get(), "PropertyEditor.AssetClass")
+					   .Font(FAppStyle::GetFontStyle(TEXT("PropertyWindow.NormalFont")))
+					   .Text(this, &SPCGEditorGraphPinSoftClassPath::GetDisplayText)
+					   .ToolTipText(this, &SPCGEditorGraphPinSoftClassPath::GetFullPathText)
+					]
+				   .OnGetMenuContent(this, &SPCGEditorGraphPinSoftClassPath::GenerateClassPicker)
+				]
+				// Use button — set from editor's selected class
+				+ SHorizontalBox::Slot()
+			   .AutoWidth()
+			   .Padding(1.0f, 0.0f)
+			   .VAlign(VAlign_Center)
+				[
+					SNew(SButton)
+				   .ButtonStyle(FAppStyle::Get(), "NoBorder")
+				   .OnClicked(this, &SPCGEditorGraphPinSoftClassPath::OnClickUse)
+				   .ContentPadding(1.0f)
+				   .ToolTipText(NSLOCTEXT("PCGGraphEditor", "SoftClassPathUseTooltip", "Use selected class"))
+					[
+						SNew(SImage)
+					   .Image(FAppStyle::GetBrush(TEXT("Icons.CircleArrowLeft")))
+					]
+				]
+				// Browse button — sync content browser to class asset
+				+ SHorizontalBox::Slot()
+			   .AutoWidth()
+			   .Padding(1.0f, 0.0f)
+			   .VAlign(VAlign_Center)
+				[
+					SNew(SButton)
+				   .ButtonStyle(FAppStyle::Get(), "NoBorder")
+				   .OnClicked(this, &SPCGEditorGraphPinSoftClassPath::OnClickBrowse)
+				   .ContentPadding(0.0f)
+				   .ToolTipText(NSLOCTEXT("PCGGraphEditor", "SoftClassPathBrowseTooltip", "Browse to class"))
+					[
+						SNew(SImage)
+					   .Image(FAppStyle::GetBrush(TEXT("Icons.Search")))
+					]
+				];
+	}
+
+private:
+	FText GetDisplayText() const
+	{
+		const FString& Path = GraphPinObj->DefaultValue;
+		if (Path.IsEmpty())
+		{
+			return NSLOCTEXT("PCGGraphEditor", "SelectClass", "Select Class");
+		}
+
+		return FText::FromString(FSoftClassPath(Path).GetAssetName());
+	}
+
+	FText GetFullPathText() const
+	{
+		return FText::FromString(GraphPinObj->DefaultValue);
+	}
+
+	TSharedRef<SWidget> GenerateClassPicker()
+	{
+		FClassViewerModule& ClassViewerModule = FModuleManager::LoadModuleChecked<FClassViewerModule>("ClassViewer");
+
+		FClassViewerInitializationOptions Options;
+		Options.Mode = EClassViewerMode::ClassPicker;
+		Options.bShowNoneOption = true;
+		Options.NameTypeToDisplay = EClassViewerNameTypeToDisplay::DisplayName;
+
+		return SNew(SBox)
+	   .WidthOverride(350.0f)
+	   .HeightOverride(400.0f)
+		[
+			SNew(SBorder)
+		   .Padding(4.0f)
+		   .BorderImage(FAppStyle::GetBrush("ToolPanel.GroupBorder"))
+			[
+				ClassViewerModule.CreateClassViewer(Options, FOnClassPicked::CreateSP(this, &SPCGEditorGraphPinSoftClassPath::OnClassPicked))
+			]
+		];
+	}
+
+	void OnClassPicked(UClass* ChosenClass)
+	{
+		const FString NewPath = ChosenClass ? ChosenClass->GetPathName() : FString();
+		if (SetPinDefaultValue(NewPath))
+		{
+			ClassPickerAnchor->SetIsOpen(false);
+		}
+	}
+
+	FReply OnClickUse()
+	{
+		FEditorDelegates::LoadSelectedAssetsIfNeeded.Broadcast();
+
+		if (const UClass* SelectedClass = GEditor->GetFirstSelectedClass(UObject::StaticClass()))
+		{
+			SetPinDefaultValue(SelectedClass->GetPathName());
+		}
+
+		return FReply::Handled();
+	}
+
+	FReply OnClickBrowse()
+	{
+		if (UClass* Class = FSoftClassPath(GraphPinObj->DefaultValue).TryLoadClass<UObject>())
+		{
+			TArray<FAssetData> Objects;
+			Objects.Add(FAssetData(Class));
+			GEditor->SyncBrowserToObjects(Objects);
+		}
+
+		return FReply::Handled();
+	}
+
+	TSharedPtr<SComboButton> ClassPickerAnchor;
+};
diff --git a/Engine/Plugins/PCG/Source/PCGEditor/Private/Pins/SPCGEditorGraphPinSoftObjectPath.h b/Engine/Plugins/PCG/Source/PCGEditor/Private/Pins/SPCGEditorGraphPinSoftObjectPath.h
new file mode 100644
--- /dev/null
+++ b/Engine/Plugins/PCG/Source/PCGEditor/Private/Pins/SPCGEditorGraphPinSoftObjectPath.h
@@ -0,0 +1,187 @@
+// Copyright Epic Games, Inc. All Rights Reserved.
+
+#pragma once
+
+#include "SPCGEditorGraphNodePin.h"
+
+#include "Editor.h"
+#include "Selection.h"
+
+#include "ContentBrowserModule.h"
+#include "IContentBrowserSingleton.h"
+#include "AssetRegistry/AssetRegistryModule.h"
+
+#include "Widgets/SBoxPanel.h"
+#include "Widgets/Images/SImage.h"
+#include "Widgets/Input/SButton.h"
+#include "Widgets/Input/SComboButton.h"
+#include "Widgets/Layout/SBorder.h"
+#include "Widgets/Layout/SBox.h"
+#include "Widgets/Text/STextBlock.h"
+
+// @todo_pcg: Propagate AllowedClasses from UPROPERTY metadata to filter the asset picker.
+// @todo_pcg: Show the asset thumbnail next to the display name.
+/** Note: Based on KismetPins/SGraphPinObject.h */
+class SPCGEditorGraphPinSoftObjectPath final : public SPCGEditorGraphNodePin
+{
+public:
+	SLATE_BEGIN_ARGS(SPCGEditorGraphPinSoftObjectPath)
+		{}
+
+	SLATE_END_ARGS()
+
+	void Construct(const FArguments& InArgs, UEdGraphPin* InPin, TDelegate<void()>&& OnModify)
+	{
+		OnModifyDelegate = MoveTemp(OnModify);
+		SPCGEditorGraphNodePin::Construct(SPCGEditorGraphNodePin::FArguments(), InPin);
+	}
+
+protected:
+	virtual TSharedRef<SWidget> GetDefaultValueWidget() override
+	{
+		return SNew(SHorizontalBox)
+			   .Visibility(this, &SPCGEditorGraphNodePin::GetDefaultValueVisibility)
+			   .IsEnabled(this, &SPCGEditorGraphNodePin::GetDefaultValueIsEditable)
+				// Combo button — shows asset name, opens picker
+				+ SHorizontalBox::Slot()
+			   .AutoWidth()
+			   .Padding(2.0f, 0.0f)
+			   .MaxWidth(200.0f)
+				[
+					SAssignNew(AssetPickerAnchor, SComboButton)
+				   .ButtonStyle(FAppStyle::Get(), "PropertyEditor.AssetComboStyle")
+				   .ContentPadding(FMargin(2.0f, 2.0f, 2.0f, 1.0f))
+				   .MenuPlacement(MenuPlacement_BelowAnchor)
+				   .ButtonContent()
+					[
+						SNew(STextBlock)
+					   .TextStyle(FAppStyle::Get(), "PropertyEditor.AssetClass")
+					   .Font(FAppStyle::GetFontStyle(TEXT("PropertyWindow.NormalFont")))
+					   .Text(this, &SPCGEditorGraphPinSoftObjectPath::GetDisplayText)
+					   .ToolTipText(this, &SPCGEditorGraphPinSoftObjectPath::GetFullPathText)
+					]
+				   .OnGetMenuContent(this, &SPCGEditorGraphPinSoftObjectPath::GenerateAssetPicker)
+				]
+				// Use button — set from content browser selection
+				+ SHorizontalBox::Slot()
+			   .AutoWidth()
+			   .Padding(1.0f, 0.0f)
+			   .VAlign(VAlign_Center)
+				[
+					SNew(SButton)
+				   .ButtonStyle(FAppStyle::Get(), "NoBorder")
+				   .OnClicked(this, &SPCGEditorGraphPinSoftObjectPath::OnClickUse)
+				   .ContentPadding(1.0f)
+				   .ToolTipText(NSLOCTEXT("PCGGraphEditor", "SoftObjectPathUseTooltip", "Use asset browser selection"))
+					[
+						SNew(SImage)
+					   .Image(FAppStyle::GetBrush(TEXT("Icons.CircleArrowLeft")))
+					]
+				]
+				// Browse button — sync content browser to value
+				+ SHorizontalBox::Slot()
+			   .AutoWidth()
+			   .Padding(1.0f, 0.0f)
+			   .VAlign(VAlign_Center)
+				[
+					SNew(SButton)
+				   .ButtonStyle(FAppStyle::Get(), "NoBorder")
+				   .OnClicked(this, &SPCGEditorGraphPinSoftObjectPath::OnClickBrowse)
+				   .ContentPadding(0.0f)
+				   .ToolTipText(NSLOCTEXT("PCGGraphEditor", "SoftObjectPathBrowseTooltip", "Browse to asset"))
+					[
+						SNew(SImage)
+					   .Image(FAppStyle::GetBrush(TEXT("Icons.Search")))
+					]
+				];
+	}
+
+private:
+	FText GetDisplayText() const
+	{
+		const FString& Path = GraphPinObj->DefaultValue;
+		if (Path.IsEmpty())
+		{
+			return NSLOCTEXT("PCGGraphEditor", "SelectAsset", "Select Asset");
+		}
+
+		return FText::FromString(FSoftObjectPath(Path).GetAssetName());
+	}
+
+	FText GetFullPathText() const
+	{
+		return FText::FromString(GraphPinObj->DefaultValue);
+	}
+
+	TSharedRef<SWidget> GenerateAssetPicker()
+	{
+		FContentBrowserModule& ContentBrowserModule = FModuleManager::Get().LoadModuleChecked<FContentBrowserModule>(TEXT("ContentBrowser"));
+
+		FAssetPickerConfig AssetPickerConfig;
+		AssetPickerConfig.bAllowNullSelection = true;
+		AssetPickerConfig.OnAssetSelected = FOnAssetSelected::CreateSP(this, &SPCGEditorGraphPinSoftObjectPath::OnAssetSelected);
+		AssetPickerConfig.OnAssetEnterPressed = FOnAssetEnterPressed::CreateSP(this, &SPCGEditorGraphPinSoftObjectPath::OnAssetEnterPressed);
+		AssetPickerConfig.InitialAssetViewType = EAssetViewType::List;
+		AssetPickerConfig.bAllowDragging = false;
+
+		return SNew(SBox)
+	   .HeightOverride(300.0f)
+	   .WidthOverride(300.0f)
+		[
+			SNew(SBorder)
+		   .BorderImage(FAppStyle::GetBrush("Menu.Background"))
+			[
+				ContentBrowserModule.Get().CreateAssetPicker(AssetPickerConfig)
+			]
+		];
+	}
+
+	void OnAssetSelected(const FAssetData& AssetData)
+	{
+		const FString NewPath = AssetData.IsValid() ? AssetData.GetObjectPathString() : FString();
+		if (SetPinDefaultValue(NewPath))
+		{
+			AssetPickerAnchor->SetIsOpen(false);
+		}
+	}
+
+	void OnAssetEnterPressed(const TArray<FAssetData>& InSelectedAssets)
+	{
+		if (InSelectedAssets.Num() > 0)
+		{
+			OnAssetSelected(InSelectedAssets[0]);
+		}
+	}
+
+	FReply OnClickUse()
+	{
+		FEditorDelegates::LoadSelectedAssetsIfNeeded.Broadcast();
+
+		if (const UObject* SelectedObject = GEditor->GetSelectedObjects()->GetTop<UObject>())
+		{
+			SetPinDefaultValue(SelectedObject->GetPathName());
+		}
+
+		return FReply::Handled();
+	}
+
+	FReply OnClickBrowse()
+	{
+		const FString& Path = GraphPinObj->DefaultValue;
+		if (!Path.IsEmpty())
+		{
+			const FAssetRegistryModule& AssetRegistryModule = FModuleManager::LoadModuleChecked<FAssetRegistryModule>(TEXT("AssetRegistry"));
+			FAssetData AssetData = AssetRegistryModule.Get().GetAssetByObjectPath(FSoftObjectPath(Path));
+			if (AssetData.IsValid())
+			{
+				TArray<FAssetData> Objects;
+				Objects.Add(AssetData);
+				GEditor->SyncBrowserToObjects(Objects);
+			}
+		}
+
+		return FReply::Handled();
+	}
+
+	TSharedPtr<SComboButton> AssetPickerAnchor;
+};
