[PCG] Inline Constant Widget for Quaternion

Note: there is already a widget for FRotator, but we support FQuat constants, so we need an inline widget.

#Synced-CL 52134510

diff --git a/Engine/Plugins/PCG/Source/PCG/Public/Helpers/PCGMetadataHelpers.h b/Engine/Plugins/PCG/Source/PCG/Public/Helpers/PCGMetadataHelpers.h
--- a/Engine/Plugins/PCG/Source/PCG/Public/Helpers/PCGMetadataHelpers.h
+++ b/Engine/Plugins/PCG/Source/PCG/Public/Helpers/PCGMetadataHelpers.h
@@ -108,22 +108,23 @@ namespace PCGMetadataHelpers
 	{
 		switch (Type)
 		{
-			case EPCGMetadataTypes::Double:    // fall-through
-			case EPCGMetadataTypes::Integer32: // fall-through
-			case EPCGMetadataTypes::Integer64: // fall-through
-			case EPCGMetadataTypes::Vector:    // fall-through
-			case EPCGMetadataTypes::Vector2:   // fall-through
-			case EPCGMetadataTypes::Vector4:   // fall-through
-			case EPCGMetadataTypes::String:    // fall-through
-			case EPCGMetadataTypes::Name:      // fall-through
-			case EPCGMetadataTypes::Boolean:   // fall-through
-			case EPCGMetadataTypes::Rotator:
-				return true;
-			// @todo_pcg: Enable the rest once they're supported in the UI
+			case EPCGMetadataTypes::Double:         // fall-through
+			case EPCGMetadataTypes::Integer32:      // fall-through
+			case EPCGMetadataTypes::Integer64:      // fall-through
+			case EPCGMetadataTypes::Vector:         // fall-through
+			case EPCGMetadataTypes::Vector2:        // fall-through
+			case EPCGMetadataTypes::Vector4:        // fall-through
+			case EPCGMetadataTypes::String:         // fall-through
+			case EPCGMetadataTypes::Name:           // fall-through
+			case EPCGMetadataTypes::Boolean:        // fall-through
+			case EPCGMetadataTypes::Rotator:        // fall-through
 			case EPCGMetadataTypes::Quaternion:     // fall-through
 			case EPCGMetadataTypes::Transform:      // fall-through
 			case EPCGMetadataTypes::SoftObjectPath: // fall-through
+			case EPCGMetadataTypes::SoftObject:     // fall-through
 			case EPCGMetadataTypes::SoftClassPath:  // fall-through
+			case EPCGMetadataTypes::SoftClass:
+				return true;
 			// This is automatically converted in the accessor to a Double
 			case EPCGMetadataTypes::Float: // fall-through
 			// Anything else is not valid.
diff --git a/Engine/Plugins/PCG/Source/PCGEditor/Private/Nodes/SPCGEditorGraphNode.cpp b/Engine/Plugins/PCG/Source/PCGEditor/Private/Nodes/SPCGEditorGraphNode.cpp
--- a/Engine/Plugins/PCG/Source/PCGEditor/Private/Nodes/SPCGEditorGraphNode.cpp
+++ b/Engine/Plugins/PCG/Source/PCGEditor/Private/Nodes/SPCGEditorGraphNode.cpp
@@ -15,6 +15,10 @@
 #include "Pins/SPCGEditorGraphPinBool.h"
 #include "Pins/SPCGEditorGraphPinNumSlider.h"
 #include "Pins/SPCGEditorGraphPinString.h"
+#include "Pins/SPCGEditorGraphPinQuaternion.h"
+#include "Pins/SPCGEditorGraphPinSoftClassPath.h"
+#include "Pins/SPCGEditorGraphPinSoftObjectPath.h"
+#include "Pins/SPCGEditorGraphPinTransform.h"
 #include "Pins/SPCGEditorGraphPinVectorSlider.h"
 #include "Schema/PCGEditorGraphSchema.h"
 
@@ -140,7 +144,12 @@ TSharedPtr<SGraphPin> SPCGEditorGraphNode::CreatePinWidget(UEdGraphPin* InPin) c
 	if (InPin && NodeSettings && NodeSettings->Implements<UPCGSettingsDefaultValueProvider>())
 	{
 		const IPCGSettingsDefaultValueProvider* DefaultValueInterface = CastChecked<IPCGSettingsDefaultValueProvider>(NodeSettings);
-		if (!InPin->HasAnyConnections() && DefaultValueInterface->DefaultValuesAreEnabled() && DefaultValueInterface->IsPinDefaultValueActivated(InPin->PinName))
+		const bool bDefaultIsEnabled = DefaultValueInterface->DefaultValuesAreEnabled();
+		const bool bDefaultIsActivated = DefaultValueInterface->IsPinDefaultValueActivated(InPin->PinName);
+		const bool bInputHasAnyConnections = InPin->HasAnyConnections();
+		const bool bIsInput = InPin->Direction == EGPD_Input;
+		// Output pins (i.e. inline constants) always show the widget and input pins only when unconnected.
+		if (bDefaultIsEnabled && bDefaultIsActivated && (!bIsInput || !bInputHasAnyConnections))
 		{
 			// Set the string default value to match the settings' source of truth.
 			InPin->DefaultValue = DefaultValueInterface->GetPinDefaultValueAsString(InPin->PinName);
@@ -187,12 +196,17 @@ TSharedPtr<SGraphPin> SPCGEditorGraphNode::CreatePinWidget(UEdGraphPin* InPin) c
 				case EPCGMetadataTypes::Rotator:
 					return SNew(SPCGEditorGraphPinVectorSlider<FRotator>, InPin, MoveTemp(OnModify));
 				case EPCGMetadataTypes::Boolean:
-					return SNew(SPCGEditorGraphPinBool, InPin, std::move(OnModify));
-				// @todo_pcg: Will be added once widgets are created.
+					return SNew(SPCGEditorGraphPinBool, InPin, MoveTemp(OnModify));
+				case EPCGMetadataTypes::Quaternion:
+					return SNew(SPCGEditorGraphPinQuaternion, InPin, MoveTemp(OnModify));
+				case EPCGMetadataTypes::Transform:
+					return SNew(SPCGEditorGraphPinTransform, InPin, MoveTemp(OnModify));
 				case EPCGMetadataTypes::SoftObjectPath: // fall-through
-				case EPCGMetadataTypes::SoftClassPath:  // fall-through
-				case EPCGMetadataTypes::Quaternion:     // fall-through
-				case EPCGMetadataTypes::Transform:      // fall-through
+				case EPCGMetadataTypes::SoftObject:
+					return SNew(SPCGEditorGraphPinSoftObjectPath, InPin, MoveTemp(OnModify));
+				case EPCGMetadataTypes::SoftClassPath: // fall-through
+				case EPCGMetadataTypes::SoftClass:
+					return SNew(SPCGEditorGraphPinSoftClassPath, InPin, MoveTemp(OnModify));
 				default:
 					break;
 			}
diff --git a/Engine/Plugins/PCG/Source/PCGEditor/Private/Pins/SPCGEditorGraphPinQuaternion.h b/Engine/Plugins/PCG/Source/PCGEditor/Private/Pins/SPCGEditorGraphPinQuaternion.h
new file mode 100644
--- /dev/null
+++ b/Engine/Plugins/PCG/Source/PCGEditor/Private/Pins/SPCGEditorGraphPinQuaternion.h
@@ -0,0 +1,106 @@
+// Copyright Epic Games, Inc. All Rights Reserved.
+
+#pragma once
+
+#include "SPCGEditorGraphNodePin.h"
+
+#include "KismetPins/SVectorSlider.h"
+
+/**
+ * Quaternion pin widget. Displays rotator inputs (Roll, Pitch, Yaw) that are converted to/from FQuat for storage.
+ * Note: Derived from SPCGEditorGraphPinVectorSlider pattern.
+ */
+class SPCGEditorGraphPinQuaternion final : public SPCGEditorGraphNodePin
+{
+public:
+	SLATE_BEGIN_ARGS(SPCGEditorGraphPinQuaternion)
+		{}
+
+	SLATE_END_ARGS()
+
+	void Construct(const FArguments& InArgs, UEdGraphPin* InPin, TDelegate<void()>&& OnModify)
+	{
+		OnModifyDelegate = MoveTemp(OnModify);
+		SPCGEditorGraphNodePin::Construct(SPCGEditorGraphNodePin::FArguments(), InPin);
+	}
+
+protected:
+	virtual TSharedRef<SWidget> GetDefaultValueWidget() override
+	{
+		return SNew(SVectorSlider<double>, /*bIsRotator=*/true, nullptr)
+	   .VisibleText_0(this, &SPCGEditorGraphPinQuaternion::GetCurrentStringValue, 0)
+	   .VisibleText_1(this, &SPCGEditorGraphPinQuaternion::GetCurrentStringValue, 1)
+	   .VisibleText_2(this, &SPCGEditorGraphPinQuaternion::GetCurrentStringValue, 2)
+	   .OnNumericCommitted_Box_0(this, &SPCGEditorGraphPinQuaternion::OnValueCommitted, 0)
+	   .OnNumericCommitted_Box_1(this, &SPCGEditorGraphPinQuaternion::OnValueCommitted, 1)
+	   .OnNumericCommitted_Box_2(this, &SPCGEditorGraphPinQuaternion::OnValueCommitted, 2)
+	   .Visibility(this, &SPCGEditorGraphNodePin::GetDefaultValueVisibility)
+	   .IsEnabled(this, &SPCGEditorGraphNodePin::GetDefaultValueIsEditable);
+	}
+
+private:
+	// --- Getter (index: 0=Roll, 1=Pitch, 2=Yaw) ---
+	FString GetCurrentStringValue(const int32 Index) const
+	{
+		if (Index == 0)
+		{
+			UpdateFromDefaultValue();
+		}
+		return FString::SanitizeFloat(Index == 0 ? CachedRotator.Roll : Index == 1 ? CachedRotator.Pitch : CachedRotator.Yaw);
+	}
+
+	// --- Setter (the component index is passed along with the value) ---
+	void OnValueCommitted(const double NewValue, ETextCommit::Type, const int32 Index)
+	{
+		double& Component = Index == 0 ? CachedRotator.Roll : Index == 1 ? CachedRotator.Pitch : CachedRotator.Yaw;
+		SetComponent(Component, NewValue);
+	}
+
+	void UpdateFromDefaultValue() const
+	{
+		if (GraphPinObj->IsPendingKill())
+		{
+			return;
+		}
+
+		const uint32 Hash = GetTypeHash(GraphPinObj->DefaultValue);
+		if (Hash == DefaultValueHash)
+		{
+			return;
+		}
+
+		DefaultValueHash = Hash;
+
+		FQuat Quat = FQuat::Identity;
+		const UScriptStruct* QuatStruct = TBaseStructure<FQuat>::Get();
+		QuatStruct->ImportText(*GraphPinObj->DefaultValue, &Quat, /*OwnerObject=*/nullptr, PPF_None, /*ErrorText=*/nullptr, QuatStruct->GetName());
+		CachedRotator = Quat.Rotator();
+	}
+
+	void SetDefaultValue()
+	{
+		const FQuat Quat = CachedRotator.Quaternion();
+		FString Result;
+		TBaseStructure<FQuat>::Get()->ExportText(Result, &Quat, /*Defaults=*/nullptr, /*OwnerObject=*/nullptr, PPF_None, /*ExportRootScope=*/nullptr);
+		SetPinDefaultValue(Result);
+	}
+
+	void SetComponent(double& Component, const double NewValue)
+	{
+		if (GraphPinObj->IsPendingKill())
+		{
+			return;
+		}
+
+		if (NewValue == Component)
+		{
+			return;
+		}
+
+		Component = NewValue;
+		SetDefaultValue();
+	}
+
+	mutable FRotator CachedRotator = FRotator::ZeroRotator;
+	mutable uint32 DefaultValueHash = 0;
+};
