[PCG] Fix pre-existing static analysis warning

#rnx
#Synced-CL 52135071

diff --git a/Engine/Plugins/PCG/Source/PCGEditor/Private/Schema/PCGEditorGraphSchema.cpp b/Engine/Plugins/PCG/Source/PCGEditor/Private/Schema/PCGEditorGraphSchema.cpp
--- a/Engine/Plugins/PCG/Source/PCGEditor/Private/Schema/PCGEditorGraphSchema.cpp
+++ b/Engine/Plugins/PCG/Source/PCGEditor/Private/Schema/PCGEditorGraphSchema.cpp
@@ -1431,7 +1431,7 @@ void FPCGEditorConnectionDrawingPolicy::DetermineWiringStyle(UEdGraphPin* Output
 	}
 
 	// Base the color of the wire on the color of the output pin
-	if (OutputPin && Graph->GetSchema())
+	if (OutputPin && Graph && Graph->GetSchema())
 	{
 		if (const UPCGEditorGraphSchema* Schema = Cast<UPCGEditorGraphSchema>(Graph->GetSchema()))
 		{
