[PCG] Add GPU backend for World Raycast node using inline ray tracing

Implements GPU execution for the PCG World Raycast node. Points are projected onto scene geometry via inline ray queries (DXR 1.1 RayQuery) dispatched from a compute shader.

Pipeline:
- PCGPreRayTraceKernel: packs point origins and ray parameters into a raw buffer (one ray per point), then dispatches the inline ray trace in PostSubmit after the TLAS is built.
- PCGRayTrace.usf: compute shader that traces rays against the scene TLAS, interpolates hit UVs via barycentric coordinates, and writes back impact position, normals, and texture coordinates.
- PCGPostRayTraceKernel: unpacks hit results back into point attributes (position, ImpactNormal, ImpactUVCoords).

#Synced-CL 52169849

diff --git a/Engine/Plugins/PCG/Shaders/Private/BuiltInKernels/PCGPostRayTrace.usf b/Engine/Plugins/PCG/Shaders/Private/BuiltInKernels/PCGPostRayTrace.usf
new file mode 100644
--- /dev/null
+++ b/Engine/Plugins/PCG/Shaders/Private/BuiltInKernels/PCGPostRayTrace.usf
@@ -0,0 +1,105 @@
+// Copyright Epic Games, Inc. All Rights Reserved.
+
+// This is a post-process kernel that unpacks a raw buffer (of packed data emitted by PCGRayTrace.usf) into one point per ray.
+
+[numthreads(64, 1, 1)]
+void PCGPostRayTraceCS(uint3 GroupId : SV_GroupID, uint GroupIndex : SV_GroupIndex)
+{
+	// Mark the kernel as having executed. Must run before we early out via thread index, because the kernel is still 'executed' even if the number of
+	// threads to iterate on is zero. Even if GetNumThreads() returns 0, the kernel will still have been dispatched on a single thread to set this flag.
+	if (all(GroupId == 0) && GroupIndex == 0)
+	{
+		Out_SetAsExecutedInternal();
+	}
+
+	const uint ThreadIndex = GetUnWrappedDispatchThreadId(GroupId, GroupIndex, 64);
+	if (ThreadIndex >= GetNumThreads().x) return;
+
+	uint In_DataIndex, In_ElementIndex;
+	if (!In_GetThreadData(ThreadIndex, In_DataIndex, In_ElementIndex))
+	{
+		return;
+	}
+
+	uint Out_DataIndex, Out_ElementIndex;
+	if (!Out_GetThreadData(ThreadIndex, Out_DataIndex, Out_ElementIndex))
+	{
+		return;
+	}
+
+	// Propagate the 'Removed' status to the output point. Otherwise, this point will not be culled.
+	if (In_IsPointRemoved(In_DataIndex, In_ElementIndex))
+	{
+		Out_RemovePoint(Out_DataIndex, Out_ElementIndex);
+		return;
+	}
+	
+	const uint RayTraceResults_DataIndex = 0;
+	const uint RayTraceResults_FirstElementIndex = ThreadIndex * PCG_RAY_TRACE_PACKED_BUFFER_STRIDE_UINTS;
+	const uint PackedPositionIndex = RayTraceResults_FirstElementIndex;
+	const uint PackedTexCoordsIndex = RayTraceResults_FirstElementIndex + 4;
+	const uint PackedNormalsIndex = RayTraceResults_FirstElementIndex + 8;
+
+	if (PackedPositionIndex >= RayTraceResults_GetNumElements(RayTraceResults_DataIndex))
+	{
+		// Should not happen - invalid results data.
+		return;
+	}
+
+	const uint4 PackedPositionAndValid = RayTraceResults_Load4(RayTraceResults_DataIndex, PackedPositionIndex);
+	const float3 HitPosition = asfloat(PackedPositionAndValid.xyz);
+	bool bMissedOrOutOfBounds = false;
+
+	// Handle if the ray was missed, culled, or outside the bounds.
+	if (PackedPositionAndValid.w == PCG_RAY_CULLED || !(WorldRaycast_GetUnbounded() != 0 || (all(HitPosition >= GetComponentBoundsMin()) && all(HitPosition < GetComponentBoundsMax()))))
+	{
+		if (WorldRaycast_GetKeepOriginalPointOnMiss() != 0)
+		{
+			bMissedOrOutOfBounds = true;
+		}
+		else
+		{
+			Out_RemovePoint(Out_DataIndex, Out_ElementIndex);
+			return;
+		}
+	}
+
+	PCG_COPY_ALL_ATTRIBUTES_TO_OUTPUT(Out, In, Out_DataIndex, Out_ElementIndex, In_DataIndex, In_ElementIndex);
+
+	if (!bMissedOrOutOfBounds)
+	{
+		Out_SetPosition(Out_DataIndex, Out_ElementIndex, HitPosition);
+	}
+
+	const int TexCoordsAttributeId = PostRayTrace_GetTexCoordsAttributeId();
+
+	if (TexCoordsAttributeId >= 0)
+	{
+		if (!bMissedOrOutOfBounds)
+		{
+			const uint4 PackedTexCoords = RayTraceResults_Load4(RayTraceResults_DataIndex, PackedTexCoordsIndex);
+			const float2 TexCoords = asfloat(PackedTexCoords.xy);
+			Out_SetFloat2(Out_DataIndex, Out_ElementIndex, TexCoordsAttributeId, TexCoords);
+		}
+		else
+		{
+			Out_SetFloat2(Out_DataIndex, Out_ElementIndex, TexCoordsAttributeId, (float2)0.0f);
+		}
+	}
+
+	const int NormalsAttributeId = PostRayTrace_GetNormalsAttributeId();
+
+	if (NormalsAttributeId >= 0)
+	{
+		if (!bMissedOrOutOfBounds)
+		{
+			const uint4 PackedNormals = RayTraceResults_Load4(RayTraceResults_DataIndex, PackedNormalsIndex);
+			const float3 Normals = asfloat(PackedNormals.xyz);
+			Out_SetFloat3(Out_DataIndex, Out_ElementIndex, NormalsAttributeId, Normals);
+		}
+		else
+		{
+			Out_SetFloat3(Out_DataIndex, Out_ElementIndex, NormalsAttributeId, float3(0.0f, 0.0f, 1.0f));
+		}
+	}
+}
diff --git a/Engine/Plugins/PCG/Shaders/Private/BuiltInKernels/PCGPreRayTrace.usf b/Engine/Plugins/PCG/Shaders/Private/BuiltInKernels/PCGPreRayTrace.usf
new file mode 100644
--- /dev/null
+++ b/Engine/Plugins/PCG/Shaders/Private/BuiltInKernels/PCGPreRayTrace.usf
@@ -0,0 +1,111 @@
+// Copyright Epic Games, Inc. All Rights Reserved.
+
+// This is a pre-process kernel that packs a raw buffer (of packed data to be fed into PCGRayTrace.usf) with one ray per point.
+
+#define PCG_END_POINTS_ENABLED {EndPointsEnabled}
+
+#define PCG_RAYCAST_MODE_INFINITE 0
+#define PCG_RAYCAST_MODE_SCALED_VECTOR 1
+#define PCG_RAYCAST_MODE_NORMALIZED_WITH_LENGTH 2
+#define PCG_RAYCAST_MODE_SEGMENTS 3
+
+[numthreads(64, 1, 1)]
+void PCGPreRayTraceCS(uint3 GroupId : SV_GroupID, uint GroupIndex : SV_GroupIndex)
+{
+	const uint ThreadIndex = GetUnWrappedDispatchThreadId(GroupId, GroupIndex, 64);
+	if (ThreadIndex >= GetNumThreads().x) return;
+
+	uint In_DataIndex, In_ElementIndex;
+	if (!In_GetThreadData(ThreadIndex, In_DataIndex, In_ElementIndex))
+	{
+		return;
+	}
+
+	const uint Out_DataIndex = 0;
+	const uint Out_FirstElementIndex = ThreadIndex * PCG_RAY_TRACE_PACKED_BUFFER_STRIDE_UINTS;
+	const uint PackedPositionIndex = Out_FirstElementIndex;
+	const uint PackedDirectionIndex = Out_FirstElementIndex + 4;
+	const uint PackedNormalsIndex = Out_FirstElementIndex + 8;
+
+	// If the point is culled, the ray should also be culled.
+	if (In_IsPointRemoved(In_DataIndex, In_ElementIndex))
+	{
+		Out_Store4(Out_DataIndex, PackedPositionIndex, (uint4)PCG_RAY_CULLED);
+		return;
+	}
+
+	const int RaycastMode = WorldRaycast_GetRaycastMode();
+	const int RayOriginAttributeId = WorldRaycast_GetRayOriginAttributeId();
+	const int RayDirectionAttributeId = WorldRaycast_GetRayDirectionAttributeId();
+	const int RayLengthAttributeId = WorldRaycast_GetRayLengthAttributeId();
+
+	float3 RayOrigin = (float3)0.0f;
+	float3 RayDirection = (float3)0.0f;
+	float RayLength = 0.0f;
+
+	if (RayOriginAttributeId != -1)
+	{
+		RayOrigin = In_GetFloat3(In_DataIndex, In_ElementIndex, RayOriginAttributeId);
+	}
+	else
+	{
+		RayOrigin = In_GetPosition(In_DataIndex, In_ElementIndex);
+	}
+
+	if (RayDirectionAttributeId != -1)
+	{
+		RayDirection = In_GetFloat3(In_DataIndex, In_ElementIndex, RayDirectionAttributeId);
+	}
+	else
+	{
+		RayDirection = WorldRaycast_GetRayDirection();
+	}
+
+	// Compute ray length.
+	if (RaycastMode == PCG_RAYCAST_MODE_INFINITE)
+	{
+		RayLength = 1e30f;
+	}
+	else if (RaycastMode == PCG_RAYCAST_MODE_SCALED_VECTOR)
+	{
+		RayLength = length(RayDirection);
+	}
+	else if (RaycastMode == PCG_RAYCAST_MODE_NORMALIZED_WITH_LENGTH)
+	{
+		if (RayLengthAttributeId != -1)
+		{
+			RayLength = In_GetFloat(In_DataIndex, In_ElementIndex, RayLengthAttributeId);
+		}
+		else
+		{
+			RayLength = WorldRaycast_GetRayLength();
+		}
+	}
+#if PCG_END_POINTS_ENABLED
+	else if (RaycastMode == PCG_RAYCAST_MODE_SEGMENTS)
+	{
+		const int RayEndAttributeId = WorldRaycast_GetRayEndAttributeId();
+		float3 RayEndPosition = (float3)0.0f;
+
+		if (RayEndAttributeId != -1)
+		{
+			RayEndPosition = End_Points_GetFloat3(In_DataIndex, In_ElementIndex, RayEndAttributeId);
+		}
+		else
+		{
+			RayEndPosition = End_Points_GetPosition(In_DataIndex, In_ElementIndex);
+		}
+
+		// In Segments mode, the ray direction is also controlled by the segment.
+		RayDirection = RayEndPosition - RayOrigin;
+		RayLength = length(RayDirection);
+	}
+#endif
+
+	// Normalize ray direction.
+	RayDirection = dot(RayDirection, RayDirection) != 0.0f ? normalize(RayDirection) : (float3)0.0f;
+
+	Out_Store4(Out_DataIndex, PackedPositionIndex, uint4(asuint(RayOrigin), 0u));
+	Out_Store4(Out_DataIndex, PackedDirectionIndex, uint4(asuint(RayDirection), asuint(RayLength)));
+	Out_Store4(Out_DataIndex, PackedNormalsIndex, (uint4)0u); // Initialize buffer for normal values that will be written later.
+}
diff --git a/Engine/Plugins/PCG/Shaders/Private/PCGRayTrace.usf b/Engine/Plugins/PCG/Shaders/Private/PCGRayTrace.usf
new file mode 100644
--- /dev/null
+++ b/Engine/Plugins/PCG/Shaders/Private/PCGRayTrace.usf
@@ -0,0 +1,170 @@
+// Copyright Epic Games, Inc. All Rights Reserved.
+
+#include "/Engine/Private/Common.ush"
+#include "/Engine/Private/RayTracing/RayTracingCommon.ush"
+#include "/Engine/Private/RayTracing/TraceRayInline.ush"
+#include "/Engine/Public/BindlessResources.ush"
+#include "/Engine/Private/Visualization.ush"
+
+#define PCG_DEFAULT_NORMAL float3(0.0f, 0.0f, 1.0f)
+
+// @todo_pcg: Consider splitting into separate resources per output type (position, normals, UVs) so that unused outputs don't consume memory/bandwidth.
+/**
+ * Input packed data for each ray:
+ * - Ray Position    - 3 uints
+ * - Ray Culled      - 1 uint
+ * - Ray Direction   - 3 uints
+ * - Ray Max Length  - 1 uint
+ * - Empty packing   - 4 uints
+ *
+ * Output packed data for each ray:
+ *  - Impact Position   - 3 uints
+ *  - Ray Culled/Missed - 1 uint
+ *  - Impact TexCoords  - 2 uints
+ *  - Empty packing     - 2 uints
+ *  - Impact Normals    - 3 uints
+ *  - Empty packing     - 1 uint
+ */
+RWByteAddressBuffer InOutPackedData;
+RaytracingAccelerationStructure InTLAS;
+#if PLATFORM_SUPPORTS_INLINE_RAY_TRACING_TRIANGLE_ATTRIBUTES
+FResourceCollection InPrimitiveUVCollection;
+Buffer<uint> InPrimitiveUVOffsets;
+StructuredBuffer<FRayTracingSceneMetadataRecord> InRayTracingSceneMetadata;
+#endif // PLATFORM_SUPPORTS_INLINE_RAY_TRACING_TRIANGLE_ATTRIBUTES
+float3 InPreViewTranslation;
+int InNumRays;
+int InTexCoordsChannelIndex;
+
+/** Unpack the data for this ray. */
+void ReadInput(int InRayIndex, out float3 OutOrigin, out float3 OutDirection, out float OutLength, out int OutStatus)
+{
+	const uint PackedPositionAddress = InRayIndex * (PCG_RAY_TRACE_PACKED_BUFFER_STRIDE_UINTS << 2);
+	const uint PackedDirectionAddress = PackedPositionAddress + (4 << 2);
+	const uint4 PositionAndValid = InOutPackedData.Load4(PackedPositionAddress);
+	const uint4 DirectionAndLength = InOutPackedData.Load4(PackedDirectionAddress);
+
+	// Account for local view when computing the ray origin.
+	OutOrigin = asfloat(PositionAndValid.xyz) + InPreViewTranslation;
+	OutDirection = asfloat(DirectionAndLength.xyz);
+	OutLength = asfloat(DirectionAndLength.w);
+	OutStatus = PositionAndValid.w;
+}
+
+/** Pack the impact data for this ray. */
+void WriteOutput(int InRayIndex, float3 InHitPosition, float2 InTexCoords, float3 InNormals, int InStatus)
+{
+	const uint PackedPositionAddress = InRayIndex * (PCG_RAY_TRACE_PACKED_BUFFER_STRIDE_UINTS << 2);
+	const uint PackedTexCoordsAddress = PackedPositionAddress + (4 << 2);
+	const uint PackedNormalsAddress = PackedPositionAddress + (8 << 2);
+	const uint4 PositionAndValid = uint4(asuint(InHitPosition - InPreViewTranslation), (uint)InStatus);
+	const uint4 TexCoords = uint4(asuint(InTexCoords).xy, 0u, 0u);
+	const uint4 Normals = uint4(asuint(InNormals).xyz, 0u);
+
+	InOutPackedData.Store4(PackedPositionAddress, PositionAndValid);
+	InOutPackedData.Store4(PackedTexCoordsAddress, TexCoords);
+	InOutPackedData.Store4(PackedNormalsAddress, Normals);
+}
+
+void PCGRayTraceCommon(int InRayIndex, float3 InOrigin, float3 InDirection, float InLength)
+{
+	FRayDesc Ray;
+	Ray.Origin = InOrigin;
+	Ray.Direction = InDirection;
+	Ray.TMin = 0.001f;
+	Ray.TMax = InLength;
+
+	FTraceRayInlineContext TraceRayInlineContext = CreateTraceRayInlineContext();
+	TraceRayInlineContext.MaxIterations = 32;
+	
+	const FTraceRayInlineResult TraceResult = TraceRayInline(InTLAS, /*RayFlags=*/0x0, /*InstanceInclusionMask=*/~0x0, Ray.GetNativeDesc(), TraceRayInlineContext);
+	
+	if (TraceResult.IsHit())
+	{
+		float2 TexCoords = 0.0f;
+		float3 Normal;
+
+#if PLATFORM_SUPPORTS_INLINE_RAY_TRACING_TRIANGLE_ATTRIBUTES
+		const FTriangleBaseAttributes Attributes = LoadInlineRayTracingTriangleAttributes(
+			InRayTracingSceneMetadata,
+			TraceResult.InstanceContributionToHitGroupIndex,
+			RAY_TRACING_NUM_SHADER_SLOTS,
+			TraceResult.GeometryIndex,
+			TraceResult.PrimitiveIndex);
+
+		const FInstanceSceneData InstanceData = GetInstanceSceneData(TraceResult.InstanceID);
+		const FPrimitiveSceneData PrimitiveData = GetPrimitiveData(InstanceData.PrimitiveId);
+
+		// First uint packed into InPrimitiveUVOffsets will be the offset into the TexCoords buffer.
+		const uint TexCoordsBufferOffset = InPrimitiveUVOffsets[PrimitiveData.PersistentPrimitiveIndex * 2 + 0];
+
+		// If the primitive has a valid tex coord offset, grab the texture coordinates.
+		if (TexCoordsBufferOffset != ~0u)
+		{
+			// Second uint packed into InPrimitiveUVOffsets is the number of UV channels on the primitive.
+			const uint NumUVChannels = InPrimitiveUVOffsets[PrimitiveData.PersistentPrimitiveIndex * 2 + 1];
+
+			if (InTexCoordsChannelIndex >= 0 && InTexCoordsChannelIndex < NumUVChannels)
+			{
+				const Buffer_float2 TexCoordsBuffer = BufferFromCollection_Buffer_float2(InPrimitiveUVCollection, TexCoordsBufferOffset);
+
+				// Get the UVs for the hit triangle's vertices. Account for the selected UV channel.
+				const float2 VertexUVs[] =
+				{
+					TexCoordsBuffer.Load(NumUVChannels * Attributes.Indices[0] + InTexCoordsChannelIndex),
+					TexCoordsBuffer.Load(NumUVChannels * Attributes.Indices[1] + InTexCoordsChannelIndex),
+					TexCoordsBuffer.Load(NumUVChannels * Attributes.Indices[2] + InTexCoordsChannelIndex)
+				};
+
+				const float3 Barycentrics = float3(1.0f - TraceResult.Barycentrics.x - TraceResult.Barycentrics.y, TraceResult.Barycentrics.x, TraceResult.Barycentrics.y);
+				TexCoords = VertexUVs[0] * Barycentrics.x + VertexUVs[1] * Barycentrics.y + VertexUVs[2] * Barycentrics.z;
+			}
+		}
+#endif // PLATFORM_SUPPORTS_INLINE_RAY_TRACING_TRIANGLE_ATTRIBUTES
+
+#if PLATFORM_SUPPORTS_INLINE_RAY_TRACING_TRIANGLE_NORMALS
+		Normal = TraceResult.WorldGeometryNormal;
+#elif PLATFORM_SUPPORTS_INLINE_RAY_TRACING_TRIANGLE_ATTRIBUTES
+		// Transform local-space geometry normal to world space using inverse transpose (matching VFXTraceRay.ush).
+		Normal = cross(Attributes.LocalPositions[2] - Attributes.LocalPositions[0], Attributes.LocalPositions[1] - Attributes.LocalPositions[0]);
+		Normal = normalize(mul(Normal, (float3x3)TraceResult.WorldToObject3x4));
+#else
+		Normal = PCG_DEFAULT_NORMAL;
+#endif
+
+		WriteOutput(InRayIndex, Ray.Origin + Ray.Direction * TraceResult.HitT, TexCoords, Normal, /*Status=*/0u);
+	}
+	else
+	{
+		// If the ray missed, mark it as culled.
+		WriteOutput(InRayIndex, Ray.Origin + Ray.Direction * Ray.TMax, /*TexCoords=*/0.0f, PCG_DEFAULT_NORMAL, PCG_RAY_CULLED);
+	}
+}
+
+[numthreads(32, 1, 1)]
+void PCGRayTraceMainCS(uint DispatchThreadId : SV_DispatchThreadID)
+{
+	const int RayIndex = DispatchThreadId;
+
+	if (RayIndex >= InNumRays)
+	{
+		return;
+	}
+
+	float3 Origin, Direction;
+	float Length;
+	int Status;
+
+	ReadInput(RayIndex, Origin, Direction, Length, Status);
+
+	// If the ray was already culled, do nothing.
+	if (Status == PCG_RAY_CULLED)
+	{
+		WriteOutput(RayIndex, Origin, /*TexCoords=*/0.0f, PCG_DEFAULT_NORMAL, Status);
+		return;
+	}
+
+	PCGRayTraceCommon(RayIndex, Origin, Direction, Length);
+}
+
+#undef PCG_DEFAULT_NORMAL
diff --git a/Engine/Plugins/PCG/Source/PCG/Private/Compute/BuiltInKernels/PCGPostRayTraceKernel.cpp b/Engine/Plugins/PCG/Source/PCG/Private/Compute/BuiltInKernels/PCGPostRayTraceKernel.cpp
new file mode 100644
--- /dev/null
+++ b/Engine/Plugins/PCG/Source/PCG/Private/Compute/BuiltInKernels/PCGPostRayTraceKernel.cpp
@@ -0,0 +1,98 @@
+// Copyright Epic Games, Inc. All Rights Reserved.
+
+#include "Compute/BuiltInKernels/PCGPostRayTraceKernel.h"
+
+#include "PCGContext.h"
+#include "Compute/Data/PCGRawBufferData.h"
+#include "Compute/DataInterfaces/PCGPostRayTraceDataInterface.h"
+#include "Compute/DataInterfaces/PCGRawBufferDataInterface.h"
+#include "Compute/DataInterfaces/Elements/PCGWorldRaycastDataInterface.h"
+#include "Compute/Elements/PCGComputeGraphElement.h"
+#include "Graph/PCGGPUGraphCompilationContext.h"
+#include "Elements/PCGWorldRaycast.h"
+#include "Helpers/PCGWorldQueryHelpers.h"
+
+#include "ShaderCompilerCore.h"
+
+#define LOCTEXT_NAMESPACE "PCGPostRayTraceKernel"
+
+TSharedPtr<const FPCGDataCollectionDesc> UPCGPostRayTraceKernel::ComputeOutputBindingDataDesc(FName InOutputPinLabel, UPCGDataBinding* InBinding) const
+{
+	check(InBinding);
+
+	// Code assumes single output pin.
+	if (!ensure(InOutputPinLabel == PCGPinConstants::DefaultOutputLabel))
+	{
+		return nullptr;
+	}
+
+	const FPCGKernelPin InputKernelPin(GetKernelIndex(), PCGPinConstants::DefaultInputLabel, /*bIsInput=*/true);
+	const TSharedPtr<const FPCGDataCollectionDesc> InputDataDesc = InBinding->ComputeKernelPinDataDesc(InputKernelPin);
+
+	if (!ensure(InputDataDesc))
+	{
+		return nullptr;
+	}
+
+	TSharedPtr<FPCGDataCollectionDesc> OutDataDesc = FPCGDataCollectionDesc::MakeSharedFrom(InputDataDesc);
+
+	const UPCGWorldRaycastElementSettings* Settings = CastChecked<UPCGWorldRaycastElementSettings>(GetSettings());
+
+	if (Settings->bGetTextureCoordinates)
+	{
+		OutDataDesc->AddAttributeToAllData(FPCGKernelAttributeKey(PCGWorldQueryConstants::UVCoordAttribute, EPCGKernelAttributeType::Float2), InBinding);
+	}
+
+	if (Settings->bGetNormals)
+	{
+		OutDataDesc->AddAttributeToAllData(FPCGKernelAttributeKey(PCGWorldQueryConstants::ImpactNormalAttribute, EPCGKernelAttributeType::Float3), InBinding);
+	}
+
+	for (FPCGDataDesc& DataDesc : OutDataDesc->GetDataDescriptionsMutable())
+	{
+		DataDesc.AllocateProperties(EPCGPointNativeProperties::Transform);
+	}
+
+	return OutDataDesc;
+}
+
+int UPCGPostRayTraceKernel::ComputeThreadCount(const UPCGDataBinding* InBinding) const
+{
+	const TSharedPtr<const FPCGDataCollectionDesc> OutputPinDesc = InBinding->GetCachedKernelPinDataDesc(this, PCGPinConstants::DefaultOutputLabel, /*bIsInput=*/false);
+	return ensure(OutputPinDesc) ? OutputPinDesc->ComputeTotalElementCount() : 0;
+}
+
+#if WITH_EDITOR
+void UPCGPostRayTraceKernel::CreateAdditionalInputDataInterfaces(FPCGGPUCompilationContext& InOutContext, UObject* InObjectOuter, TArray<TObjectPtr<UComputeDataInterface>>& OutDataInterfaces) const
+{
+	Super::CreateAdditionalInputDataInterfaces(InOutContext, InObjectOuter, OutDataInterfaces);
+
+	TObjectPtr<UPCGPostRayTraceDataInterface> PostRayTraceDI = InOutContext.NewObject_AnyThread<UPCGPostRayTraceDataInterface>(InObjectOuter);
+	PostRayTraceDI->SetProducerKernel(this);
+	OutDataInterfaces.Add(PostRayTraceDI);
+
+	TObjectPtr<UPCGWorldRaycastDataInterface> WorldRaycastDI = InOutContext.NewObject_AnyThread<UPCGWorldRaycastDataInterface>(InObjectOuter);
+	WorldRaycastDI->SetProducerKernel(this);
+
+	OutDataInterfaces.Add(WorldRaycastDI);
+}
+#endif
+
+void UPCGPostRayTraceKernel::GetKernelAttributeKeys(TArray<FPCGKernelAttributeKey>& OutKeys) const
+{
+	OutKeys.Add(FPCGKernelAttributeKey(PCGWorldQueryConstants::UVCoordAttribute, EPCGKernelAttributeType::Float2));
+	OutKeys.Add(FPCGKernelAttributeKey(PCGWorldQueryConstants::ImpactNormalAttribute, EPCGKernelAttributeType::Float3));
+}
+
+void UPCGPostRayTraceKernel::GetInputPins(TArray<FPCGPinProperties>& OutPins) const
+{
+	OutPins.Emplace(PCGPinConstants::DefaultInputLabel, EPCGDataType::Point);
+	OutPins.Emplace(PCGPostRayTraceKernel::RayTraceResultsLabel, FPCGDataTypeInfoRawBuffer::AsId());
+}
+
+void UPCGPostRayTraceKernel::GetOutputPins(TArray<FPCGPinPropertiesGPU>& OutPins) const
+{
+	OutPins.Emplace(PCGPinConstants::DefaultOutputLabel, EPCGDataType::Point);
+}
+
+#undef LOCTEXT_NAMESPACE
diff --git a/Engine/Plugins/PCG/Source/PCG/Private/Compute/BuiltInKernels/PCGPostRayTraceKernel.h b/Engine/Plugins/PCG/Source/PCG/Private/Compute/BuiltInKernels/PCGPostRayTraceKernel.h
new file mode 100644
--- /dev/null
+++ b/Engine/Plugins/PCG/Source/PCG/Private/Compute/BuiltInKernels/PCGPostRayTraceKernel.h
@@ -0,0 +1,36 @@
+// Copyright Epic Games, Inc. All Rights Reserved.
+
+#pragma once
+
+#include "Compute/PCGComputeKernel.h"
+
+#include "Compute/PCGPinPropertiesGPU.h"
+
+#include "PCGPostRayTraceKernel.generated.h"
+
+namespace PCGPostRayTraceKernel
+{
+	const FName RayTraceResultsLabel = TEXT("RayTraceResults");
+}
+
+/** Post-process kernel that unpacks data after raytracing. Data is unpacked into points (one per ray). */
+UCLASS()
+class UPCGPostRayTraceKernel : public UPCGComputeKernel
+{
+	GENERATED_BODY()
+
+public:
+	//~Begin UPCGComputeKernel interface
+	virtual TSharedPtr<const FPCGDataCollectionDesc> ComputeOutputBindingDataDesc(FName InOutputPinLabel, UPCGDataBinding* InBinding) const override;
+	virtual int ComputeThreadCount(const UPCGDataBinding* Binding) const override;
+#if WITH_EDITOR
+	virtual const TCHAR* GetSourceFilePath() const override { return TEXT("/Plugin/PCG/Private/BuiltInKernels/PCGPostRayTrace.usf"); }
+	virtual FString GetEntryPoint() const override { return TEXT("PCGPostRayTraceCS"); }
+	void CreateAdditionalInputDataInterfaces(FPCGGPUCompilationContext& InOutContext, UObject* InObjectOuter, TArray<TObjectPtr<UComputeDataInterface>>& OutDataInterfaces) const override;
+#endif
+	virtual void GetKernelAttributeKeys(TArray<FPCGKernelAttributeKey>& OutKeys) const override;
+	virtual void GetInputPins(TArray<FPCGPinProperties>& OutPins) const override;
+	virtual void GetOutputPins(TArray<FPCGPinPropertiesGPU>& OutPins) const override;
+	bool RequiresPostTLASBuildExecutionGroup() const override { return true; }
+	//~End UPCGComputeKernel interface
+};
diff --git a/Engine/Plugins/PCG/Source/PCG/Private/Compute/BuiltInKernels/PCGPreRayTraceKernel.cpp b/Engine/Plugins/PCG/Source/PCG/Private/Compute/BuiltInKernels/PCGPreRayTraceKernel.cpp
new file mode 100644
--- /dev/null
+++ b/Engine/Plugins/PCG/Source/PCG/Private/Compute/BuiltInKernels/PCGPreRayTraceKernel.cpp
@@ -0,0 +1,322 @@
+// Copyright Epic Games, Inc. All Rights Reserved.
+
+#include "Compute/BuiltInKernels/PCGPreRayTraceKernel.h"
+
+#include "PCGContext.h"
+#include "PCGRayTrace.h"
+#include "Compute/Data/PCGRawBufferData.h"
+#include "Compute/DataInterfaces/PCGExportToRayTracingDataInterface.h"
+#include "Compute/DataInterfaces/Elements/PCGWorldRaycastDataInterface.h"
+#include "Compute/Elements/PCGComputeGraphElement.h"
+#include "Elements/PCGWorldRaycast.h"
+#include "Graph/PCGGPUGraphCompilationContext.h"
+
+#include "RHIGlobals.h"
+#include "RendererUtils.h"
+#include "ShaderCompilerCore.h"
+
+#define LOCTEXT_NAMESPACE "PCGPreRayTraceKernel"
+
+bool UPCGPreRayTraceKernel::IsKernelDataValid(FPCGContext* InContext) const
+{
+	TRACE_CPUPROFILER_EVENT_SCOPE(UPCGPreRayTraceKernel::IsKernelDataValid);
+
+	if (!Super::IsKernelDataValid(InContext))
+	{
+		return false;
+	}
+
+	const UPCGWorldRaycastElementSettings* Settings = CastChecked<UPCGWorldRaycastElementSettings>(GetSettings());
+
+	if (!IsRayTracingEnabled())
+	{
+		PCG_KERNEL_VALIDATION_ERR(InContext, Settings, LOCTEXT("RayTracingDisabled", "Ray tracing is not enabled."));
+		return false;
+	}
+
+	if (!GRHISupportsInlineRayTracing)
+	{
+		PCG_KERNEL_VALIDATION_ERR(InContext, Settings, LOCTEXT("InlineRayTracingUnsupported", "Inline ray tracing is not supported."));
+		return false;
+	}
+
+	if (!InContext || !InContext->IsComputeContext())
+	{
+		return true;
+	}
+
+	const FPCGComputeGraphContext* Context = static_cast<FPCGComputeGraphContext*>(InContext);
+	const UPCGDataBinding* DataBinding = Context->DataBinding.Get();
+	if (DataBinding)
+	{
+		const TSharedPtr<const FPCGDataCollectionDesc> InputDataDesc = DataBinding->GetCachedKernelPinDataDesc(this, PCGPinConstants::DefaultInputLabel, /*bIsInputPin=*/true);
+
+		if (!ensure(InputDataDesc))
+		{
+			return false;
+		}
+
+		// Validate the attribute exists and is the correct type.
+		auto ValidateAttribute = [InContext, Settings](const TSharedPtr<const FPCGDataCollectionDesc> InDataDesc, FName InAttributeName, EPCGKernelAttributeType InAttributeType)
+		{
+			FPCGKernelAttributeDesc AttributeDesc;
+			bool bConflictingTypesFound = false;
+			bool bPresentOnAllData = false;
+			InDataDesc->GetAttributeDesc(InAttributeName, AttributeDesc, bConflictingTypesFound, bPresentOnAllData);
+
+			if (!bPresentOnAllData)
+			{
+				PCG_KERNEL_VALIDATION_ERR(InContext, Settings, FText::Format(LOCTEXT("AttributeMissing", "Attribute '{0}' not found, this attribute must be present on all input data."), FText::FromName(InAttributeName)));
+				return false;
+			}
+
+			if (bConflictingTypesFound)
+			{
+				PCG_KERNEL_VALIDATION_ERR(InContext, Settings, FText::Format(LOCTEXT("AttributeHasMultipleTypes", "Attribute '{0}' found with multiple types in input data."), FText::FromName(InAttributeName)));
+				return false;
+			}
+
+			if (AttributeDesc.GetAttributeKey().GetType() != InAttributeType)
+			{
+				PCG_KERNEL_VALIDATION_ERR(InContext, Settings, FText::Format(
+					LOCTEXT("AttributeTypeInvalid", "Attribute '{0}', must be of type '{1}'."),
+					FText::FromName(InAttributeName),
+					StaticEnum<EPCGKernelAttributeType>()->GetDisplayNameTextByValue(static_cast<int64>(InAttributeType))));
+
+				return false;
+			}
+
+			return true;
+		};
+
+		if (!InputDataDesc->GetDataDescriptions().IsEmpty())
+		{
+			if (Settings->OriginInputAttribute.GetSelection() == EPCGAttributePropertySelection::Attribute && !ValidateAttribute(InputDataDesc, Settings->OriginInputAttribute.GetAttributeName(), EPCGKernelAttributeType::Float3))
+			{
+				return false;
+			}
+
+			if (Settings->bOverrideRayDirections && !ValidateAttribute(InputDataDesc, Settings->RayDirectionAttribute.GetAttributeName(), EPCGKernelAttributeType::Float3))
+			{
+				return false;
+			}
+
+			if (Settings->bOverrideRayLengths && !ValidateAttribute(InputDataDesc, Settings->RayLengthAttribute.GetAttributeName(), EPCGKernelAttributeType::Float))
+			{
+				return false;
+			}
+		}
+
+		if (Settings->RaycastMode == EPCGWorldRaycastMode::Segments)
+		{
+			const TSharedPtr<const FPCGDataCollectionDesc> EndPointsDataDesc = DataBinding->GetCachedKernelPinDataDesc(this, PCGWorldRaycastElement::Constants::Labels::EndPointsInput, /*bIsInputPin=*/true);
+
+			if (!ensure(EndPointsDataDesc))
+			{
+				return false;
+			}
+
+			const uint32 NumStartPoints = InputDataDesc->ComputeTotalElementCount();
+			const uint32 NumEndPoints = EndPointsDataDesc->ComputeTotalElementCount();
+
+			if (NumStartPoints != NumEndPoints)
+			{
+				PCG_KERNEL_VALIDATION_ERR(InContext, Settings, FText::Format(LOCTEXT("InputMismatch", "Mismatch between number of input points ({0}) and number of end points ({1}). Must be N:N for GPU execution."), NumStartPoints, NumEndPoints));
+				return false;
+			}
+
+			if (!EndPointsDataDesc->GetDataDescriptions().IsEmpty())
+			{
+				if (Settings->EndPointAttribute.GetSelection() == EPCGAttributePropertySelection::Attribute && !ValidateAttribute(EndPointsDataDesc, Settings->EndPointAttribute.GetAttributeName(), EPCGKernelAttributeType::Float3))
+				{
+					return false;
+				}
+			}
+		}
+	}
+
+	return true;
+}
+
+TSharedPtr<const FPCGDataCollectionDesc> UPCGPreRayTraceKernel::ComputeOutputBindingDataDesc(FName InOutputPinLabel, UPCGDataBinding* InBinding) const
+{
+	check(InBinding);
+
+	// Code assumes single output pin.
+	if (!ensure(InOutputPinLabel == PCGPinConstants::DefaultOutputLabel))
+	{
+		return nullptr;
+	}
+
+	const FPCGKernelPin InputKernelPin(GetKernelIndex(), PCGPinConstants::DefaultInputLabel, /*bIsInput=*/true);
+	const TSharedPtr<const FPCGDataCollectionDesc> InputDataDesc = InBinding->ComputeKernelPinDataDesc(InputKernelPin);
+
+	if (!ensure(InputDataDesc))
+	{
+		return nullptr;
+	}
+
+	const uint32 PointElementCount = InputDataDesc->ComputeTotalElementCount();
+
+	TSharedPtr<FPCGDataCollectionDesc> OutDataDesc = FPCGDataCollectionDesc::MakeShared();
+	OutDataDesc->GetDataDescriptionsMutable().Add(FPCGDataDesc(FPCGDataTypeInfoRawBuffer::AsId(), PointElementCount * PCGRaytraceConstants::RAY_TRACE_PACKED_BUFFER_STRIDE_UINTS));
+
+	return OutDataDesc;
+}
+
+int UPCGPreRayTraceKernel::ComputeThreadCount(const UPCGDataBinding* InBinding) const
+{
+	const TSharedPtr<const FPCGDataCollectionDesc> PinDesc = InBinding->GetCachedKernelPinDataDesc(this, PCGPinConstants::DefaultInputLabel, /*bIsInput=*/true);
+	return (ensure(PinDesc)) ? PinDesc->ComputeTotalElementCount() : 0;
+}
+
+#if WITH_EDITOR
+FString UPCGPreRayTraceKernel::GetCookedSource(FPCGGPUCompilationContext& InOutContext) const
+{
+	FString SourceFile;
+	ensure(LoadShaderSourceFile(GetSourceFilePath(), EShaderPlatform::SP_PCD3D_SM5, &SourceFile, nullptr));
+	
+	TMap<FString, FStringFormatArg> TemplateArgs =
+	{
+		{ TEXT("EndPointsEnabled"), bIncludeEndPointsInputPin ? 1 : 0 },
+	};
+
+	SourceFile = FString::Format(*SourceFile, TemplateArgs);
+
+	return SourceFile;
+}
+
+void UPCGPreRayTraceKernel::CreateAdditionalInputDataInterfaces(FPCGGPUCompilationContext& InOutContext, UObject* InObjectOuter, TArray<TObjectPtr<UComputeDataInterface>>& OutDataInterfaces) const
+{
+	Super::CreateAdditionalInputDataInterfaces(InOutContext, InObjectOuter, OutDataInterfaces);
+
+	TObjectPtr<UPCGWorldRaycastDataInterface> WorldRaycastDI = InOutContext.NewObject_AnyThread<UPCGWorldRaycastDataInterface>(InObjectOuter);
+	WorldRaycastDI->SetProducerKernel(this);
+
+	OutDataInterfaces.Add(WorldRaycastDI);
+}
+
+UPCGComputeDataInterface* UPCGPreRayTraceKernel::CreateOutputPinDataInterface(const PCGComputeHelpers::FCreateDataInterfaceParams& InParams) const
+{
+	check(InParams.Context);
+	check(InParams.PinProperties);
+	check(InParams.ObjectOuter);
+
+	if (InParams.PinProperties->Label == PCGPinConstants::DefaultOutputLabel && ensure(InParams.PinProperties->AllowedTypes == FPCGDataTypeInfoRawBuffer::AsId()))
+	{
+		// Use a special data interface for the raw buffer output - one that dispatches a ray trace pass after the buffer is populated.
+		UPCGExportToRayTracingDataInterface* ExportToRTDI = InParams.Context->NewObject_AnyThread<UPCGExportToRayTracingDataInterface>(InParams.ObjectOuter);
+		ExportToRTDI->SetInputPointsPinLabel(PCGPinConstants::DefaultInputLabel);
+		
+		return ExportToRTDI;
+	}
+	else
+	{
+		// Fallback to default DI creation logic.
+		return Super::CreateOutputPinDataInterface(InParams);
+	}
+}
+#endif
+
+void UPCGPreRayTraceKernel::GetInputPins(TArray<FPCGPinProperties>& OutPins) const
+{
+	OutPins.Emplace(PCGPinConstants::DefaultInputLabel, EPCGDataType::Point);
+	OutPins.Emplace(PCGWorldRaycastElement::Constants::Labels::EndPointsInput, EPCGDataType::Point);
+}
+
+void UPCGPreRayTraceKernel::GetOutputPins(TArray<FPCGPinPropertiesGPU>& OutPins) const
+{
+	OutPins.Emplace(PCGPinConstants::DefaultOutputLabel, FPCGDataTypeInfoRawBuffer::AsId());
+}
+
+#if WITH_EDITOR
+bool UPCGPreRayTraceKernel::PerformStaticValidation()
+{
+	TRACE_CPUPROFILER_EVENT_SCOPE(UPCGPreRayTraceKernel::PerformStaticValidation);
+
+	if (!Super::PerformStaticValidation())
+	{
+		return false;
+	}
+
+	const UPCGWorldRaycastElementSettings* Settings = CastChecked<UPCGWorldRaycastElementSettings>(GetSettings());
+
+#if PCG_KERNEL_LOGGING_ENABLED
+	auto LogUnsupportedAttribute = [this](const FPCGAttributePropertyInputSelector& InputSelector)
+	{
+		AddStaticLogEntry(FText::Format(LOCTEXT("InvalidAttribute", "Invalid attribute '{0}', unsupported for GPU execution. Must be a basic attribute (not a property and no extra names)."), InputSelector.GetDisplayText()), EPCGKernelLogVerbosity::Error);
+	};
+
+	auto LogUnsupportedProperty = [this](const FPCGAttributePropertyInputSelector& InputSelector)
+	{
+		AddStaticLogEntry(FText::Format(LOCTEXT("InvalidProperty", "Invalid property '{0}', unsupported for GPU execution. Must be '$Position'."), InputSelector.GetDisplayText()), EPCGKernelLogVerbosity::Error);
+	};
+#endif
+
+	if (Settings->OriginInputAttribute.GetSelection() == EPCGAttributePropertySelection::Property && Settings->OriginInputAttribute.GetPointProperty() != EPCGPointProperties::Position)
+	{
+#if PCG_KERNEL_LOGGING_ENABLED
+		LogUnsupportedProperty(Settings->OriginInputAttribute);
+#endif
+		return false;
+	}
+
+	if (Settings->bOverrideRayDirections && !Settings->RayDirectionAttribute.IsBasicAttribute())
+	{
+#if PCG_KERNEL_LOGGING_ENABLED
+		LogUnsupportedAttribute(Settings->RayDirectionAttribute);
+#endif
+		return false;
+	}
+
+	if (Settings->bOverrideRayLengths && !Settings->RayLengthAttribute.IsBasicAttribute())
+	{
+#if PCG_KERNEL_LOGGING_ENABLED
+		LogUnsupportedAttribute(Settings->RayLengthAttribute);
+#endif
+		return false;
+	}
+
+	if (Settings->RaycastMode == EPCGWorldRaycastMode::Segments)
+	{
+		const UPCGNode* Node = Settings->GetTypedOuter<UPCGNode>();
+		check(Node);
+
+		if (!Node->IsInputPinConnected(PCGWorldRaycastElement::Constants::Labels::EndPointsInput))
+		{
+#if PCG_KERNEL_LOGGING_ENABLED
+			AddStaticLogEntry(FText::Format(LOCTEXT("MissingOutputPoints", "Missing connection on input pin '{0}'"), FText::FromName(PCGWorldRaycastElement::Constants::Labels::EndPointsInput)), EPCGKernelLogVerbosity::Error);
+#endif
+			return false;
+		}
+
+		if (Settings->EndPointAttribute.GetSelection() == EPCGAttributePropertySelection::Property && Settings->EndPointAttribute.GetPointProperty() != EPCGPointProperties::Position)
+		{
+	#if PCG_KERNEL_LOGGING_ENABLED
+			LogUnsupportedProperty(Settings->EndPointAttribute);
+	#endif
+			return false;
+		}
+	}
+
+	if (Settings->CollisionShape.ShapeType != EPCGCollisionShapeType::Line)
+	{
+#if PCG_KERNEL_LOGGING_ENABLED
+		AddStaticLogEntry(FText::Format(LOCTEXT("InvalidCollisionShape", "Invalid collision shape '{0}', unsupported for GPU execution. Must be '{1}'."),
+			StaticEnum<EPCGCollisionShapeType>()->GetDisplayNameTextByValue((int64)Settings->CollisionShape.ShapeType),
+			StaticEnum<EPCGCollisionShapeType>()->GetDisplayNameTextByValue((int64)EPCGCollisionShapeType::Line)),
+			EPCGKernelLogVerbosity::Error);
+#endif
+		return false;
+	}
+
+	return true;
+}
+
+void UPCGPreRayTraceKernel::SetIncludeEndPointsInputPin(bool bEnabled)
+{
+	bIncludeEndPointsInputPin = bEnabled;
+}
+#endif
+
+#undef LOCTEXT_NAMESPACE
diff --git a/Engine/Plugins/PCG/Source/PCG/Private/Compute/BuiltInKernels/PCGPreRayTraceKernel.h b/Engine/Plugins/PCG/Source/PCG/Private/Compute/BuiltInKernels/PCGPreRayTraceKernel.h
new file mode 100644
--- /dev/null
+++ b/Engine/Plugins/PCG/Source/PCG/Private/Compute/BuiltInKernels/PCGPreRayTraceKernel.h
@@ -0,0 +1,46 @@
+// Copyright Epic Games, Inc. All Rights Reserved.
+
+#pragma once
+
+#include "Compute/PCGComputeKernel.h"
+
+#include "Compute/PCGPinPropertiesGPU.h"
+
+#include "PCGPreRayTraceKernel.generated.h"
+
+/** Pre-process kernel that packs data for raytracing, producing one ray per input point. Dispatches an inlined raytrace in PostSubmit. */
+UCLASS()
+class UPCGPreRayTraceKernel : public UPCGComputeKernel
+{
+	GENERATED_BODY()
+
+public:
+	//~Begin UPCGComputeKernel interface
+	virtual bool IsKernelDataValid(FPCGContext* InContext) const override;
+	virtual TSharedPtr<const FPCGDataCollectionDesc> ComputeOutputBindingDataDesc(FName InOutputPinLabel, UPCGDataBinding* InBinding) const override;
+	virtual int ComputeThreadCount(const UPCGDataBinding* Binding) const override;
+#if WITH_EDITOR
+	virtual FString GetCookedSource(FPCGGPUCompilationContext& InOutContext) const override;
+	virtual const TCHAR* GetSourceFilePath() const override { return TEXT("/Plugin/PCG/Private/BuiltInKernels/PCGPreRayTrace.usf"); }
+	virtual FString GetEntryPoint() const override { return TEXT("PCGPreRayTraceCS"); }
+	virtual void CreateAdditionalInputDataInterfaces(FPCGGPUCompilationContext& InOutContext, UObject* InObjectOuter, TArray<TObjectPtr<UComputeDataInterface>>& OutDataInterfaces) const override;
+	virtual UPCGComputeDataInterface* CreateOutputPinDataInterface(const PCGComputeHelpers::FCreateDataInterfaceParams& InParams) const override;
+#endif
+	virtual void GetInputPins(TArray<FPCGPinProperties>& OutPins) const override;
+	virtual void GetOutputPins(TArray<FPCGPinPropertiesGPU>& OutPins) const override;
+	bool RequiresPostTLASBuildExecutionGroup() const override { return true; }
+
+protected:
+#if WITH_EDITOR
+	virtual bool PerformStaticValidation() override;
+#endif
+	//~End UPCGComputeKernel interface
+
+#if WITH_EDITOR
+public:
+	void SetIncludeEndPointsInputPin(bool bEnabled);
+
+private:
+	bool bIncludeEndPointsInputPin = false;
+#endif
+};
diff --git a/Engine/Plugins/PCG/Source/PCG/Private/Compute/DataInterfaces/Elements/PCGWorldRaycastDataInterface.cpp b/Engine/Plugins/PCG/Source/PCG/Private/Compute/DataInterfaces/Elements/PCGWorldRaycastDataInterface.cpp
new file mode 100644
--- /dev/null
+++ b/Engine/Plugins/PCG/Source/PCG/Private/Compute/DataInterfaces/Elements/PCGWorldRaycastDataInterface.cpp
@@ -0,0 +1,216 @@
+// Copyright Epic Games, Inc. All Rights Reserved.
+
+#include "PCGWorldRaycastDataInterface.h"
+
+#include "Compute/PCGComputeCommon.h"
+#include "Compute/PCGDataBinding.h"
+#include "Elements/PCGWorldRaycast.h"
+
+#include "GlobalRenderResources.h"
+#include "RenderGraphBuilder.h"
+#include "RenderGraphResources.h"
+#include "RHIResources.h"
+#include "ShaderCompilerCore.h"
+#include "ShaderCore.h"
+#include "ShaderParameterMetadataBuilder.h"
+#include "ComputeFramework/ShaderParamTypeDefinition.h"
+
+#include UE_INLINE_GENERATED_CPP_BY_NAME(PCGWorldRaycastDataInterface)
+
+void UPCGWorldRaycastDataInterface::GetSupportedInputs(TArray<FShaderFunctionDefinition>& OutFunctions) const
+{
+	OutFunctions.AddDefaulted_GetRef()
+		.SetName(TEXT("WorldRaycast_GetRaycastMode"))
+		.AddReturnType(EShaderFundamentalType::Int);
+
+	OutFunctions.AddDefaulted_GetRef()
+		.SetName(TEXT("WorldRaycast_GetRayDirection"))
+		.AddReturnType(EShaderFundamentalType::Float, 3);
+
+	OutFunctions.AddDefaulted_GetRef()
+		.SetName(TEXT("WorldRaycast_GetRayLength"))
+		.AddReturnType(EShaderFundamentalType::Float);
+
+	OutFunctions.AddDefaulted_GetRef()
+		.SetName(TEXT("WorldRaycast_GetUnbounded"))
+		.AddReturnType(EShaderFundamentalType::Int);
+
+	OutFunctions.AddDefaulted_GetRef()
+		.SetName(TEXT("WorldRaycast_GetKeepOriginalPointOnMiss"))
+		.AddReturnType(EShaderFundamentalType::Int);
+
+	OutFunctions.AddDefaulted_GetRef()
+		.SetName(TEXT("WorldRaycast_GetRayOriginAttributeId"))
+		.AddReturnType(EShaderFundamentalType::Int);
+
+	OutFunctions.AddDefaulted_GetRef()
+		.SetName(TEXT("WorldRaycast_GetRayEndAttributeId"))
+		.AddReturnType(EShaderFundamentalType::Int);
+
+	OutFunctions.AddDefaulted_GetRef()
+		.SetName(TEXT("WorldRaycast_GetRayDirectionAttributeId"))
+		.AddReturnType(EShaderFundamentalType::Int);
+
+	OutFunctions.AddDefaulted_GetRef()
+		.SetName(TEXT("WorldRaycast_GetRayLengthAttributeId"))
+		.AddReturnType(EShaderFundamentalType::Int);
+}
+
+BEGIN_SHADER_PARAMETER_STRUCT(FPCGWorldRaycastDataInterfaceParameters,)
+	SHADER_PARAMETER(int32, RaycastMode)
+	SHADER_PARAMETER(FVector3f, RayDirection)
+	SHADER_PARAMETER(float, RayLength)
+	SHADER_PARAMETER(int32, Unbounded)
+	SHADER_PARAMETER(int32, KeepOriginalPointOnMiss)
+	SHADER_PARAMETER(int32, RayOriginAttributeId)
+	SHADER_PARAMETER(int32, RayEndAttributeId)
+	SHADER_PARAMETER(int32, RayDirectionAttributeId)
+	SHADER_PARAMETER(int32, RayLengthAttributeId)
+END_SHADER_PARAMETER_STRUCT()
+
+void UPCGWorldRaycastDataInterface::GetShaderParameters(TCHAR const* UID, FShaderParametersMetadataBuilder& InOutBuilder, FShaderParametersMetadataAllocations& InOutAllocations) const
+{
+	InOutBuilder.AddNestedStruct<FPCGWorldRaycastDataInterfaceParameters>(UID);
+}
+
+void UPCGWorldRaycastDataInterface::GetHLSL(FString& OutHLSL, FString const& InDataInterfaceName) const
+{
+	TMap<FString, FStringFormatArg> TemplateArgs =
+	{
+		{ TEXT("DataInterfaceName"), InDataInterfaceName },
+	};
+
+	OutHLSL += FString::Format(TEXT(
+		"int {DataInterfaceName}_RaycastMode;\n"
+		"float3 {DataInterfaceName}_RayDirection;\n"
+		"float {DataInterfaceName}_RayLength;\n"
+		"int {DataInterfaceName}_Unbounded;\n"
+		"int {DataInterfaceName}_KeepOriginalPointOnMiss;\n"
+		"int {DataInterfaceName}_RayOriginAttributeId;\n"
+		"int {DataInterfaceName}_RayEndAttributeId;\n"
+		"int {DataInterfaceName}_RayDirectionAttributeId;\n"
+		"int {DataInterfaceName}_RayLengthAttributeId;\n"
+		"\n"
+		"int WorldRaycast_GetRaycastMode_{DataInterfaceName}()\n"
+		"{\n"
+		"	return {DataInterfaceName}_RaycastMode;\n"
+		"}\n"
+		"\n"
+		"float3 WorldRaycast_GetRayDirection_{DataInterfaceName}()\n"
+		"{\n"
+		"	return {DataInterfaceName}_RayDirection;\n"
+		"}\n"
+		"\n"
+		"float WorldRaycast_GetRayLength_{DataInterfaceName}()\n"
+		"{\n"
+		"	return {DataInterfaceName}_RayLength;\n"
+		"}\n"
+		"\n"
+		"int WorldRaycast_GetUnbounded_{DataInterfaceName}()\n"
+		"{\n"
+		"	return {DataInterfaceName}_Unbounded;\n"
+		"}\n"
+		"\n"
+		"int WorldRaycast_GetKeepOriginalPointOnMiss_{DataInterfaceName}()\n"
+		"{\n"
+		"	return {DataInterfaceName}_KeepOriginalPointOnMiss;\n"
+		"}\n"
+		"\n"
+		"int WorldRaycast_GetRayOriginAttributeId_{DataInterfaceName}()\n"
+		"{\n"
+		"	return {DataInterfaceName}_RayOriginAttributeId;\n"
+		"}\n"
+		"\n"
+		"int WorldRaycast_GetRayEndAttributeId_{DataInterfaceName}()\n"
+		"{\n"
+		"	return {DataInterfaceName}_RayEndAttributeId;\n"
+		"}\n"
+		"\n"
+		"int WorldRaycast_GetRayDirectionAttributeId_{DataInterfaceName}()\n"
+		"{\n"
+		"	return {DataInterfaceName}_RayDirectionAttributeId;\n"
+		"}\n"
+		"\n"
+		"int WorldRaycast_GetRayLengthAttributeId_{DataInterfaceName}()\n"
+		"{\n"
+		"	return {DataInterfaceName}_RayLengthAttributeId;\n"
+		"}\n"
+	), TemplateArgs);
+}
+
+UComputeDataProvider* UPCGWorldRaycastDataInterface::CreateDataProvider() const
+{
+	return NewObject<UPCGWorldRaycastDataProvider>();
+}
+
+FComputeDataProviderRenderProxy* UPCGWorldRaycastDataProvider::GetRenderProxy()
+{
+	const UPCGWorldRaycastElementSettings* Settings = CastChecked<UPCGWorldRaycastElementSettings>(GetProducerKernel()->GetSettings());
+
+	FPCGWorldRaycastProviderProxy::FData ProxyData =
+	{
+		.RaycastMode = static_cast<int32>(Settings->RaycastMode),
+		.RayDirection = Settings->RayDirection,
+		.RayLength = static_cast<float>(Settings->RayLength),
+		.Unbounded = Settings->bUnbounded,
+		.KeepOriginalPointOnMiss = Settings->bKeepOriginalPointOnMiss,
+		.RayOriginAttributeId = RayOriginAttributeId,
+		.RayEndAttributeId = RayEndAttributeId,
+		.RayDirectionAttributeId = RayDirectionAttributeId,
+		.RayLengthAttributeId = RayLengthAttributeId
+	};
+
+	return new FPCGWorldRaycastProviderProxy(ProxyData);
+}
+
+void UPCGWorldRaycastDataProvider::Reset()
+{
+	RayOriginAttributeId = INDEX_NONE;
+	RayEndAttributeId = INDEX_NONE;
+	RayDirectionAttributeId = INDEX_NONE;
+	RayLengthAttributeId = INDEX_NONE;
+
+	Super::Reset();
+}
+
+bool UPCGWorldRaycastDataProvider::PrepareForExecute_GameThread(UPCGDataBinding* InBinding)
+{
+	TRACE_CPUPROFILER_EVENT_SCOPE(UPCGWorldRaycastDataProvider::PrepareForExecute_GameThread);
+	check(InBinding);
+
+	if (!Super::PrepareForExecute_GameThread(InBinding))
+	{
+		return false;
+	}
+
+	const UPCGWorldRaycastElementSettings* Settings = CastChecked<UPCGWorldRaycastElementSettings>(GetProducerKernel()->GetSettings());
+	RayOriginAttributeId = Settings->OriginInputAttribute.GetSelection() == EPCGAttributePropertySelection::Attribute ? InBinding->GetAttributeId(Settings->OriginInputAttribute.GetAttributeName(), EPCGKernelAttributeType::Float3) : INDEX_NONE;
+	RayEndAttributeId = Settings->EndPointAttribute.GetSelection() == EPCGAttributePropertySelection::Attribute ? InBinding->GetAttributeId(Settings->EndPointAttribute.GetAttributeName(), EPCGKernelAttributeType::Float3) : INDEX_NONE;
+	RayDirectionAttributeId = Settings->bOverrideRayDirections ? InBinding->GetAttributeId(Settings->RayDirectionAttribute.GetAttributeName(), EPCGKernelAttributeType::Float3) : INDEX_NONE;
+	RayLengthAttributeId = Settings->bOverrideRayLengths ? InBinding->GetAttributeId(Settings->RayLengthAttribute.GetAttributeName(), EPCGKernelAttributeType::Float) : INDEX_NONE;
+
+	return true;
+}
+
+bool FPCGWorldRaycastProviderProxy::IsValid(FValidationData const& InValidationData) const
+{
+	return InValidationData.ParameterStructSize == sizeof(FParameters);
+}
+
+void FPCGWorldRaycastProviderProxy::GatherDispatchData(FDispatchData const& InDispatchData)
+{
+	const TStridedView<FParameters> ParameterArray = MakeStridedParameterView<FParameters>(InDispatchData);
+	for (int32 InvocationIndex = 0; InvocationIndex < ParameterArray.Num(); ++InvocationIndex)
+	{
+		FParameters& Parameters = ParameterArray[InvocationIndex];
+		Parameters.RaycastMode = Data.RaycastMode;
+		Parameters.RayDirection = static_cast<FVector3f>(Data.RayDirection);
+		Parameters.RayLength = Data.RayLength;
+		Parameters.Unbounded = Data.Unbounded;
+		Parameters.KeepOriginalPointOnMiss = Data.KeepOriginalPointOnMiss;
+		Parameters.RayOriginAttributeId = Data.RayOriginAttributeId;
+		Parameters.RayEndAttributeId = Data.RayEndAttributeId;
+		Parameters.RayDirectionAttributeId = Data.RayDirectionAttributeId;
+		Parameters.RayLengthAttributeId = Data.RayLengthAttributeId;
+	}
+}
diff --git a/Engine/Plugins/PCG/Source/PCG/Private/Compute/DataInterfaces/Elements/PCGWorldRaycastDataInterface.h b/Engine/Plugins/PCG/Source/PCG/Private/Compute/DataInterfaces/Elements/PCGWorldRaycastDataInterface.h
new file mode 100644
--- /dev/null
+++ b/Engine/Plugins/PCG/Source/PCG/Private/Compute/DataInterfaces/Elements/PCGWorldRaycastDataInterface.h
@@ -0,0 +1,78 @@
+// Copyright Epic Games, Inc. All Rights Reserved.
+
+#pragma once
+
+#include "Compute/DataInterfaces/PCGComputeDataInterface.h"
+
+#include "PCGWorldRaycastDataInterface.generated.h"
+
+class FPCGWorldRaycastDataInterfaceParameters;
+
+/** Data Interface to marshal World Raycast settings to the GPU. */
+UCLASS(ClassGroup = (Procedural))
+class UPCGWorldRaycastDataInterface : public UPCGComputeDataInterface
+{
+	GENERATED_BODY()
+
+public:
+	//~ Begin UComputeDataInterface Interface
+	TCHAR const* GetClassName() const override { return TEXT("PCGWorldRaycast"); }
+	void GetSupportedInputs(TArray<FShaderFunctionDefinition>& OutFunctions) const override;
+	void GetShaderParameters(TCHAR const* UID, FShaderParametersMetadataBuilder& InOutBuilder, FShaderParametersMetadataAllocations& InOutAllocations) const override;
+	void GetHLSL(FString& OutHLSL, FString const& InDataInterfaceName) const override;
+	UComputeDataProvider* CreateDataProvider() const override;
+	//~ End UComputeDataInterface Interface
+};
+
+UCLASS()
+class UPCGWorldRaycastDataProvider : public UPCGComputeDataProvider
+{
+	GENERATED_BODY()
+
+public:
+	//~ Begin UComputeDataProvider Interface
+	FComputeDataProviderRenderProxy* GetRenderProxy() override;
+	//~ End UComputeDataProvider Interface
+
+	//~ Begin UPCGComputeDataProvider Interface
+	virtual void Reset() override;
+	virtual bool PrepareForExecute_GameThread(UPCGDataBinding* InBinding) override;
+	//~ End UPCGComputeDataProvider Interface
+
+public:
+	int32 RayOriginAttributeId = INDEX_NONE;
+	int32 RayEndAttributeId = INDEX_NONE;
+	int32 RayDirectionAttributeId = INDEX_NONE;
+	int32 RayLengthAttributeId = INDEX_NONE;
+};
+
+class FPCGWorldRaycastProviderProxy : public FComputeDataProviderRenderProxy
+{
+public:
+	struct FData
+	{
+		int32 RaycastMode = 0;
+		FVector RayDirection = FVector::ZeroVector;
+		float RayLength = 0.0f;
+		int32 Unbounded = 0;
+		int32 KeepOriginalPointOnMiss = 0;
+		int32 RayOriginAttributeId = 0;
+		int32 RayEndAttributeId = 0;
+		int32 RayDirectionAttributeId = 0;
+		int32 RayLengthAttributeId = 0;
+	};
+
+	FPCGWorldRaycastProviderProxy(const FData& InData)
+		: Data(InData)
+	{}
+
+	//~ Begin FComputeDataProviderRenderProxy Interface
+	bool IsValid(FValidationData const& InValidationData) const override;
+	void GatherDispatchData(FDispatchData const& InDispatchData) override;
+	//~ End FComputeDataProviderRenderProxy Interface
+
+protected:
+	using FParameters = FPCGWorldRaycastDataInterfaceParameters;
+
+	FData Data;
+};
diff --git a/Engine/Plugins/PCG/Source/PCG/Private/Compute/DataInterfaces/PCGExportToRayTracingDataInterface.cpp b/Engine/Plugins/PCG/Source/PCG/Private/Compute/DataInterfaces/PCGExportToRayTracingDataInterface.cpp
new file mode 100644
--- /dev/null
+++ b/Engine/Plugins/PCG/Source/PCG/Private/Compute/DataInterfaces/PCGExportToRayTracingDataInterface.cpp
@@ -0,0 +1,104 @@
+// Copyright Epic Games, Inc. All Rights Reserved.
+
+#include "PCGExportToRayTracingDataInterface.h"
+
+#include "PCGContext.h"
+#include "PCGModule.h"
+#include "PCGRayTrace.h"
+#include "Compute/PCGComputeKernel.h"
+#include "Compute/PCGDataBinding.h"
+#include "Elements/PCGWorldRaycast.h"
+
+#include "ComputeWorkerInterface.h"
+#include "RenderGraphBuilder.h"
+#include "ShaderParameterMetadataBuilder.h"
+#include "ShaderParameterStruct.h"
+#include "ComputeFramework/ComputeKernelPermutationSet.h"
+#include "ComputeFramework/ComputeKernelPermutationVector.h"
+#include "ComputeFramework/ShaderParamTypeDefinition.h"
+
+#include UE_INLINE_GENERATED_CPP_BY_NAME(PCGExportToRayTracingDataInterface)
+
+void UPCGExportToRayTracingDataInterface::GetDefines(FComputeKernelDefinitionSet& OutDefinitionSet) const
+{
+	Super::GetDefines(OutDefinitionSet);
+
+	OutDefinitionSet.Defines.Add(FComputeKernelDefinition(TEXT("PCG_RAY_CULLED"), FString::FromInt(PCGRaytraceConstants::RAY_CULLED)));
+	OutDefinitionSet.Defines.Add(FComputeKernelDefinition(TEXT("PCG_RAY_TRACE_PACKED_BUFFER_STRIDE_UINTS"), FString::FromInt(PCGRaytraceConstants::RAY_TRACE_PACKED_BUFFER_STRIDE_UINTS)));
+}
+
+UComputeDataProvider* UPCGExportToRayTracingDataInterface::CreateDataProvider() const
+{
+	return NewObject<UPCGExportToRayTracingDataProvider>();
+}
+
+FComputeDataProviderRenderProxy* UPCGExportToRayTracingDataProvider::GetRenderProxy()
+{
+	TRACE_CPUPROFILER_EVENT_SCOPE(UPCGExportToRayTracingDataProvider::GetRenderProxy);
+
+	FPCGRawBufferDataProviderProxy::FParams RawBufferParams;
+	RawBufferParams.SizeBytes = SizeBytes;
+
+	FPCGExportToRayTracingDataProviderProxy::FParams Params;
+	Params.NumRays = NumRays;
+	Params.TexCoordsChannelIndex = TexCoordsChannelIndex;
+
+	return new FPCGExportToRayTracingDataProviderProxy(RawBufferParams, Params);
+}
+
+void UPCGExportToRayTracingDataProvider::Initialize(const UComputeDataInterface* InDataInterface, UObject* InBinding, uint64 InInputMask, uint64 InOutputMask)
+{
+	Super::Initialize(InDataInterface, InBinding, InInputMask, InOutputMask);
+
+	const UPCGExportToRayTracingDataInterface* DataInterface = CastChecked<UPCGExportToRayTracingDataInterface>(InDataInterface);
+	InputPointsPinLabel = DataInterface->GetInputPointsPinLabel();
+}
+
+void UPCGExportToRayTracingDataProvider::Reset()
+{
+	InputPointsPinLabel = NAME_None;
+	NumRays = INDEX_NONE;
+	TexCoordsChannelIndex = INDEX_NONE;
+
+	Super::Reset();
+}
+
+bool UPCGExportToRayTracingDataProvider::PrepareForExecute_GameThread(UPCGDataBinding* InBinding)
+{
+	check(InBinding);
+
+	if (!Super::PrepareForExecute_GameThread(InBinding))
+	{
+		return false;
+	}
+
+	const TSharedPtr<const FPCGDataCollectionDesc> Desc = InBinding->GetCachedKernelPinDataDesc(GetProducerKernel(), InputPointsPinLabel, /*bInIsInput=*/true);
+
+	if (!ensure(Desc))
+	{
+		return true;
+	}
+
+	NumRays = Desc->ComputeTotalElementCount();
+
+	const UPCGWorldRaycastElementSettings* Settings = CastChecked<UPCGWorldRaycastElementSettings>(GetProducerKernel()->GetSettings());
+	TexCoordsChannelIndex = Settings->bGetTextureCoordinates ? Settings->TextureCoordinatesChannelIndex : INDEX_NONE;
+
+	return true;
+}
+
+void FPCGExportToRayTracingDataProviderProxy::PostSubmit(FComputeContext& Context) const
+{
+#if RHI_RAYTRACING
+	FPCGRayTraceParams RayTraceParams;
+	RayTraceParams.Scene = Context.Scene;
+	RayTraceParams.View = Context.View;
+	RayTraceParams.NumRays = Params.NumRays;
+	RayTraceParams.bNeedsUVData = Params.TexCoordsChannelIndex != INDEX_NONE;
+	RayTraceParams.TexCoordsChannelIndex = Params.TexCoordsChannelIndex;
+	RayTraceParams.PackedDataUAV = DataUAV;
+	PCGRayTrace::RenderPCGRayTraceInline(Context.GraphBuilder, RayTraceParams);
+#else
+	UE_LOG(LogPCG, Error, TEXT("Failed to submit inline raytrace shader. Requires RHI ray tracing support."));
+#endif
+}
diff --git a/Engine/Plugins/PCG/Source/PCG/Private/Compute/DataInterfaces/PCGExportToRayTracingDataInterface.h b/Engine/Plugins/PCG/Source/PCG/Private/Compute/DataInterfaces/PCGExportToRayTracingDataInterface.h
new file mode 100644
--- /dev/null
+++ b/Engine/Plugins/PCG/Source/PCG/Private/Compute/DataInterfaces/PCGExportToRayTracingDataInterface.h
@@ -0,0 +1,74 @@
+// Copyright Epic Games, Inc. All Rights Reserved.
+
+#pragma once
+
+#include "Compute/DataInterfaces/PCGRawBufferDataInterface.h"
+
+#include "PCGExportToRayTracingDataInterface.generated.h"
+
+class FSceneInterface;
+
+/** Represents a buffer of data packed from points for tracing rays against the scene. Dispatches an inlined raytrace in PostSubmit. */
+UCLASS(ClassGroup = (Procedural))
+class UPCGExportToRayTracingDataInterface : public UPCGRawBufferDataInterface
+{
+	GENERATED_BODY()
+
+public:
+	//~ Begin UComputeDataInterface Interface
+	void GetDefines(FComputeKernelDefinitionSet& OutDefinitionSet) const override;
+	bool GetRequiresPostSubmitCall() const override { return true; }
+	virtual UComputeDataProvider* CreateDataProvider() const override;
+	//~ End UComputeDataInterface Interface
+
+	FName GetInputPointsPinLabel() const { return InputPointsPinLabel; }
+	void SetInputPointsPinLabel(FName InInputPointsPinLabel) { InputPointsPinLabel = InInputPointsPinLabel; }
+
+protected:
+	UPROPERTY()
+	FName InputPointsPinLabel;
+};
+
+UCLASS()
+class UPCGExportToRayTracingDataProvider : public UPCGRawBufferDataProvider
+{
+	GENERATED_BODY()
+
+public:
+	//~ Begin UComputeDataProvider Interface
+	virtual FComputeDataProviderRenderProxy* GetRenderProxy() override;
+	//~ End UComputeDataProvider Interface
+
+	//~ Begin UPCGComputeDataProvider Interface
+	virtual void Initialize(const UComputeDataInterface* InDataInterface, UObject* InBinding, uint64 InInputMask, uint64 InOutputMask) override;
+	virtual void Reset() override;
+	virtual bool PrepareForExecute_GameThread(UPCGDataBinding* InBinding) override;
+	//~ End UPCGComputeDataProvider Interface
+
+protected:
+	FName InputPointsPinLabel = NAME_None;
+	int32 NumRays = INDEX_NONE;
+	int32 TexCoordsChannelIndex = INDEX_NONE;
+};
+
+class FPCGExportToRayTracingDataProviderProxy : public FPCGRawBufferDataProviderProxy
+{
+public:
+	struct FParams
+	{
+		int32 NumRays = INDEX_NONE;
+		int32 TexCoordsChannelIndex = INDEX_NONE;
+	};
+
+	FPCGExportToRayTracingDataProviderProxy(const FPCGRawBufferDataProviderProxy::FParams& InRawBufferParams, const FParams& InParams)
+		: FPCGRawBufferDataProviderProxy(InRawBufferParams)
+		, Params(InParams)
+	{}
+
+	//~ Begin FComputeDataProviderRenderProxy Interface
+	void PostSubmit(FComputeContext& Context) const override;
+	//~ End FComputeDataProviderRenderProxy Interface
+
+protected:
+	FParams Params;
+};
diff --git a/Engine/Plugins/PCG/Source/PCG/Private/Compute/DataInterfaces/PCGPostRayTraceDataInterface.cpp b/Engine/Plugins/PCG/Source/PCG/Private/Compute/DataInterfaces/PCGPostRayTraceDataInterface.cpp
new file mode 100644
--- /dev/null
+++ b/Engine/Plugins/PCG/Source/PCG/Private/Compute/DataInterfaces/PCGPostRayTraceDataInterface.cpp
@@ -0,0 +1,121 @@
+// Copyright Epic Games, Inc. All Rights Reserved.
+
+#include "PCGPostRayTraceDataInterface.h"
+
+#include "Compute/PCGComputeCommon.h"
+#include "Compute/PCGDataBinding.h"
+#include "Elements/PCGWorldRaycast.h"
+#include "Helpers/PCGWorldQueryHelpers.h"
+
+#include "GlobalRenderResources.h"
+#include "RenderGraphBuilder.h"
+#include "RenderGraphResources.h"
+#include "RHIResources.h"
+#include "ShaderCompilerCore.h"
+#include "ShaderCore.h"
+#include "ShaderParameterMetadataBuilder.h"
+#include "ComputeFramework/ShaderParamTypeDefinition.h"
+
+#include UE_INLINE_GENERATED_CPP_BY_NAME(PCGPostRayTraceDataInterface)
+
+void UPCGPostRayTraceDataInterface::GetSupportedInputs(TArray<FShaderFunctionDefinition>& OutFunctions) const
+{
+	OutFunctions.AddDefaulted_GetRef()
+		.SetName(TEXT("PostRayTrace_GetTexCoordsAttributeId"))
+		.AddReturnType(EShaderFundamentalType::Int);
+
+	OutFunctions.AddDefaulted_GetRef()
+		.SetName(TEXT("PostRayTrace_GetNormalsAttributeId"))
+		.AddReturnType(EShaderFundamentalType::Int);
+}
+
+BEGIN_SHADER_PARAMETER_STRUCT(FPCGPostRayTraceDataInterfaceParameters,)
+	SHADER_PARAMETER(int32, TexCoordsAttributeId)
+	SHADER_PARAMETER(int32, NormalsAttributeId)
+END_SHADER_PARAMETER_STRUCT()
+
+void UPCGPostRayTraceDataInterface::GetShaderParameters(TCHAR const* UID, FShaderParametersMetadataBuilder& InOutBuilder, FShaderParametersMetadataAllocations& InOutAllocations) const
+{
+	InOutBuilder.AddNestedStruct<FPCGPostRayTraceDataInterfaceParameters>(UID);
+}
+
+void UPCGPostRayTraceDataInterface::GetHLSL(FString& OutHLSL, FString const& InDataInterfaceName) const
+{
+	TMap<FString, FStringFormatArg> TemplateArgs =
+	{
+		{ TEXT("DataInterfaceName"), InDataInterfaceName },
+	};
+
+	OutHLSL += FString::Format(TEXT(
+		"int {DataInterfaceName}_TexCoordsAttributeId;\n"
+		"int {DataInterfaceName}_NormalsAttributeId;\n"
+		"\n"
+		"int PostRayTrace_GetTexCoordsAttributeId_{DataInterfaceName}()\n"
+		"{\n"
+		"	return {DataInterfaceName}_TexCoordsAttributeId;\n"
+		"}\n"
+		"\n"
+		"int PostRayTrace_GetNormalsAttributeId_{DataInterfaceName}()\n"
+		"{\n"
+		"	return {DataInterfaceName}_NormalsAttributeId;\n"
+		"}\n"
+	), TemplateArgs);
+}
+
+UComputeDataProvider* UPCGPostRayTraceDataInterface::CreateDataProvider() const
+{
+	return NewObject<UPCGPostRayTraceDataProvider>();
+}
+
+FComputeDataProviderRenderProxy* UPCGPostRayTraceDataProvider::GetRenderProxy()
+{
+	FPCGPostRayTraceProviderProxy::FData ProxyData =
+	{
+		.TexCoordsAttributeId = TexCoordsAttributeId,
+		.NormalsAttributeId = NormalsAttributeId
+	};
+
+	return new FPCGPostRayTraceProviderProxy(ProxyData);
+}
+
+void UPCGPostRayTraceDataProvider::Reset()
+{
+	TexCoordsAttributeId = INDEX_NONE;
+	NormalsAttributeId = INDEX_NONE;
+
+	Super::Reset();
+}
+
+bool UPCGPostRayTraceDataProvider::PrepareForExecute_GameThread(UPCGDataBinding* InBinding)
+{
+	TRACE_CPUPROFILER_EVENT_SCOPE(UPCGPostRayTraceDataInterface::PrepareForExecute_GameThread);
+	check(InBinding);
+
+	if (!Super::PrepareForExecute_GameThread(InBinding))
+	{
+		return false;
+	}
+
+	const UPCGWorldRaycastElementSettings* Settings = CastChecked<UPCGWorldRaycastElementSettings>(GetProducerKernel()->GetSettings());
+	TexCoordsAttributeId = Settings->bGetTextureCoordinates ? InBinding->GetAttributeId(PCGWorldQueryConstants::UVCoordAttribute, EPCGKernelAttributeType::Float2) : INDEX_NONE;
+	NormalsAttributeId = Settings->bGetNormals ? InBinding->GetAttributeId(PCGWorldQueryConstants::ImpactNormalAttribute, EPCGKernelAttributeType::Float3) : INDEX_NONE;
+
+	return true;
+}
+
+bool FPCGPostRayTraceProviderProxy::IsValid(FValidationData const& InValidationData) const
+{
+	return InValidationData.ParameterStructSize == sizeof(FParameters);
+}
+
+void FPCGPostRayTraceProviderProxy::GatherDispatchData(FDispatchData const& InDispatchData)
+{
+	const TStridedView<FParameters> ParameterArray = MakeStridedParameterView<FParameters>(InDispatchData);
+
+	for (int32 InvocationIndex = 0; InvocationIndex < ParameterArray.Num(); ++InvocationIndex)
+	{
+		FParameters& Parameters = ParameterArray[InvocationIndex];
+		Parameters.TexCoordsAttributeId = Data.TexCoordsAttributeId;
+		Parameters.NormalsAttributeId = Data.NormalsAttributeId;
+	}
+}
diff --git a/Engine/Plugins/PCG/Source/PCG/Private/Compute/DataInterfaces/PCGPostRayTraceDataInterface.h b/Engine/Plugins/PCG/Source/PCG/Private/Compute/DataInterfaces/PCGPostRayTraceDataInterface.h
new file mode 100644
--- /dev/null
+++ b/Engine/Plugins/PCG/Source/PCG/Private/Compute/DataInterfaces/PCGPostRayTraceDataInterface.h
@@ -0,0 +1,69 @@
+// Copyright Epic Games, Inc. All Rights Reserved.
+
+#pragma once
+
+#include "Compute/DataInterfaces/PCGComputeDataInterface.h"
+
+#include "PCGPostRayTraceDataInterface.generated.h"
+
+class FPCGPostRayTraceDataInterfaceParameters;
+
+/** Data Interface to marshal post ray trace settings to the GPU. */
+UCLASS(ClassGroup = (Procedural))
+class UPCGPostRayTraceDataInterface : public UPCGComputeDataInterface
+{
+	GENERATED_BODY()
+
+public:
+	//~ Begin UComputeDataInterface Interface
+	TCHAR const* GetClassName() const override { return TEXT("PCGPostRayTrace"); }
+	void GetSupportedInputs(TArray<FShaderFunctionDefinition>& OutFunctions) const override;
+	void GetShaderParameters(TCHAR const* UID, FShaderParametersMetadataBuilder& InOutBuilder, FShaderParametersMetadataAllocations& InOutAllocations) const override;
+	void GetHLSL(FString& OutHLSL, FString const& InDataInterfaceName) const override;
+	UComputeDataProvider* CreateDataProvider() const override;
+	//~ End UComputeDataInterface Interface
+};
+
+UCLASS()
+class UPCGPostRayTraceDataProvider : public UPCGComputeDataProvider
+{
+	GENERATED_BODY()
+
+public:
+	//~ Begin UComputeDataProvider Interface
+	FComputeDataProviderRenderProxy* GetRenderProxy() override;
+	//~ End UComputeDataProvider Interface
+
+	//~ Begin UPCGComputeDataProvider Interface
+	virtual void Reset() override;
+	virtual bool PrepareForExecute_GameThread(UPCGDataBinding* InBinding) override;
+	//~ End UPCGComputeDataProvider Interface
+
+public:
+	int32 TexCoordsAttributeId = INDEX_NONE;
+	int32 NormalsAttributeId = INDEX_NONE;
+};
+
+class FPCGPostRayTraceProviderProxy : public FComputeDataProviderRenderProxy
+{
+public:
+	struct FData
+	{
+		int32 TexCoordsAttributeId = 0;
+		int32 NormalsAttributeId = 0;
+	};
+
+	FPCGPostRayTraceProviderProxy(const FData& InData)
+		: Data(InData)
+	{}
+
+	//~ Begin FComputeDataProviderRenderProxy Interface
+	bool IsValid(FValidationData const& InValidationData) const override;
+	void GatherDispatchData(FDispatchData const& InDispatchData) override;
+	//~ End FComputeDataProviderRenderProxy Interface
+
+protected:
+	using FParameters = FPCGPostRayTraceDataInterfaceParameters;
+
+	FData Data;
+};
diff --git a/Engine/Plugins/PCG/Source/PCG/Private/Compute/Elements/PCGComputeGraphElement.cpp b/Engine/Plugins/PCG/Source/PCG/Private/Compute/Elements/PCGComputeGraphElement.cpp
--- a/Engine/Plugins/PCG/Source/PCG/Private/Compute/Elements/PCGComputeGraphElement.cpp
+++ b/Engine/Plugins/PCG/Source/PCG/Private/Compute/Elements/PCGComputeGraphElement.cpp
@@ -533,7 +533,7 @@ bool FPCGComputeGraphElement::ExecuteInternal(FPCGContext* InContext) const
 		const bool bGraphEnqueued = Context->ComputeGraphInstance->EnqueueWork(
 			Context->ComputeGraph.Get(),
 			InContext->ExecutionSource->GetExecutionState().GetWorld()->Scene,
-			ComputeTaskExecutionGroup::EndOfFrameUpdate,
+			Context->ComputeGraph->GetRequiredExecutionGroup(),
 			FName(*InContext->ExecutionSource->GetExecutionState().GetDebugName()),
 			/*InFallbackDelegate=*/FSimpleDelegate::CreateLambda([ContextHandle = Context->GetOrCreateHandle()]()
 			{
diff --git a/Engine/Plugins/PCG/Source/PCG/Private/Compute/PCGComputeGraph.cpp b/Engine/Plugins/PCG/Source/PCG/Private/Compute/PCGComputeGraph.cpp
--- a/Engine/Plugins/PCG/Source/PCG/Private/Compute/PCGComputeGraph.cpp
+++ b/Engine/Plugins/PCG/Source/PCG/Private/Compute/PCGComputeGraph.cpp
@@ -15,6 +15,7 @@
 #include "Data/PCGTextureData.h"
 #include "Metadata/PCGMetadata.h"
 
+#include "ComputeWorkerInterface.h"
 #include "RenderGraphResources.h"
 #include "ComputeFramework/ComputeKernelCompileResult.h"
 
@@ -377,6 +378,22 @@ bool UPCGComputeGraph::ShouldRepeatDispatch() const
 }
 #endif // PCG_GPU_KERNEL_PROFILING
 
+FName UPCGComputeGraph::GetRequiredExecutionGroup() const
+{
+	for (const TObjectPtr<UComputeKernel>& Kernel : KernelInvocations)
+	{
+		if (const UPCGComputeKernel* PCGKernel = Cast<UPCGComputeKernel>(Kernel))
+		{
+			if (PCGKernel->RequiresPostTLASBuildExecutionGroup())
+			{
+				return ComputeTaskExecutionGroup::PostTLASBuild;
+			}
+		}
+	}
+
+	return ComputeTaskExecutionGroup::EndOfFrameUpdate;
+}
+
 TSharedPtr<const FPCGDataCollectionDesc> UPCGComputeGraph::ComputeExternalPinDesc(FName InVirtualLabel, UPCGDataBinding* InBinding) const
 {
 	check(InBinding);
diff --git a/Engine/Plugins/PCG/Source/PCG/Private/Compute/PCGComputeGraph.h b/Engine/Plugins/PCG/Source/PCG/Private/Compute/PCGComputeGraph.h
--- a/Engine/Plugins/PCG/Source/PCG/Private/Compute/PCGComputeGraph.h
+++ b/Engine/Plugins/PCG/Source/PCG/Private/Compute/PCGComputeGraph.h
@@ -79,6 +79,9 @@ public:
 	bool ShouldRepeatDispatch() const;
 #endif
 
+	/** Returns the execution group required by this compute graph based on its kernels' requirements. */
+	FName GetRequiredExecutionGroup() const;
+
 private:
 	/** Compute the data description on a pin external to the compute graph. */
 	TSharedPtr<const FPCGDataCollectionDesc> ComputeExternalPinDesc(FName InVirtualLabel, UPCGDataBinding* InBinding) const;
diff --git a/Engine/Plugins/PCG/Source/PCG/Private/Compute/PCGComputeKernel.cpp b/Engine/Plugins/PCG/Source/PCG/Private/Compute/PCGComputeKernel.cpp
--- a/Engine/Plugins/PCG/Source/PCG/Private/Compute/PCGComputeKernel.cpp
+++ b/Engine/Plugins/PCG/Source/PCG/Private/Compute/PCGComputeKernel.cpp
@@ -96,6 +96,11 @@ const UPCGSettings* UPCGComputeKernel::GetSettings() const
 	return ResolvedSettings;
 }
 
+bool UPCGComputeKernel::RequiresPostTLASBuildExecutionGroup() const
+{
+	return false;
+}
+
 void UPCGComputeKernel::AddReferencedObjects(UObject* InThis, FReferenceCollector& Collector)
 {
 	Super::AddReferencedObjects(InThis, Collector);
diff --git a/Engine/Plugins/PCG/Source/PCG/Private/Elements/PCGWorldRaycast.cpp b/Engine/Plugins/PCG/Source/PCG/Private/Elements/PCGWorldRaycast.cpp
--- a/Engine/Plugins/PCG/Source/PCG/Private/Elements/PCGWorldRaycast.cpp
+++ b/Engine/Plugins/PCG/Source/PCG/Private/Elements/PCGWorldRaycast.cpp
@@ -5,6 +5,10 @@
 #include "PCGComponent.h"
 #include "PCGContext.h"
 #include "PCGPoint.h"
+#include "PCGRayTrace.h"
+#include "Compute/PCGKernelHelpers.h"
+#include "Compute/BuiltInKernels/PCGPreRayTraceKernel.h"
+#include "Compute/BuiltInKernels/PCGPostRayTraceKernel.h"
 #include "Data/PCGBasePointData.h"
 #include "Helpers/PCGBlueprintHelpers.h"
 #include "Helpers/PCGHelpers.h"
@@ -108,8 +112,48 @@ EPCGChangeType UPCGWorldRaycastElementSettings::GetChangeTypeForProperty(FProper
 
 	return Result;
 }
+
+void UPCGWorldRaycastElementSettings::CreateKernels(FPCGGPUCompilationContext& InOutContext, UObject* InObjectOuter, TArray<UPCGComputeKernel*>& OutKernels, TArray<FPCGKernelEdge>& OutEdges) const
+{
+	PCGKernelHelpers::FCreateKernelParams CreateParams(InObjectOuter, this);
+
+	// Wiring is complex enough to just wire it out manually below (and node input label is not "In").
+	CreateParams.NodeInputPinsToWire.Empty();
+	CreateParams.NodeOutputPinsToWire.Empty();
+
+	UPCGPreRayTraceKernel* PreKernel = PCGKernelHelpers::CreateKernel<UPCGPreRayTraceKernel>(InOutContext, CreateParams, OutKernels, OutEdges);
+	UPCGPostRayTraceKernel* PostKernel = PCGKernelHelpers::CreateKernel<UPCGPostRayTraceKernel>(InOutContext, CreateParams, OutKernels, OutEdges);
+
+	// Node input to pre-kernel which exports data collection buffer to raw buffer.
+	OutEdges.Emplace(FPCGPinReference(PCGWorldRaycastElement::Constants::Labels::OriginsInput), FPCGPinReference(PreKernel, PCGPinConstants::DefaultInputLabel));
+
+	if (RaycastMode == EPCGWorldRaycastMode::Segments)
+	{
+		OutEdges.Emplace(FPCGPinReference(PCGWorldRaycastElement::Constants::Labels::EndPointsInput), FPCGPinReference(PreKernel, PCGWorldRaycastElement::Constants::Labels::EndPointsInput));
+		PreKernel->SetIncludeEndPointsInputPin(true);
+	}
+
+	// Node input to post-kernel so it can flow the data through (and modify positions based on RT results).
+	OutEdges.Emplace(FPCGPinReference(PCGWorldRaycastElement::Constants::Labels::OriginsInput), FPCGPinReference(PostKernel, PCGPinConstants::DefaultInputLabel));
+
+	// Ray trace results to post-kernel.
+	OutEdges.Emplace(FPCGPinReference(FPCGPinReference(PreKernel, PCGPinConstants::DefaultOutputLabel)), FPCGPinReference(PostKernel, PCGPostRayTraceKernel::RayTraceResultsLabel));
+
+	// Post kernel to node output.
+	OutEdges.Emplace(FPCGPinReference(FPCGPinReference(PostKernel, PCGPinConstants::DefaultOutputLabel)), FPCGPinReference(PCGPinConstants::DefaultOutputLabel));
+}
 #endif // WITH_EDITOR
 
+bool UPCGWorldRaycastElementSettings::IsPinUsedByNodeExecution(const UPCGPin* InPin) const
+{
+	if (InPin && ShouldExecuteOnGPU() && InPin->Properties.Label == PCGWorldRaycastElement::Constants::Labels::BoundingShape)
+	{
+		return false;
+	}
+
+	return Super::IsPinUsedByNodeExecution(InPin);
+}
+
 FString UPCGWorldRaycastElementSettings::GetAdditionalTitleInformation() const
 {
 	using namespace PCGWorldRaycastElement::Constants;
@@ -178,10 +222,14 @@ void UPCGWorldRaycastElementSettings::ApplyPreconfiguredSettings(const FPCGPreCo
 
 TArray<FPCGPinProperties> UPCGWorldRaycastElementSettings::InputPinProperties() const
 {
+	// Note: If executing on the GPU, we need to prevent multiple connections on inputs, since it is not supported at this time.
+	// Also note: Since the ShouldExecuteOnGPU() is already tied to structural changes, we don't need to implement any logic for this in GetChangeTypeForProperty()
+	const bool bAllowMultipleConnections = !ShouldExecuteOnGPU();
+
 	TArray<FPCGPinProperties> PinProperties;
 
 	using namespace PCGWorldRaycastElement::Constants;
-	FPCGPinProperties& OriginPointsInputPin = PinProperties.Emplace_GetRef(Labels::OriginsInput, EPCGDataType::PointOrParam);
+	FPCGPinProperties& OriginPointsInputPin = PinProperties.Emplace_GetRef(Labels::OriginsInput, EPCGDataType::PointOrParam, bAllowMultipleConnections);
 	OriginPointsInputPin.SetRequiredPin();
 #if WITH_EDITOR
 	OriginPointsInputPin.Tooltip = LOCTEXT("OriginPointsInputPinTooltip", "Points acting as the point origin of each individual ray. Only point data from this pin may be forwarded.");
@@ -189,7 +237,7 @@ TArray<FPCGPinProperties> UPCGWorldRaycastElementSettings::InputPinProperties()
 
 	if (RaycastMode == EPCGWorldRaycastMode::Segments)
 	{
-		FPCGPinProperties& EndPointsInputPin = PinProperties.Emplace_GetRef(Labels::EndPointsInput, EPCGDataType::PointOrParam);
+		FPCGPinProperties& EndPointsInputPin = PinProperties.Emplace_GetRef(Labels::EndPointsInput, EPCGDataType::PointOrParam, bAllowMultipleConnections);
 #if WITH_EDITOR
 		EndPointsInputPin.Tooltip = LOCTEXT("EndPointsInputPinTooltip", "Points acting as the end point of each individual ray.");
 #endif // WITH_EDITOR
diff --git a/Engine/Plugins/PCG/Source/PCG/Public/Compute/PCGComputeKernel.h b/Engine/Plugins/PCG/Source/PCG/Public/Compute/PCGComputeKernel.h
--- a/Engine/Plugins/PCG/Source/PCG/Public/Compute/PCGComputeKernel.h
+++ b/Engine/Plugins/PCG/Source/PCG/Public/Compute/PCGComputeKernel.h
@@ -198,6 +198,9 @@ public:
 
 	bool HasStaticValidationErrors() const { return !!(PCGKernelFlags & EPCGComputeKernelFlags::HasStaticValidationErrors); }
 
+	/** True if this compute kernel must run after the TLAS is built. This will force the entire compute graph to execute in this group. */
+	virtual bool RequiresPostTLASBuildExecutionGroup() const;
+
 #ifdef PCG_GPU_KERNEL_PROFILING
 	void EnableRepeatDispatch() { PCGKernelFlags |= EPCGComputeKernelFlags::RepeatDispatchForProfiling; }
 
diff --git a/Engine/Plugins/PCG/Source/PCG/Public/Elements/PCGWorldRaycast.h b/Engine/Plugins/PCG/Source/PCG/Public/Elements/PCGWorldRaycast.h
--- a/Engine/Plugins/PCG/Source/PCG/Public/Elements/PCGWorldRaycast.h
+++ b/Engine/Plugins/PCG/Source/PCG/Public/Elements/PCGWorldRaycast.h
@@ -11,6 +11,11 @@
 
 #include "PCGWorldRaycast.generated.h"
 
+namespace PCGWorldRaycastElement::Constants::Labels
+{
+	extern const FName EndPointsInput;
+}
+
 /**
  * Casts rays from provided points along a given direction and transform points to the impact point.
  */
@@ -33,10 +38,13 @@ public:
 	virtual bool OnlyExposePreconfiguredSettings() const override { return false; }
 	virtual bool GroupPreconfiguredSettings() const override { return true; }
 	virtual EPCGChangeType GetChangeTypeForProperty(FPropertyChangedEvent& PropertyChangedEvent) const override;
+	virtual bool DisplayExecuteOnGPUSetting() const override { return true; }
+	virtual void CreateKernels(FPCGGPUCompilationContext& InOutContext, UObject* InObjectOuter, TArray<UPCGComputeKernel*>& OutKernels, TArray<FPCGKernelEdge>& OutEdges) const override;
 #endif // WITH_EDITOR
 
 	virtual FString GetAdditionalTitleInformation() const override;
 	virtual void ApplyPreconfiguredSettings(const FPCGPreConfiguredSettingsInfo& PreconfigureInfo) override;
+	virtual bool IsPinUsedByNodeExecution(const UPCGPin* InPin) const override;
 
 protected:
 	virtual TArray<FPCGPinProperties> InputPinProperties() const override;
@@ -86,7 +94,7 @@ public:
 	FPCGAttributePropertyInputSelector RayLengthAttribute;
 
 	/** World ray trace parameters. */
-	UPROPERTY(BlueprintReadWrite, EditAnywhere, Category = Settings, meta = (PCG_Overridable, ShowOnlyInnerProperties))
+	UPROPERTY(BlueprintReadWrite, EditAnywhere, Category = Settings, meta = (PCG_Overridable, ShowOnlyInnerProperties, EditCondition = "!bExecuteOnGPU", EditConditionHides))
 	FPCGWorldRaycastQueryParams WorldQueryParams;
 
 	/** Will keep the original points at their location if the raycast misses or if the hit result is out of bounds. */
@@ -96,6 +104,18 @@ public:
 	/** If no Bounding Shape input is provided, the actor bounds are used to limit the sample generation domain. */
 	UPROPERTY(BlueprintReadWrite, EditAnywhere, Category = Settings, meta = (PCG_Overridable))
 	bool bUnbounded = false;
+
+	/** [EXPERIMENTAL] Create a Vector2 attribute ('ImpactUVCoords') for UV Coordinates of the surface hit. */
+	UPROPERTY(BlueprintReadWrite, EditAnywhere, Category = Settings, meta = (PCG_Overridable, EditCondition = "bExecuteOnGPU", EditConditionHides))
+	bool bGetTextureCoordinates = false;
+
+	/** [EXPERIMENTAL] Which channel to sample UV Coordinates from. If this channel does not exist, UV coordinates will be (0, 0). */
+	UPROPERTY(BlueprintReadWrite, EditAnywhere, Category = Settings, meta = (PCG_Overridable, EditCondition = "bExecuteOnGPU && bGetTextureCoordinates", EditConditionHides))
+	int TextureCoordinatesChannelIndex = 0;
+
+	/** [EXPERIMENTAL] Create a Vector attribute ('ImpactNormal') for normals of the surface hit. */
+	UPROPERTY(BlueprintReadWrite, EditAnywhere, Category = Settings, meta = (PCG_Overridable, EditCondition = "bExecuteOnGPU", EditConditionHides))
+	bool bGetNormals = false;
 };
 
 namespace FPCGWorldRaycast
diff --git a/Engine/Plugins/PCG/Source/PCGCompute/Internal/PCGRayTrace.h b/Engine/Plugins/PCG/Source/PCGCompute/Internal/PCGRayTrace.h
new file mode 100644
--- /dev/null
+++ b/Engine/Plugins/PCG/Source/PCGCompute/Internal/PCGRayTrace.h
@@ -0,0 +1,34 @@
+// Copyright Epic Games, Inc. All Rights Reserved.
+
+#pragma once
+
+#include "RHIDefinitions.h"
+#include "RenderGraphFwd.h"
+
+#if RHI_RAYTRACING
+class FSceneInterface;
+class FSceneView;
+#endif
+
+namespace PCGRaytraceConstants
+{
+	constexpr int RAY_CULLED = -1;
+	constexpr int RAY_TRACE_PACKED_BUFFER_STRIDE_UINTS = 12;
+}
+
+#if RHI_RAYTRACING
+struct FPCGRayTraceParams
+{
+	const FSceneInterface* Scene = nullptr;
+	const FSceneView* View = nullptr;
+	int32 NumRays = 0;
+	bool bNeedsUVData = false;
+	int32 TexCoordsChannelIndex = 0;
+	FRDGBufferUAVRef PackedDataUAV;
+};
+
+namespace PCGRayTrace
+{
+	void PCGCOMPUTE_API RenderPCGRayTraceInline(FRDGBuilder& GraphBuilder, const FPCGRayTraceParams& InParams);
+};
+#endif // RHI_RAYTRACING
diff --git a/Engine/Plugins/PCG/Source/PCGCompute/PCGCompute.Build.cs b/Engine/Plugins/PCG/Source/PCGCompute/PCGCompute.Build.cs
--- a/Engine/Plugins/PCG/Source/PCGCompute/PCGCompute.Build.cs
+++ b/Engine/Plugins/PCG/Source/PCGCompute/PCGCompute.Build.cs
@@ -1,5 +1,7 @@
 // Copyright Epic Games, Inc. All Rights Reserved.
 
+using System.IO;
+
 namespace UnrealBuildTool.Rules
 {
 	public class PCGCompute : ModuleRules
@@ -24,6 +26,11 @@ namespace UnrealBuildTool.Rules
 				{
 					"Engine",
 				});
+
+			PrivateIncludePaths.AddRange(
+				new string[] {
+					Path.Combine(GetModuleDirectory("Renderer"), "Private"),
+				});
 		}
 	}
 }
diff --git a/Engine/Plugins/PCG/Source/PCGCompute/Private/PCGRayTrace.cpp b/Engine/Plugins/PCG/Source/PCGCompute/Private/PCGRayTrace.cpp
new file mode 100644
--- /dev/null
+++ b/Engine/Plugins/PCG/Source/PCGCompute/Private/PCGRayTrace.cpp
@@ -0,0 +1,212 @@
+// Copyright Epic Games, Inc. All Rights Reserved.
+
+#include "PCGRayTrace.h"
+
+#include "RHI.h"
+
+#include "PCGComputeModule.h"
+
+#include "BuiltInRayTracingShaders.h"
+#include "DataDrivenShaderPlatformInfo.h"
+#include "EngineModule.h"
+#include "FXRenderingUtils.h"
+#include "GlobalShader.h"
+#include "PipelineStateCache.h"
+#include "RHIGPUReadback.h"
+#include "RayTracingPayloadType.h"
+#include "RayTracingShaderBindingLayout.h"
+#include "RenderGraphBuilder.h"
+#include "RenderGraphUtils.h"
+#include "RendererUtils.h"
+#include "SceneInterface.h"
+#include "ScenePrivate.h"
+#include "SceneUniformBuffer.h"
+#include "SceneView.h"
+#include "ShaderParameterStruct.h"
+
+#if RHI_RAYTRACING
+class FPCGRayTraceCS : public FGlobalShader
+{
+	DECLARE_GLOBAL_SHADER(FPCGRayTraceCS)
+	SHADER_USE_PARAMETER_STRUCT(FPCGRayTraceCS, FGlobalShader)
+
+	BEGIN_SHADER_PARAMETER_STRUCT(FParameters, )
+		SHADER_PARAMETER_RDG_BUFFER_UAV(RWByteAddressBuffer, InOutPackedData)
+		SHADER_PARAMETER_RDG_BUFFER_SRV(RaytracingAccelerationStructure, InTLAS)
+		SHADER_PARAMETER_RESOURCE_COLLECTION(FResourceCollection, InPrimitiveUVCollection)
+		SHADER_PARAMETER_RDG_BUFFER_SRV(Buffer<uint>, InPrimitiveUVOffsets)
+		SHADER_PARAMETER_RDG_BUFFER_SRV(StructuredBuffer<FRayTracingSceneMetadataRecord>, InRayTracingSceneMetadata)
+		SHADER_PARAMETER(FVector3f, InPreViewTranslation)
+		SHADER_PARAMETER(int32, InNumRays)
+		SHADER_PARAMETER(int32, InTexCoordsChannelIndex)
+
+		// These might just be here to prompt RG to create a resource. Its bound statically in addpass.
+		SHADER_PARAMETER_RDG_UNIFORM_BUFFER(FNaniteRayTracingUniformParameters, NaniteRayTracing)
+		SHADER_PARAMETER_RDG_UNIFORM_BUFFER(FSceneUniformParameters, Scene)
+		SHADER_PARAMETER_STRUCT_REF(FViewUniformShaderParameters, View)
+		SHADER_PARAMETER_STRUCT_INCLUDE(FGPUSceneResourceParameters, GPUSceneParameters)
+	END_SHADER_PARAMETER_STRUCT()
+
+	static bool ShouldCompilePermutation(const FGlobalShaderPermutationParameters& Parameters)
+	{
+		return IsRayTracingEnabledForProject(Parameters.Platform);
+	}
+
+	static void ModifyCompilationEnvironment(const FGlobalShaderPermutationParameters& Parameters, FShaderCompilerEnvironment& OutEnvironment)
+	{
+		FGlobalShader::ModifyCompilationEnvironment(Parameters, OutEnvironment);
+
+		OutEnvironment.SetDefine(TEXT("VF_SUPPORTS_PRIMITIVE_SCENE_DATA"), 1);
+		OutEnvironment.SetDefine(TEXT("USE_GLOBAL_GPU_SCENE_DATA"), 1);
+		OutEnvironment.SetDefine(TEXT("PCG_RAY_CULLED"), FString::FromInt(PCGRaytraceConstants::RAY_CULLED));
+		OutEnvironment.SetDefine(TEXT("PCG_RAY_TRACE_PACKED_BUFFER_STRIDE_UINTS"), FString::FromInt(PCGRaytraceConstants::RAY_TRACE_PACKED_BUFFER_STRIDE_UINTS));
+		
+		OutEnvironment.CompilerFlags.Add(CFLAG_SupportsMinimalBindless);
+		OutEnvironment.CompilerFlags.Add(CFLAG_InlineRayTracing);
+		OutEnvironment.CompilerFlags.Add(CFLAG_ForceDXC);
+	}
+};
+IMPLEMENT_GLOBAL_SHADER(FPCGRayTraceCS, "/PCGComputeShaders/PCGRayTrace.usf", "PCGRayTraceMainCS", SF_Compute);
+
+void PCGRayTrace::RenderPCGRayTraceInline(FRDGBuilder& GraphBuilder, const FPCGRayTraceParams& InParams)
+{
+	TRACE_CPUPROFILER_EVENT_SCOPE(PCGRayTrace::RenderPCGRayTraceInline);
+	check(IsInRenderingThread());
+
+	if (!IsRayTracingEnabled())
+	{
+		UE_LOG(LogPCGCompute, Error, TEXT("Failed to submit inline raytrace shader. Ray tracing is not enabled."));
+		return;
+	}
+
+	if (!GRHISupportsInlineRayTracing)
+	{
+		UE_LOG(LogPCGCompute, Error, TEXT("Failed to submit inline raytrace shader. Inline ray tracing is not supported."));
+		return;
+	}
+
+	if (!InParams.Scene)
+	{
+		UE_LOG(LogPCGCompute, Error, TEXT("Failed to submit inline raytrace shader. Invalid scene."));
+		return;
+	}
+
+	if (!InParams.View)
+	{
+		UE_LOG(LogPCGCompute, Error, TEXT("Failed to submit inline raytrace shader. Invalid view."));
+		return;
+	}
+
+	if (!InParams.View->GetInlineRayTracingBindingDataBuffer())
+	{
+		UE_LOG(LogPCGCompute, Error, TEXT("Failed to submit inline raytrace shader. Inline ray tracing binding data buffer is not available."));
+		return;
+	}
+
+	const FScene* RenderScene = InParams.Scene->GetRenderScene();
+
+	FGlobalShaderMap* ShaderMap = GetGlobalShaderMap(GMaxRHIFeatureLevel);
+	TShaderMapRef<FPCGRayTraceCS> ComputeShader(ShaderMap);
+
+	FRHIResourceCollectionRef PrimitiveUVCollection;
+	FRDGBufferRef PrimitiveIndirectionsRDG;
+
+	if (InParams.bNeedsUVData && InParams.TexCoordsChannelIndex >= 0)
+	{
+		TRACE_CPUPROFILER_EVENT_SCOPE(PCGRayTrace::BuildPrimitiveUVData);
+
+		TArray<TPair<FRaytracingVFUVData, uint32>, SceneRenderingAllocator> UVBuffers;
+		uint32 MaxPrimitiveId = 0;
+
+		// Build UV data for all primitives in the render scene.
+		// @todo_pcg: Instead of building this every time we dispatch the RT, we should cache the UVs when primitives are added/removed to the scene.
+		for (int32 PrimitiveIndex = 0; PrimitiveIndex < RenderScene->Primitives.Num(); PrimitiveIndex++)
+		{
+			const FPrimitiveSceneInfo* Primitive = RenderScene->Primitives[PrimitiveIndex];
+
+			if (!Primitive || !Primitive->Proxy || !Primitive->Proxy->HasRayTracingRepresentation())
+			{
+				continue;
+			}
+
+			const uint32 PrimitiveId = Primitive->GetPersistentIndex().Index;
+			MaxPrimitiveId = FMath::Max(MaxPrimitiveId, PrimitiveId);
+
+			const FRaytracingVFAttributeData VFAttributeData = Primitive->Proxy->GetRaytracingVFAttributeData(InParams.View);
+
+			if (VFAttributeData.Sections.IsEmpty() || !VFAttributeData.Sections[0].UV.BufferSRV)
+			{
+				continue;
+			}
+
+			// Currently only FStaticMeshSceneProxy overrides GetRaytracingVFAttributeData, and it returns
+			// a single section (static meshes share one UV buffer across all sections per LOD).
+			// If other proxy types override this with multiple sections, this will need to select
+			// the correct section based on the geometry index from the ray hit.
+			UVBuffers.Emplace(VFAttributeData.Sections[0].UV, PrimitiveId);
+		}
+
+		// Build a resource collection for the primitives.
+		TArray<FRHIResourceCollectionMember, SceneRenderingAllocator> PrimitiveUVCollectionMembers;
+		FRDGUploadData<FUintPoint> PrimitiveIndirections(GraphBuilder, MaxPrimitiveId + 1);
+		FMemory::Memset(PrimitiveIndirections.GetData(), 0xFF, PrimitiveIndirections.GetTotalSize());
+
+		for (auto&& [UVBuffer, Key] : UVBuffers)
+		{
+			PrimitiveIndirections[Key] = { static_cast<uint32>(PrimitiveUVCollectionMembers.Num()), UVBuffer.NumCoordinates };
+			PrimitiveUVCollectionMembers.Add(FRHIResourceCollectionMember(UVBuffer.BufferSRV));
+		}
+
+		PrimitiveUVCollection = GraphBuilder.RHICmdList.CreateResourceCollection(PrimitiveUVCollectionMembers);
+		PrimitiveIndirectionsRDG = CreateUploadBuffer(
+			GraphBuilder,
+			TEXT("PCG::UVIndirectionBuffer"),
+			/*BytesPerElement=*/sizeof(FUintPoint),
+			/*NumElements=*/FMath::RoundUpToPowerOfTwo(PrimitiveIndirections.Num()),
+			PrimitiveIndirections);
+	}
+	else
+	{
+		// Create empty dummy resources to satisfy shader parameter binding.
+		PrimitiveUVCollection = GraphBuilder.RHICmdList.CreateResourceCollection({});
+		FRDGUploadData<FUintPoint> DummyIndirections(GraphBuilder, 1);
+		DummyIndirections[0] = { ~0u, 0u };
+		PrimitiveIndirectionsRDG = CreateUploadBuffer(
+			GraphBuilder,
+			TEXT("PCG::UVIndirectionBuffer"),
+			/*BytesPerElement=*/sizeof(FUintPoint),
+			/*NumElements=*/1,
+			DummyIndirections);
+	}
+
+	FSceneUniformBuffer SceneUniformBuffer;
+	TRDGUniformBufferRef<FSceneUniformParameters> SceneUniformBufferRDG = UE::FXRenderingUtils::GetSceneUniformBuffer(GraphBuilder, SceneUniformBuffer);
+
+	FPCGRayTraceCS::FParameters* ShaderParameters = GraphBuilder.AllocParameters<FPCGRayTraceCS::FParameters>();
+	ShaderParameters->InOutPackedData = InParams.PackedDataUAV;
+	ShaderParameters->InTLAS = UE::FXRenderingUtils::RayTracing::GetRayTracingSceneViewRDG(*InParams.Scene, *InParams.View);
+	ShaderParameters->InPrimitiveUVCollection = PrimitiveUVCollection;
+	ShaderParameters->InPrimitiveUVOffsets = GraphBuilder.CreateSRV(PrimitiveIndirectionsRDG, PF_R32_UINT);
+	ShaderParameters->InRayTracingSceneMetadata = GraphBuilder.CreateSRV(InParams.View->GetInlineRayTracingBindingDataBuffer());
+	ShaderParameters->InPreViewTranslation = FVector3f(InParams.View->ViewMatrices.GetPreViewTranslation());
+	ShaderParameters->InNumRays = InParams.NumRays;
+	ShaderParameters->InTexCoordsChannelIndex = InParams.TexCoordsChannelIndex;
+	ShaderParameters->NaniteRayTracing = Nanite::GetPublicGlobalRayTracingUniformBuffer();
+	ShaderParameters->Scene = SceneUniformBufferRDG;
+	ShaderParameters->View = InParams.View->ViewUniformBuffer;
+	ShaderParameters->GPUSceneParameters = RenderScene->GPUScene.GetShaderParameters(GraphBuilder);
+
+	GraphBuilder.AddPass(
+		RDG_EVENT_NAME("PCGRayTrace"),
+		ShaderParameters,
+		ERDGPassFlags::Compute | ERDGPassFlags::NeverCull,
+		[NumRays=InParams.NumRays, ShaderParameters, ComputeShader](FRDGAsyncTask, FRHICommandList& RHICmdList)
+		{
+			FComputeShaderUtils::Dispatch(
+				RHICmdList,
+				ComputeShader,
+				*ShaderParameters,
+				FIntVector(FMath::DivideAndRoundUp(NumRays, 32), 1, 1));
+		});
+}
+#endif // RHI_RAYTRACING
