Unshelved from pending changelist '51851200':

[PCG] First version for manual edit of splineData and Polygon data.

#Synced-CL 52134369

diff --git a/Engine/Plugins/PCG/Source/PCG/Private/Graph/DataOverride/PCGDataOverride.cpp b/Engine/Plugins/PCG/Source/PCG/Private/Graph/DataOverride/PCGDataOverride.cpp
--- a/Engine/Plugins/PCG/Source/PCG/Private/Graph/DataOverride/PCGDataOverride.cpp
+++ b/Engine/Plugins/PCG/Source/PCG/Private/Graph/DataOverride/PCGDataOverride.cpp
@@ -2,6 +2,8 @@
 
 #include "Graph/DataOverride/PCGDataOverride.h"
 
+#include "Graph/DataOverride/PCGDataOverrideHelpers.h"
+
 FPCGDeltaKey::FPCGDeltaKey(const FXxHash64 InHash, const FName InDeltaLabel)
 #if WITH_EDITOR
 	: DeltaLabel(InDeltaLabel)
@@ -75,20 +77,28 @@ bool FPCGDeltaCollection::IsEmpty() const
 	return Deltas.IsEmpty();
 }
 
-void FPCGDeltaCollection::ForEachDelta(TFunctionRef<void(const FPCGDeltaKey&, TInstancedStruct<FPCGDeltaBase>&)> Func)
+bool FPCGDeltaCollection::ForEachDelta(TFunctionRef<bool(const FPCGDeltaKey&, TInstancedStruct<FPCGDeltaBase>&)> Func)
 {
 	for (TPair<FPCGDeltaKey, TInstancedStruct<FPCGDeltaBase>>& Pair : Deltas)
 	{
-		Func(Pair.Key, Pair.Value);
+		if (!Func(Pair.Key, Pair.Value))
+		{
+			return false;
+		}
 	}
+	return true;
 }
 
-void FPCGDeltaCollection::ForEachDelta(TFunctionRef<void(const FPCGDeltaKey&, const TInstancedStruct<FPCGDeltaBase>&)> Func) const
+bool FPCGDeltaCollection::ForEachDelta(TFunctionRef<bool(const FPCGDeltaKey&, const TInstancedStruct<FPCGDeltaBase>&)> Func) const
 {
 	for (const TPair<FPCGDeltaKey, TInstancedStruct<FPCGDeltaBase>>& Pair : Deltas)
 	{
-		Func(Pair.Key, Pair.Value);
+		if (!Func(Pair.Key, Pair.Value))
+		{
+			return false;
+		}
 	}
+	return true;
 }
 
 FName FPCGDeltaBase::GetDeltaName() const
@@ -110,3 +120,19 @@ bool FPCGDeltaBase::Apply(UPCGData* InData, const int32 ResolvedIndex) const
 {
 	return false;
 }
+
+FPCGDeltaKey FPCGDeltaBase::ComputeKey(const FPCGDeltaSettings& DeltaSettings, const FTransform& InTransform, const FBox& InWorldBounds)
+{
+	using namespace PCG::DataOverride::Keys;
+	switch (DeltaSettings.KeyTarget)
+	{
+	case EPCGSpatialKeyTarget::Position:
+		return FPCGLocationDeltaKey(InTransform.GetLocation(), DeltaSettings.SpatialTolerance);
+	case EPCGSpatialKeyTarget::AABB:
+		return FPCGAABBDeltaKey(InWorldBounds, DeltaSettings.SpatialTolerance);
+	case EPCGSpatialKeyTarget::Transform:
+		return FPCGTransformDeltaKey(InTransform, DeltaSettings.SpatialTolerance);
+	default:
+		return {};
+	}
+}
diff --git a/Engine/Plugins/PCG/Source/PCG/Private/Graph/DataOverride/PCGDataOverrideHelpers.cpp b/Engine/Plugins/PCG/Source/PCG/Private/Graph/DataOverride/PCGDataOverrideHelpers.cpp
--- a/Engine/Plugins/PCG/Source/PCG/Private/Graph/DataOverride/PCGDataOverrideHelpers.cpp
+++ b/Engine/Plugins/PCG/Source/PCG/Private/Graph/DataOverride/PCGDataOverrideHelpers.cpp
@@ -203,7 +203,7 @@ namespace PCG::DataOverride
 		{
 			// @todo_pcg: This should be part of a registry. For now, default to points, but should be element
 			// We have data ownership at either phase and thus const_cast is okay.
-			if (UPCGPointArrayData* PointData = const_cast<UPCGPointArrayData*>(Cast<UPCGPointArrayData>(TaggedData.Data)))
+			if (UPCGData* PCGData = const_cast<UPCGData*>(Cast<UPCGData>(TaggedData.Data)))
 			{
 				FPCGSourceDataStorageKey PinDataKey(Constants::DefaultOverrideLabel, Helpers::ComputeDataHashWithPin(NodeKey, TaggedData.Pin).Hash);
 
@@ -223,17 +223,35 @@ namespace PCG::DataOverride
 				{
 					// @todo_pcg: Sort deltas (edits->deletions->insertions) to avoid index instability.
 					// @todo_pcg: As a future optimization pass, aggregate candidates to avoid redundant key computation.
-					DeltaCollection->ForEachDelta([DeltaCollection, PointData](const FPCGDeltaKey&, const TInstancedStruct<FPCGDeltaBase>& DeltaStruct)
+					DeltaCollection->ForEachDelta([DeltaCollection, PCGData, &TaggedData](const FPCGDeltaKey&, const TInstancedStruct<FPCGDeltaBase>& DeltaStruct)
 					{
-						const FPCGDeltaBase* Delta = DeltaStruct.GetPtr<FPCGDeltaBase>();
-						if (!Delta)
+						if (const FPCGDeltaBase* Delta = DeltaStruct.GetPtr<FPCGDeltaBase>())
 						{
-							return;
+							if (const PCGIndexing::FPCGIndexCollection Candidates = Delta->FilterCandidates(PCGData); !Candidates.IsEmpty())
+							{
+								if (const int32 ResolvedIndex = Delta->Resolve(PCGData, Candidates, DeltaCollection->Settings); ResolvedIndex != INDEX_NONE)
+								{
+									if (Delta->UsesReplacementData())
+									{
+										if (UPCGData* ReplacementData = Delta->GetReplacementData())
+										{
+											TaggedData.Data = ReplacementData;
+											// stop iterating if replacement occurred.
+											return false;
+										}
+										else
+										{
+											UE_LOG(LogPCG, Error, TEXT("Delta marked as replacing but returned null replacement data"));
+										}
+									}
+									else
+									{
+										Delta->Apply(PCGData, ResolvedIndex);
+									}
+							   }
+						   }
 						}
-
-						const PCGIndexing::FPCGIndexCollection Candidates = Delta->FilterCandidates(PointData);
-						const int32 ResolvedIndex = Delta->Resolve(PointData, Candidates, DeltaCollection->Settings);
-						Delta->Apply(PointData, ResolvedIndex);
+						return true;
 					});
 				}
 			}
diff --git a/Engine/Plugins/PCG/Source/PCG/Private/Graph/DataOverride/PCGDataOverridePoints.cpp b/Engine/Plugins/PCG/Source/PCG/Private/Graph/DataOverride/PCGDataOverridePoints.cpp
--- a/Engine/Plugins/PCG/Source/PCG/Private/Graph/DataOverride/PCGDataOverridePoints.cpp
+++ b/Engine/Plugins/PCG/Source/PCG/Private/Graph/DataOverride/PCGDataOverridePoints.cpp
@@ -41,24 +41,9 @@ int32 FPCGPointDeltaBase::Resolve(const UPCGData* InData, const PCGIndexing::FPC
 		return ResolvedIndex;
 	}
 
-	auto ComputeKey = [&DeltaSettings](const FTransform& InTransform, const FBox& InWorldBounds) -> FPCGDeltaKey
-	{
-		switch (DeltaSettings.KeyTarget)
-		{
-			case EPCGSpatialKeyTarget::Position:
-				return Keys::FPCGLocationDeltaKey(InTransform.GetLocation(), DeltaSettings.SpatialTolerance);
-			case EPCGSpatialKeyTarget::AABB:
-				return Keys::FPCGAABBDeltaKey(InWorldBounds, DeltaSettings.SpatialTolerance);
-			case EPCGSpatialKeyTarget::Transform:
-				return Keys::FPCGTransformDeltaKey(InTransform, DeltaSettings.SpatialTolerance);
-			default:
-				return {};
-		}
-	};
-
-	FPCGDeltaKey DeltaKey = ComputeKey(OriginalTransform, Bounds);
+	FPCGDeltaKey DeltaKey = ComputeKey(DeltaSettings, OriginalTransform, Bounds);
 
-	FilteredCandidates.ForEach([&ResolvedIndex, &DeltaKey, ComputeKey, PointData](const PCGIndexing::FPCGIndexRange& IndexRange)
+	FilteredCandidates.ForEach([&ResolvedIndex, &DeltaKey, &DeltaSettings, PointData](const PCGIndexing::FPCGIndexRange& IndexRange)
 	{
 		for (int32 Index = IndexRange.StartIndex; Index < IndexRange.EndIndex; ++Index)
 		{
@@ -66,8 +51,7 @@ int32 FPCGPointDeltaBase::Resolve(const UPCGData* InData, const PCGIndexing::FPC
 			FBox LocalBounds = PointData->GetLocalBounds(Index);
 			FBox WorldBounds = LocalBounds.TransformBy(Transform);
 
-			FPCGDeltaKey CandidateKey = ComputeKey(Transform, WorldBounds);
-			if (DeltaKey == CandidateKey)
+			if (FPCGDeltaKey CandidateKey = ComputeKey(DeltaSettings, Transform, WorldBounds); DeltaKey == CandidateKey)
 			{
 				ResolvedIndex = Index;
 				return true;
diff --git a/Engine/Plugins/PCG/Source/PCG/Private/Graph/DataOverride/PCGDataOverridePolygon.cpp b/Engine/Plugins/PCG/Source/PCG/Private/Graph/DataOverride/PCGDataOverridePolygon.cpp
new file mode 100644
--- /dev/null
+++ b/Engine/Plugins/PCG/Source/PCG/Private/Graph/DataOverride/PCGDataOverridePolygon.cpp
@@ -0,0 +1,71 @@
+// Copyright Epic Games, Inc. All Rights Reserved.
+
+#include "Graph/DataOverride/PCGDataOverridePolygon.h"
+
+#include "Graph/DataOverride/PCGDataOverrideHelpers.h"
+
+
+namespace PCG::DataOverride::Polygon::Constants
+{
+	inline constexpr FLazyName PolygonOverrideDeltaName = "PolygonOverride";
+}
+
+FName FPCGPolygon2DDataDelta::GetDeltaNameStatic()
+{
+	return PCG::DataOverride::Polygon::Constants::PolygonOverrideDeltaName;
+}
+
+PCGIndexing::FPCGIndexCollection FPCGPolygon2DDataDelta::FilterCandidates(const UPCGData* InData) const
+{
+	if (!PolygonOverride)
+	{
+		return PCGIndexing::FPCGIndexCollection::Invalid();
+	}
+
+	if (const UPCGPolygon2DData* PolygonData = Cast<UPCGPolygon2DData>(InData))
+	{
+		PCGIndexing::FPCGIndexCollection IndexCollection(PolygonData->GetNumVertices());
+		IndexCollection.AddRange(0, PolygonData->GetNumVertices());
+		return IndexCollection;
+	}
+
+	return PCGIndexing::FPCGIndexCollection::Invalid();
+}
+
+int32 FPCGPolygon2DDataDelta::Resolve(const UPCGData* InData, const PCGIndexing::FPCGIndexCollection& FilteredCandidates,
+	const FPCGDeltaSettings& DeltaSettings) const
+{
+	if (!PolygonOverride || FilteredCandidates.IsEmpty() || (DeltaSettings.KeyTarget == EPCGSpatialKeyTarget::AABB))
+	{
+		return INDEX_NONE;
+	}
+
+	if (const UPCGPolygon2DData* PolygonData = Cast<UPCGPolygon2DData>(InData); PolygonData && PolygonData->GetNumVertices() > 0)
+	{
+		const FPCGDeltaKey InputKey = ComputeKey(DeltaSettings, PolygonData->GetTransformAtDistance(0, 0.0, true), FBox{});
+	
+		const FPCGDeltaKey OverrideKey = ComputeKey(DeltaSettings, OriginalStartTransform, FBox{});
+	
+		// Index 0 means do replace the polygon data.
+		return InputKey == OverrideKey ? 0 : INDEX_NONE;
+	}
+
+	return INDEX_NONE;
+}
+
+FName FPCGPolygon2DDataDelta::GetDeltaName() const
+{
+	return GetDeltaNameStatic();
+}
+
+bool FPCGPolygon2DDataDelta::Apply(UPCGData* InData, int32 ResolvedIndex) const
+{
+	// Should not call apply on this delta, it is a replacement delta.
+	checkNoEntry();
+	return false;
+}
+
+UPCGData* FPCGPolygon2DDataDelta::GetReplacementData() const
+{
+	return PolygonOverride;
+}
diff --git a/Engine/Plugins/PCG/Source/PCG/Private/Graph/DataOverride/PCGDataOverrideSpline.cpp b/Engine/Plugins/PCG/Source/PCG/Private/Graph/DataOverride/PCGDataOverrideSpline.cpp
new file mode 100644
--- /dev/null
+++ b/Engine/Plugins/PCG/Source/PCG/Private/Graph/DataOverride/PCGDataOverrideSpline.cpp
@@ -0,0 +1,74 @@
+// Copyright Epic Games, Inc. All Rights Reserved.
+
+#include "Graph/DataOverride/PCGDataOverrideSpline.h"
+
+#include "Data/PCGSplineData.h"
+#include "Graph/DataOverride/PCGDataOverrideHelpers.h"
+
+
+namespace PCG::DataOverride::Spline::Constants
+{
+	inline constexpr FLazyName SplineOverrideDeltaName = "SplineOverride";
+}
+
+FName FPCGSplineDataDelta::GetDeltaNameStatic()
+{
+	return PCG::DataOverride::Spline::Constants::SplineOverrideDeltaName;
+}
+
+PCGIndexing::FPCGIndexCollection FPCGSplineDataDelta::FilterCandidates(const UPCGData* InData) const
+{
+	if (!SplineOverride)
+	{
+		return PCGIndexing::FPCGIndexCollection::Invalid();
+	}
+
+	if (const UPCGSplineData* SplineData = Cast<UPCGSplineData>(InData))
+	{
+		PCGIndexing::FPCGIndexCollection IndexCollection(SplineData->GetNumVertices());
+		IndexCollection.AddRange(0, SplineData->GetNumVertices());
+		return IndexCollection;
+	}
+
+	return PCGIndexing::FPCGIndexCollection::Invalid();
+}
+
+int32 FPCGSplineDataDelta::Resolve(const UPCGData* InData, const PCGIndexing::FPCGIndexCollection& FilteredCandidates,
+	const FPCGDeltaSettings& DeltaSettings) const
+{
+	if (!SplineOverride || FilteredCandidates.IsEmpty() || (DeltaSettings.KeyTarget == EPCGSpatialKeyTarget::AABB))
+	{
+		return INDEX_NONE;
+	}
+	
+	if (const UPCGSplineData* SplineData = Cast<UPCGSplineData>(InData); SplineData && SplineData->GetNumVertices() > 0)
+	{
+		FSplinePoint InputSplinePoint = SplineData->GetSplinePoints()[0];
+		FTransform InputSplinePointTransform{InputSplinePoint.Rotation, InputSplinePoint.Position, InputSplinePoint.Scale};
+		FPCGDeltaKey InputKey = ComputeKey(DeltaSettings, InputSplinePointTransform, FBox{});
+		FPCGDeltaKey OverrideKey = ComputeKey(DeltaSettings, OriginalStartTransform, FBox{});
+	
+		// Index 0 means do replace the whole spline data.
+		return InputKey == OverrideKey ? 0 : INDEX_NONE;
+	}
+
+	return INDEX_NONE;
+}
+
+FName FPCGSplineDataDelta::GetDeltaName() const
+{
+	return GetDeltaNameStatic();
+}
+
+bool FPCGSplineDataDelta::Apply(UPCGData* InData, int32 ResolvedIndex) const
+{
+	// Should not call apply on this delta, it is a replacement delta.
+	checkNoEntry();
+	return false;
+}
+
+UPCGData* FPCGSplineDataDelta::GetReplacementData() const
+{
+	return SplineOverride;
+}
+
diff --git a/Engine/Plugins/PCG/Source/PCG/Public/Graph/DataOverride/PCGDataOverride.h b/Engine/Plugins/PCG/Source/PCG/Public/Graph/DataOverride/PCGDataOverride.h
--- a/Engine/Plugins/PCG/Source/PCG/Public/Graph/DataOverride/PCGDataOverride.h
+++ b/Engine/Plugins/PCG/Source/PCG/Public/Graph/DataOverride/PCGDataOverride.h
@@ -124,12 +124,20 @@ struct FPCGDeltaBase
 	PCG_API virtual PCGIndexing::FPCGIndexCollection FilterCandidates(const UPCGData* InData) const;
 	PCG_API virtual int32 Resolve(const UPCGData* InData, const PCGIndexing::FPCGIndexCollection& FilteredCandidates, const FPCGDeltaSettings& DeltaSettings) const;
 	PCG_API virtual bool Apply(UPCGData* InData, int32 ResolvedIndex) const;
+	
+	/** Returns true if Delta application is a full data replacement. In this case, Apply() should not be called, UPCGData should simply be replaced by the return value of GetReplacementData() */
+	virtual bool UsesReplacementData() const { return false; }
+	
+	/** Returns the Data that should replace the current one, returns nullptr if the delta does not support full replacement.*/
+	virtual UPCGData* GetReplacementData() const { return nullptr; } 
 
 #if WITH_EDITORONLY_DATA
 	// @todo_pcg: The per-point spatial key used as the key for this specific delta. Keep for debugging, potentially permanently.
 	UPROPERTY()
 	FPCGDeltaKey ComputedKey;
 #endif // WITH_EDITORONLY_DATA
+	
+	static FPCGDeltaKey ComputeKey(const FPCGDeltaSettings& DeltaSettings, const FTransform& InTransform, const FBox& InWorldBounds);
 };
 
 /** A collection of deltas, usually tied directly to a node element or more specifically, to a pin output. */
@@ -163,8 +171,12 @@ struct FPCGDeltaCollection
 	PCG_API int32 Num() const;
 	PCG_API bool IsEmpty() const;
 
-	PCG_API void ForEachDelta(TFunctionRef<void(const FPCGDeltaKey&, TInstancedStruct<FPCGDeltaBase>&)> Func);
-	PCG_API void ForEachDelta(TFunctionRef<void(const FPCGDeltaKey&, const TInstancedStruct<FPCGDeltaBase>&)> Func) const;
+	/** Iterates over all deltas, calling the provided function for each.
+	 * @param Func Lambda returning true to continue iteration, false to stop early
+	 * @return true if all deltas were processed, false if iteration stopped early
+	 */
+	PCG_API bool ForEachDelta(TFunctionRef<bool(const FPCGDeltaKey&, TInstancedStruct<FPCGDeltaBase>&)> Func);
+	PCG_API bool ForEachDelta(TFunctionRef<bool(const FPCGDeltaKey&, const TInstancedStruct<FPCGDeltaBase>&)> Func) const;
 	// @todo_pcg: ForEachDeltaAs
 
 	UPROPERTY()
diff --git a/Engine/Plugins/PCG/Source/PCG/Public/Graph/DataOverride/PCGDataOverridePolygon.h b/Engine/Plugins/PCG/Source/PCG/Public/Graph/DataOverride/PCGDataOverridePolygon.h
new file mode 100644
--- /dev/null
+++ b/Engine/Plugins/PCG/Source/PCG/Public/Graph/DataOverride/PCGDataOverridePolygon.h
@@ -0,0 +1,33 @@
+// Copyright Epic Games, Inc. All Rights Reserved.
+
+#pragma once
+
+#include "Graph/DataOverride/PCGDataOverride.h"
+#include "Data/PCGPolygon2DData.h"
+
+#include "PCGDataOverridePolygon.generated.h"
+
+USTRUCT()
+struct FPCGPolygon2DDataDelta : public FPCGDeltaBase
+{
+	GENERATED_BODY()
+
+	static FName GetDeltaNameStatic();
+
+	//~ Begin FPCGDeltaBase Interface
+	PCG_API virtual PCGIndexing::FPCGIndexCollection FilterCandidates(const UPCGData* InData) const override;
+	PCG_API virtual int32 Resolve(const UPCGData* InData, const PCGIndexing::FPCGIndexCollection& FilteredCandidates, const FPCGDeltaSettings& DeltaSettings) const override;
+	PCG_API virtual FName GetDeltaName() const override;
+	PCG_API virtual bool Apply(UPCGData* InData, int32 ResolvedIndex) const override;
+	virtual bool UsesReplacementData() const override { return true; }
+	PCG_API virtual UPCGData* GetReplacementData() const override;
+	//~ End FPCGDeltaBase Interface
+	
+	/** The original transform of the first point of the polygon. For use with matching the key during override. */
+	UPROPERTY()
+	FTransform OriginalStartTransform = FTransform::Identity;
+
+	/** The polygon that should replace the other polygon.*/
+	UPROPERTY()
+	TObjectPtr<UPCGPolygon2DData> PolygonOverride;
+};
diff --git a/Engine/Plugins/PCG/Source/PCG/Public/Graph/DataOverride/PCGDataOverrideSpline.h b/Engine/Plugins/PCG/Source/PCG/Public/Graph/DataOverride/PCGDataOverrideSpline.h
new file mode 100644
--- /dev/null
+++ b/Engine/Plugins/PCG/Source/PCG/Public/Graph/DataOverride/PCGDataOverrideSpline.h
@@ -0,0 +1,33 @@
+// Copyright Epic Games, Inc. All Rights Reserved.
+
+#pragma once
+
+#include "Graph/DataOverride/PCGDataOverride.h"
+#include "Data/PCGSplineData.h"
+
+#include "PCGDataOverrideSpline.generated.h"
+
+USTRUCT()
+struct FPCGSplineDataDelta : public FPCGDeltaBase
+{
+	GENERATED_BODY()
+
+	static FName GetDeltaNameStatic();
+
+	//~ Begin FPCGDeltaBase Interface
+	PCG_API virtual PCGIndexing::FPCGIndexCollection FilterCandidates(const UPCGData* InData) const override;
+	PCG_API virtual int32 Resolve(const UPCGData* InData, const PCGIndexing::FPCGIndexCollection& FilteredCandidates, const FPCGDeltaSettings& DeltaSettings) const override;
+	PCG_API virtual FName GetDeltaName() const override;
+	PCG_API virtual bool Apply(UPCGData* InData, int32 ResolvedIndex) const override;
+	virtual bool UsesReplacementData() const override { return true; }
+	PCG_API virtual UPCGData* GetReplacementData() const override;
+	//~ End FPCGDeltaBase Interface
+	
+	/** The original transform of the starting point of the spline. For use with matching the key during override. */
+	UPROPERTY()
+	FTransform OriginalStartTransform = FTransform::Identity;
+
+	/** The spline data that should replace the spline*/
+	UPROPERTY()
+	TObjectPtr<UPCGSplineData> SplineOverride;
+};
