[PCG] Inline constant and default value interfaces

#Synced-CL 52142960

diff --git a/Engine/Plugins/PCG/Source/PCG/Private/Elements/Metadata/PCGMetadataOpElementBase.cpp b/Engine/Plugins/PCG/Source/PCG/Private/Elements/Metadata/PCGMetadataOpElementBase.cpp
--- a/Engine/Plugins/PCG/Source/PCG/Private/Elements/Metadata/PCGMetadataOpElementBase.cpp
+++ b/Engine/Plugins/PCG/Source/PCG/Private/Elements/Metadata/PCGMetadataOpElementBase.cpp
@@ -13,8 +13,6 @@
 #include "Metadata/PCGMetadataTypesConstantStruct.h"
 #include "Metadata/Accessors/PCGAttributeAccessorHelpers.h"
 
-#include "Serialization/ArchiveCrc32.h"
-
 #include UE_INLINE_GENERATED_CPP_BY_NAME(PCGMetadataOpElementBase)
 
 #define LOCTEXT_NAMESPACE "PCGMetadataElementBaseElement"
@@ -317,16 +315,19 @@ void UPCGMetadataSettingsBase::SetPinDefaultValueIsActivated(const FName PinLabe
 {
 	if (ensure(IsPinDefaultValueEnabled(PinLabel)))
 	{
-		if (bDirtySettings)
-		{
-			Modify();
-		}
-
 		const FName PropertyName = PCGMetadataBase::Helpers::GetDefaultValuePropertyName(GetInputPinIndex(PinLabel));
-		const bool bPropertyChanged = DefaultValues.SetPropertyActivated(PropertyName, bIsActivated);
-		if (bPropertyChanged && bDirtySettings)
+		if (DefaultValues.IsPropertyActivated(PropertyName) != bIsActivated)
 		{
-			OnSettingsChangedDelegate.Broadcast(this, EPCGChangeType::Node | EPCGChangeType::Edge | EPCGChangeType::Settings);
+			if (bDirtySettings)
+			{
+				Modify();
+			}
+
+			DefaultValues.SetPropertyActivated(PropertyName, bIsActivated);
+			if (bDirtySettings)
+			{
+				OnSettingsChangedDelegate.Broadcast(this, EPCGChangeType::Node | EPCGChangeType::Edge | EPCGChangeType::Settings);
+			}
 		}
 	}
 }
@@ -355,13 +356,13 @@ EPCGMetadataTypes UPCGMetadataSettingsBase::GetPinDefaultValueType(const FName P
 	return EPCGMetadataTypes::Unknown;
 }
 
+#if WITH_EDITOR
 bool UPCGMetadataSettingsBase::IsPinDefaultValueMetadataTypeValid(const FName PinLabel, EPCGMetadataTypes DataType) const
 {
 	bool bHasSpecialRequirement;
 	return IsSupportedInputType(static_cast<uint16>(DataType), GetInputPinIndex(PinLabel), bHasSpecialRequirement);
 }
 
-#if WITH_EDITOR
 void UPCGMetadataSettingsBase::ResetDefaultValues()
 {
 	DefaultValues.Reset();
@@ -413,17 +414,13 @@ EPCGMetadataTypes UPCGMetadataSettingsBase::GetPinInitialDefaultValueType(const
 	return EPCGMetadataTypes::Unknown;
 }
 
+
 bool UPCGMetadataSettingsBase::IsInputPinRequiredByExecution(const UPCGPin* InPin) const
 {
 	// The pin maintains 'required' status unless the default value is both enabled and activated.
 	return InPin && (InPin->IsConnected() || !IsPinDefaultValueEnabled(InPin->Properties.Label) || !IsPinDefaultValueActivated(InPin->Properties.Label));
 }
 
-void UPCGMetadataSettingsBase::AddDefaultValuesToCrc(FArchiveCrc32& Crc32) const
-{
-	const_cast<FPCGDefaultValueContainer&>(DefaultValues).SerializeCrc(Crc32);
-}
-
 TArray<FPCGPinProperties> UPCGMetadataSettingsBase::InputPinProperties() const
 {
 	TArray<FPCGPinProperties> PinProperties;
@@ -706,9 +703,9 @@ bool FPCGMetadataElementBase::PrepareDataInternal(FPCGContext* Context) const
 					// Need to make sure the param data is properly tracked by the context to prevent garbage collection
 					TimeSlicedContext->TrackObject(DefaultData.Data);
 
-					// Need to make sure the param data has at least one entry
+					// @todo_pcg: To be removed once it's confirmed this never triggers.
 					UPCGMetadata* DefaultParamMetadata = CastChecked<UPCGParamData>(DefaultData.Data)->Metadata;
-					if (DefaultParamMetadata->GetLocalItemCount() == 0)
+					if (!ensure(DefaultParamMetadata->GetLocalItemCount() != 0))
 					{
 						DefaultParamMetadata->AddEntry();
 					}
@@ -1052,25 +1049,4 @@ bool FPCGMetadataElementBase::ExecuteInternal(FPCGContext* Context) const
 	});
 }
 
-void FPCGMetadataElementBase::GetDependenciesCrc(const FPCGGetDependenciesCrcParams& InParams, FPCGCrc& OutCrc) const
-{
-	FPCGCrc Crc;
-	IPCGElement::GetDependenciesCrc(InParams, Crc);
-
-	// Add the default values to the crc
-	if (const UPCGMetadataSettingsBase* Settings = Cast<UPCGMetadataSettingsBase>(InParams.Settings))
-	{
-		FArchiveCrc32 Crc32;
-		/*
-		 * Implementation note: In theory, the default value behind connected pins should not factor into the Crc, but in practice,
-		 * it makes the code more obscure for a gain that might not even be meaningful.
-		 */
-		Crc32.SetIsSaving(true);
-		Settings->AddDefaultValuesToCrc(Crc32);
-		Crc.Combine(Crc32.GetCrc());
-	}
-
-	OutCrc = Crc;
-}
-
 #undef LOCTEXT_NAMESPACE
diff --git a/Engine/Plugins/PCG/Source/PCG/Private/Helpers/PCGDefaultValueContainer.cpp b/Engine/Plugins/PCG/Source/PCG/Private/Helpers/PCGDefaultValueContainer.cpp
--- a/Engine/Plugins/PCG/Source/PCG/Private/Helpers/PCGDefaultValueContainer.cpp
+++ b/Engine/Plugins/PCG/Source/PCG/Private/Helpers/PCGDefaultValueContainer.cpp
@@ -58,13 +58,14 @@ FString FPCGDefaultValueContainer::GetPropertyValueAsString(const FName Property
 	return Result.HasValue() ? Result.GetValue() : FString(TEXT("Error"));
 }
 
-const UPCGParamData* FPCGDefaultValueContainer::CreateParamData(FPCGContext* Context, const FName PropertyName) const
+const UPCGParamData* FPCGDefaultValueContainer::CreateParamData(FPCGContext* Context, const FName PropertyKey, const FName AttributeName) const
 {
-	if (const FProperty* PropertyPtr = FindProperty(PropertyName))
+	if (const FProperty* PropertyPtr = FindProperty(PropertyKey))
 	{
 		const TObjectPtr<UPCGParamData> NewParamData = FPCGContext::NewObject_AnyThread<UPCGParamData>(Context);
-		if (NewParamData->Metadata->CreateAttributeFromDataProperty(NAME_None, PropertyBag.GetValue().GetMemory(), PropertyPtr))
+		if (NewParamData->Metadata->CreateAttributeFromDataProperty(AttributeName, PropertyBag.GetValue().GetMemory(), PropertyPtr))
 		{
+			NewParamData->Metadata->AddEntry();
 			return NewParamData;
 		}
 	}
@@ -72,6 +73,12 @@ const UPCGParamData* FPCGDefaultValueContainer::CreateParamData(FPCGContext* Con
 	return nullptr;
 }
 
+void FPCGDefaultValueContainer::SerializeCrc(FArchiveCrc32& Crc32) const
+{
+	// const_cast safe: FArchiveCrc32 is read-only, FInstancedPropertyBag::Serialize is non-const due to FArchive API.
+	const_cast<FInstancedPropertyBag&>(PropertyBag).Serialize(Crc32);
+}
+
 bool FPCGDefaultValueContainer::IsPropertyActivated(const FName PropertyName) const
 {
 	return ActivatedProperties.Contains(PropertyName);
@@ -96,7 +103,8 @@ bool FPCGDefaultValueContainer::SetPropertyValueFromString(const FName PropertyN
 	}
 
 	SetPropertyActivated(PropertyName, /*bIsActive=*/true);
-	return PropertyBag.SetValueSerializedString(PropertyName, ValueString) == EPropertyBagResult::Success;
+	const EPropertyBagResult Result = PropertyBag.SetValueSerializedString(PropertyName, ValueString);
+	return Result == EPropertyBagResult::Success;
 }
 
 bool FPCGDefaultValueContainer::SetPropertyActivated(const FName PropertyName, const bool bIsActivated)
diff --git a/Engine/Plugins/PCG/Source/PCG/Private/Metadata/PCGInlineConstantInterface.cpp b/Engine/Plugins/PCG/Source/PCG/Private/Metadata/PCGInlineConstantInterface.cpp
new file mode 100644
--- /dev/null
+++ b/Engine/Plugins/PCG/Source/PCG/Private/Metadata/PCGInlineConstantInterface.cpp
@@ -0,0 +1,431 @@
+// Copyright Epic Games, Inc. All Rights Reserved.
+
+#include "Metadata/PCGInlineConstantInterface.h"
+
+#include "PCGParamData.h"
+#include "PCGSettings.h"
+#include "Helpers/PCGMetadataHelpers.h"
+
+#include UE_INLINE_GENERATED_CPP_BY_NAME(PCGInlineConstantInterface)
+
+#define LOCTEXT_NAMESPACE "PCGInlineConstantInterface"
+
+bool FPCGInlineConstantState::DefaultValuesAreEnabled() const
+{
+	return true;
+}
+
+bool FPCGInlineConstantState::IsPinDefaultValueEnabled(const FName PinLabel) const
+{
+	return DefaultValueMap.Contains(PinLabel) && DefaultValues.FindProperty(PinLabel);
+}
+
+bool FPCGInlineConstantState::IsPinDefaultValueActivated(const FName PinLabel) const
+{
+	return IsPinDefaultValueEnabled(PinLabel) && DefaultValues.IsPropertyActivated(PinLabel);
+}
+
+EPCGMetadataTypes FPCGInlineConstantState::GetPinDefaultValueType(const FName PinLabel) const
+{
+	const FPCGPinDefaultValueInfo* DefaultValueInfo = DefaultValueMap.Find(PinLabel);
+	return DefaultValueInfo ? DefaultValueInfo->CurrentType : EPCGMetadataTypes::Unknown;
+}
+
+const UPCGParamData* FPCGInlineConstantState::CreateDefaultValueParamData(FPCGContext* InContext, const FName PropertyKey, const FName AttributeName) const
+{
+	const FPCGPinDefaultValueInfo* DefaultValueInfo = DefaultValueMap.Find(PropertyKey);
+	if (!DefaultValueInfo)
+	{
+		return nullptr;
+	}
+
+	if (!ensure(DefaultValues.GetCurrentPropertyType(PropertyKey) == DefaultValueInfo->CurrentType))
+	{
+		return nullptr;
+	}
+
+	return DefaultValues.CreateParamData(InContext, PropertyKey, AttributeName);
+}
+
+#if WITH_EDITOR
+bool FPCGInlineConstantState::IsPinDefaultValueMetadataTypeValid(const FName PinLabel, const EPCGMetadataTypes DataType) const
+{
+	const FPCGPinDefaultValueInfo* DefaultValueInfo = DefaultValueMap.Find(PinLabel);
+	return DefaultValueInfo && (DefaultValueInfo->AvailableTypes.IsEmpty() || DefaultValueInfo->AvailableTypes.Contains(DataType));
+}
+
+bool FPCGInlineConstantState::PinIsMapped(const FName PinLabel) const
+{
+	return DefaultValueMap.Contains(PinLabel);
+}
+
+EPCGMetadataTypes FPCGInlineConstantState::GetPinInitialDefaultValueType(const FName PinLabel) const
+{
+	const FPCGPinDefaultValueInfo* DefaultValueInfo = DefaultValueMap.Find(PinLabel);
+	return DefaultValueInfo ? DefaultValueInfo->InitialType : EPCGMetadataTypes::Unknown;
+}
+
+FString FPCGInlineConstantState::GetPinDefaultValueAsString(const FName PinLabel) const
+{
+	if (PinLabel != NAME_None)
+	{
+		return DefaultValues.FindProperty(PinLabel) ? DefaultValues.GetPropertyValueAsString(PinLabel) : GetPinInitialDefaultValueString(PinLabel);
+	}
+
+	return FString{};
+}
+
+FString FPCGInlineConstantState::GetPinInitialDefaultValueString(const FName PinLabel) const
+{
+	if (const FPCGPinDefaultValueInfo* DefaultValueInfo = PinLabel != NAME_None ? DefaultValueMap.Find(PinLabel) : nullptr)
+	{
+		return PCGMetadataAttribute::CallbackWithRightType(static_cast<uint16>(DefaultValueInfo->InitialType), []<typename T>(const T&)
+		{
+			const T Value = PCG::Private::MetadataTraits<T>::ZeroValue();
+			if constexpr (requires { TBaseStructure<T>::Get(); })
+			{
+				FString Result;
+				TBaseStructure<T>::Get()->ExportText(Result, &Value, /*Defaults=*/nullptr, /*OwnerObject=*/nullptr, PPF_None, /*ExportRootScope=*/nullptr);
+				return Result;
+			}
+			else
+			{
+				return LexToString(Value);
+			}
+		});
+	}
+
+	return FString{};
+}
+
+EPCGSettingDefaultValueExtraFlags FPCGInlineConstantState::GetDefaultValueExtraFlags(const FName PinLabel) const
+{
+	const FPCGPinDefaultValueInfo* DefaultValueInfo = DefaultValueMap.Find(PinLabel);
+	return DefaultValueInfo ? DefaultValueInfo->ExtraFlags : EPCGSettingDefaultValueExtraFlags::None;
+}
+
+bool FPCGInlineConstantState::AddDefaultValueToPin(const FName PinLabel, FPCGPinDefaultValueInfo::FInitParams InitParams)
+{
+	TRACE_CPUPROFILER_EVENT_SCOPE(FPCGInlineConstantState::AddDefaultValueToPin);
+
+	if (DefaultValueMap.Contains(PinLabel))
+	{
+		return false;
+	}
+
+	DefaultValues.CreateNewProperty(PinLabel, InitParams.InitialType);
+	DefaultValues.SetPropertyActivated(PinLabel, InitParams.bStartActivated);
+	DefaultValueMap.Add(PinLabel, FPCGPinDefaultValueInfo(MoveTemp(InitParams)));
+
+	return true;
+}
+
+bool FPCGInlineConstantState::RemoveDefaultValueFromPin(const FName PinLabel)
+{
+	TRACE_CPUPROFILER_EVENT_SCOPE(FPCGInlineConstantState::RemoveDefaultValueFromPin);
+
+	if (!DefaultValueMap.Contains(PinLabel))
+	{
+		return false;
+	}
+
+	DefaultValueMap.Remove(PinLabel);
+	DefaultValues.RemoveProperty(PinLabel);
+
+	return true;
+}
+
+bool FPCGInlineConstantState::ResetDefaultValues()
+{
+	bool bDefaultValuesReset = false;
+
+	// Validate a reset is needed
+	for (const TTuple<FName, FPCGPinDefaultValueInfo>& DefaultValue : DefaultValueMap)
+	{
+		if (DefaultValue.Key != NAME_None && DefaultValues.FindProperty(DefaultValue.Key))
+		{
+			bDefaultValuesReset = true;
+			break;
+		}
+	}
+
+	if (!bDefaultValuesReset)
+	{
+		return false;
+	}
+
+	for (const TTuple<FName, FPCGPinDefaultValueInfo>& DefaultValue : DefaultValueMap)
+	{
+		ResetDefaultValueInternal(DefaultValue.Key);
+	}
+
+	return true;
+}
+
+bool FPCGInlineConstantState::ResetDefaultValue(const FName PinLabel)
+{
+	if (PinLabel == NAME_None || !DefaultValues.FindProperty(PinLabel) || !DefaultValueMap.Contains(PinLabel))
+	{
+		return false;
+	}
+
+	ResetDefaultValueInternal(PinLabel);
+	return true;
+}
+
+bool FPCGInlineConstantState::SetPinDefaultValue(const FName PinLabel, const FString& DefaultValue, const bool bCreateIfNeeded)
+{
+	TRACE_CPUPROFILER_EVENT_SCOPE(FPCGInlineConstantState::SetPinDefaultValue);
+
+	if (PinLabel == NAME_None)
+	{
+		return false;
+	}
+
+	const FProperty* Property = DefaultValues.FindProperty(PinLabel);
+	if (Property && DefaultValue == DefaultValues.GetPropertyValueAsString(PinLabel))
+	{
+		return false;
+	}
+
+	if (!Property && bCreateIfNeeded)
+	{
+		const EPCGMetadataTypes Type = GetPinInitialDefaultValueType(PinLabel);
+		DefaultValues.CreateNewProperty(PinLabel, Type);
+	}
+
+	return DefaultValues.SetPropertyValueFromString(PinLabel, DefaultValue);
+}
+
+bool FPCGInlineConstantState::ConvertPinDefaultValueMetadataType(const FName PinLabel, const EPCGMetadataTypes DataType)
+{
+	TRACE_CPUPROFILER_EVENT_SCOPE(FPCGInlineConstantState::ConvertPinDefaultValueMetadataType);
+
+	if (!ensure(IsPinDefaultValueActivated(PinLabel)))
+	{
+		return false;
+	}
+
+	if (PinLabel == NAME_None || !IsPinDefaultValueMetadataTypeValid(PinLabel, DataType))
+	{
+		return false;
+	}
+
+	if (DefaultValues.ConvertPropertyType(PinLabel, DataType))
+	{
+		DefaultValueMap[PinLabel].CurrentType = DataType;
+	}
+
+	return true;
+}
+
+bool FPCGInlineConstantState::SetPinDefaultValueIsActivated(const FName PinLabel, const bool bIsActivated)
+{
+	if (!ensure(IsPinDefaultValueEnabled(PinLabel)))
+	{
+		return false;
+	}
+
+	if (DefaultValues.IsPropertyActivated(PinLabel) == bIsActivated)
+	{
+		return false;
+	}
+
+	DefaultValues.SetPropertyActivated(PinLabel, bIsActivated);
+
+	return true;
+}
+
+void FPCGInlineConstantState::ResetDefaultValueInternal(const FName PinLabel)
+{
+	if (PinLabel != NAME_None && DefaultValues.FindProperty(PinLabel))
+	{
+		if (!ensure(DefaultValueMap.Contains(PinLabel)))
+		{
+			return;
+		}
+
+		DefaultValues.RemoveProperty(PinLabel);
+		DefaultValues.CreateNewProperty(PinLabel, DefaultValueMap[PinLabel].InitialType);
+		DefaultValueMap[PinLabel].CurrentType = DefaultValueMap[PinLabel].InitialType;
+	}
+}
+#endif // WITH_EDITOR
+
+const UPCGParamData* IPCGSettingsInlineConstant::CreateDefaultValueParamData(FPCGContext* InContext, const FName PropertyKey, const FName AttributeName) const
+{
+	const FPCGInlineConstantState* State = GetInlineConstantState();
+	return State ? State->CreateDefaultValueParamData(InContext, PropertyKey, AttributeName) : nullptr;
+}
+
+bool IPCGSettingsInlineConstant::DefaultValuesAreEnabled() const
+{
+	const FPCGInlineConstantState* State = GetInlineConstantState();
+	return State && State->DefaultValuesAreEnabled();
+}
+
+bool IPCGSettingsInlineConstant::IsPinDefaultValueEnabled(const FName PinLabel) const
+{
+	const FPCGInlineConstantState* State = GetInlineConstantState();
+	if (State && State->IsPinDefaultValueEnabled(PinLabel))
+	{
+		return true;
+	}
+
+	// Fallback here supporting pins that are eligible if their initial type supports default values (pre-activation).
+	const EPCGMetadataTypes InitialType = GetPinInitialDefaultValueType(PinLabel);
+	return InitialType != EPCGMetadataTypes::Unknown && PCGMetadataHelpers::MetadataTypeSupportsDefaultValues(InitialType);
+}
+
+bool IPCGSettingsInlineConstant::IsPinDefaultValueActivated(const FName PinLabel) const
+{
+	const FPCGInlineConstantState* State = GetInlineConstantState();
+	return State && State->IsPinDefaultValueActivated(PinLabel);
+}
+
+EPCGMetadataTypes IPCGSettingsInlineConstant::GetPinDefaultValueType(const FName PinLabel) const
+{
+	const FPCGInlineConstantState* State = GetInlineConstantState();
+	return State ? State->GetPinDefaultValueType(PinLabel) : EPCGMetadataTypes::Unknown;
+}
+
+#if WITH_EDITOR
+bool IPCGSettingsInlineConstant::IsPinDefaultValueMetadataTypeValid(const FName PinLabel, const EPCGMetadataTypes DataType) const
+{
+	const FPCGInlineConstantState* State = GetInlineConstantState();
+	return State && State->IsPinDefaultValueMetadataTypeValid(PinLabel, DataType);
+}
+
+void IPCGSettingsInlineConstant::ResetDefaultValues()
+{
+	if (FPCGInlineConstantState* State = GetMutableInlineConstantState())
+	{
+		if (UPCGSettings* Settings = Cast<UPCGSettings>(this))
+		{
+			Settings->Modify();
+
+			if (State->ResetDefaultValues())
+			{
+				Settings->OnSettingsChangedDelegate.Broadcast(Settings, EPCGChangeType::Settings);
+			}
+		}
+	}
+}
+
+void IPCGSettingsInlineConstant::ResetDefaultValue(const FName PinLabel)
+{
+	if (FPCGInlineConstantState* State = GetMutableInlineConstantState())
+	{
+		if (UPCGSettings* Settings = Cast<UPCGSettings>(this))
+		{
+			Settings->Modify();
+
+			if (State->ResetDefaultValue(PinLabel))
+			{
+				Settings->OnSettingsChangedDelegate.Broadcast(Settings, EPCGChangeType::Settings);
+			}
+		}
+	}
+}
+
+void IPCGSettingsInlineConstant::SetPinDefaultValue(const FName PinLabel, const FString& DefaultValue, const bool bCreateIfNeeded)
+{
+	if (FPCGInlineConstantState* State = GetMutableInlineConstantState())
+	{
+		UPCGSettings* Settings = Cast<UPCGSettings>(this);
+		if (Settings)
+		{
+			Settings->Modify();
+		}
+
+		if (State->SetPinDefaultValue(PinLabel, DefaultValue, bCreateIfNeeded))
+		{
+			if (Settings)
+			{
+				Settings->OnSettingsChangedDelegate.Broadcast(Settings, EPCGChangeType::Settings);
+			}
+		}
+	}
+}
+
+void IPCGSettingsInlineConstant::ConvertPinDefaultValueMetadataType(const FName PinLabel, const EPCGMetadataTypes DataType)
+{
+	if (FPCGInlineConstantState* State = GetMutableInlineConstantState())
+	{
+		UPCGSettings* Settings = Cast<UPCGSettings>(this);
+		if (Settings)
+		{
+			Settings->Modify();
+		}
+
+		if (State->ConvertPinDefaultValueMetadataType(PinLabel, DataType))
+		{
+			if (Settings)
+			{
+				Settings->OnSettingsChangedDelegate.Broadcast(Settings, EPCGChangeType::Node | EPCGChangeType::Edge | EPCGChangeType::Settings);
+			}
+		}
+	}
+}
+
+void IPCGSettingsInlineConstant::SetPinDefaultValueIsActivated(const FName PinLabel, const bool bIsActivated, const bool bDirtySettings)
+{
+	FPCGInlineConstantState* State = GetMutableInlineConstantState();
+	if (!State)
+	{
+		return;
+	}
+
+	// Lazy initialization -- if the pin hasn't been registered in state yet, create it now.
+	if (!State->PinIsMapped(PinLabel))
+	{
+		const EPCGMetadataTypes InitialType = GetPinInitialDefaultValueType(PinLabel);
+		if (InitialType == EPCGMetadataTypes::Unknown || !PCGMetadataHelpers::MetadataTypeSupportsDefaultValues(InitialType))
+		{
+			return;
+		}
+
+		FPCGPinDefaultValueInfo::FInitParams InitParams;
+		InitParams.InitialType = InitialType;
+		InitParams.bStartActivated = false;
+		State->AddDefaultValueToPin(PinLabel, MoveTemp(InitParams));
+	}
+
+	if (ensure(State->IsPinDefaultValueEnabled(PinLabel)))
+	{
+		UPCGSettings* Settings = bDirtySettings ? Cast<UPCGSettings>(this) : nullptr;
+		if (Settings)
+		{
+			Settings->Modify();
+		}
+
+		if (State->SetPinDefaultValueIsActivated(PinLabel, bIsActivated))
+		{
+			if (Settings)
+			{
+				Settings->OnSettingsChangedDelegate.Broadcast(Settings, EPCGChangeType::Node | EPCGChangeType::Edge | EPCGChangeType::Settings);
+			}
+		}
+	}
+}
+
+FString IPCGSettingsInlineConstant::GetPinDefaultValueAsString(const FName PinLabel) const
+{
+	const FPCGInlineConstantState* State = GetInlineConstantState();
+	return State ? State->GetPinDefaultValueAsString(PinLabel) : FString{};
+}
+
+FString IPCGSettingsInlineConstant::GetPinInitialDefaultValueString(const FName PinLabel) const
+{
+	const FPCGInlineConstantState* State = GetInlineConstantState();
+	return State ? State->GetPinInitialDefaultValueString(PinLabel) : FString{};
+}
+
+EPCGSettingDefaultValueExtraFlags IPCGSettingsInlineConstant::GetDefaultValueExtraFlags(const FName PinLabel) const
+{
+	const FPCGInlineConstantState* State = GetInlineConstantState();
+	return State ? State->GetDefaultValueExtraFlags(PinLabel) : EPCGSettingDefaultValueExtraFlags::None;
+}
+#endif // WITH_EDITOR
+
+#undef LOCTEXT_NAMESPACE
diff --git a/Engine/Plugins/PCG/Source/PCG/Private/PCGSettings.cpp b/Engine/Plugins/PCG/Source/PCG/Private/PCGSettings.cpp
--- a/Engine/Plugins/PCG/Source/PCG/Private/PCGSettings.cpp
+++ b/Engine/Plugins/PCG/Source/PCG/Private/PCGSettings.cpp
@@ -385,6 +385,7 @@ void UPCGSettings::PostInitProperties()
 
 #if WITH_EDITOR
 	InitializeCachedOverridableParams();
+	OnSettingsChangedDelegate.AddUObject(this, &UPCGSettings::OnSettingsChanged);
 #endif //WITH_EDITOR
 
 	InitializeOverriddenSettingsPropertyBag();
@@ -1352,6 +1353,17 @@ void UPCGSettings::CacheCrc()
 #endif
 }
 
+#if WITH_EDITOR
+void UPCGSettings::OnSettingsChanged(UPCGSettings* InSettings, EPCGChangeType ChangeType)
+{
+	if (ensure(InSettings == this) && EnumHasAnyFlags(ChangeType, EPCGChangeType::Settings))
+	{
+		// Clear the cached crc if an external call has settings changed externally.
+		CachedCrc = FPCGCrc();
+	}
+}
+#endif
+
 TArray<FPCGPinProperties> UPCGSettings::DefaultPointInputPinProperties() const
 {
 	TArray<FPCGPinProperties> Properties;
diff --git a/Engine/Plugins/PCG/Source/PCG/Public/Elements/Metadata/PCGMetadataMathsOpElement.h b/Engine/Plugins/PCG/Source/PCG/Public/Elements/Metadata/PCGMetadataMathsOpElement.h
--- a/Engine/Plugins/PCG/Source/PCG/Public/Elements/Metadata/PCGMetadataMathsOpElement.h
+++ b/Engine/Plugins/PCG/Source/PCG/Public/Elements/Metadata/PCGMetadataMathsOpElement.h
@@ -81,7 +81,9 @@ public:
 
 	//~Begin IPCGSettingsDefaultValueProvider interface
 	virtual EPCGMetadataTypes GetPinInitialDefaultValueType(FName PinLabel) const override { return EPCGMetadataTypes::Double; }
+#if WITH_EDITOR
 	virtual bool IsPinDefaultValueMetadataTypeValid(FName PinLabel, EPCGMetadataTypes DataType) const override;
+#endif
 	//~End IPCGSettingsDefaultValueProvider interface
 
 protected:
diff --git a/Engine/Plugins/PCG/Source/PCG/Public/Elements/Metadata/PCGMetadataOpElementBase.h b/Engine/Plugins/PCG/Source/PCG/Public/Elements/Metadata/PCGMetadataOpElementBase.h
--- a/Engine/Plugins/PCG/Source/PCG/Public/Elements/Metadata/PCGMetadataOpElementBase.h
+++ b/Engine/Plugins/PCG/Source/PCG/Public/Elements/Metadata/PCGMetadataOpElementBase.h
@@ -104,8 +104,8 @@ public:
 	UE_API virtual bool IsPinDefaultValueEnabled(FName PinLabel) const override;
 	UE_API virtual bool IsPinDefaultValueActivated(FName PinLabel) const override;
 	UE_API virtual EPCGMetadataTypes GetPinDefaultValueType(FName PinLabel) const override;
-	UE_API virtual bool IsPinDefaultValueMetadataTypeValid(FName PinLabel, EPCGMetadataTypes DataType) const override;
 #if WITH_EDITOR
+	UE_API virtual bool IsPinDefaultValueMetadataTypeValid(FName PinLabel, EPCGMetadataTypes DataType) const override;
 	UE_API virtual void SetPinDefaultValue(FName PinLabel, const FString& DefaultValue, bool bCreateIfNeeded = false) override;
 	UE_API virtual void ConvertPinDefaultValueMetadataType(FName PinLabel, EPCGMetadataTypes DataType) override;
 	UE_API virtual void SetPinDefaultValueIsActivated(FName PinLabel, bool bIsActivated, bool bDirtySettings = true) override;
@@ -138,8 +138,8 @@ public:
 	virtual bool CanCullTaskIfUnwired() const override { return false; }
 	UE_API virtual bool IsInputPinRequiredByExecution(const UPCGPin* InPin) const override;
 
-	/** Adds the default values to the Crc for caching. */
-	UE_API void AddDefaultValuesToCrc(FArchiveCrc32& Crc32) const;
+	UE_DEPRECATED(5.8, "No longer needed - default values are included in settings CRC automatically via UPROPERTY serialization.")
+	void AddDefaultValuesToCrc(FArchiveCrc32& Crc32) const {}
 
 protected:
 	UE_API virtual TArray<FPCGPinProperties> InputPinProperties() const override;
@@ -263,7 +263,6 @@ protected:
 	virtual bool ExecuteInternal(FPCGContext* Context) const override;
 	virtual EPCGElementExecutionLoopMode ExecutionLoopMode(const UPCGSettings* Settings) const override { return EPCGElementExecutionLoopMode::SinglePrimaryPin; }
 	virtual bool SupportsBasePointDataInputs(FPCGContext* InContext) const override { return true; }
-	virtual void GetDependenciesCrc(const FPCGGetDependenciesCrcParams& InParams, FPCGCrc& OutCrc) const override;
 
 	virtual bool DoOperation(PCGMetadataOps::FOperationData& InOperationData) const = 0;
 
diff --git a/Engine/Plugins/PCG/Source/PCG/Public/Elements/Metadata/PCGMetadataRotatorOpElement.h b/Engine/Plugins/PCG/Source/PCG/Public/Elements/Metadata/PCGMetadataRotatorOpElement.h
--- a/Engine/Plugins/PCG/Source/PCG/Public/Elements/Metadata/PCGMetadataRotatorOpElement.h
+++ b/Engine/Plugins/PCG/Source/PCG/Public/Elements/Metadata/PCGMetadataRotatorOpElement.h
@@ -65,7 +65,9 @@ protected:
 	//~End UPCGSettings interface
 
 	//~Begin IPCGSettingsDefaultValueProvider interface
+#if WITH_EDITOR
 	virtual bool IsPinDefaultValueMetadataTypeValid(FName PinLabel, EPCGMetadataTypes DataType) const override;
+#endif
 	virtual EPCGMetadataTypes GetPinInitialDefaultValueType(FName PinLabel) const override;
 	//~End IPCGSettingsDefaultValueProvider interface
 
diff --git a/Engine/Plugins/PCG/Source/PCG/Public/Helpers/PCGDefaultValueContainer.h b/Engine/Plugins/PCG/Source/PCG/Public/Helpers/PCGDefaultValueContainer.h
--- a/Engine/Plugins/PCG/Source/PCG/Public/Helpers/PCGDefaultValueContainer.h
+++ b/Engine/Plugins/PCG/Source/PCG/Public/Helpers/PCGDefaultValueContainer.h
@@ -33,10 +33,10 @@ struct FPCGDefaultValueContainer
 	FString GetPropertyValueAsString(FName PropertyName) const;
 
 	/** Serialize the properties into a Crc32 to be used for caching purposes. */
-	void SerializeCrc(FArchiveCrc32& Crc32) { PropertyBag.Serialize(Crc32); }
+	void SerializeCrc(FArchiveCrc32& Crc32) const;
 
 	/** Helper that creates a Param Data with the typed property as an attribute attached to the metadata. */
-	const UPCGParamData* CreateParamData(FPCGContext* Context, FName PropertyName) const;
+	const UPCGParamData* CreateParamData(FPCGContext* Context, FName PropertyKey, FName AttributeName = NAME_None) const;
 
 	/** The property exists and is set to active. */
 	bool IsPropertyActivated(FName PropertyName) const;
diff --git a/Engine/Plugins/PCG/Source/PCG/Public/Metadata/PCGDefaultValueInterface.h b/Engine/Plugins/PCG/Source/PCG/Public/Metadata/PCGDefaultValueInterface.h
--- a/Engine/Plugins/PCG/Source/PCG/Public/Metadata/PCGDefaultValueInterface.h
+++ b/Engine/Plugins/PCG/Source/PCG/Public/Metadata/PCGDefaultValueInterface.h
@@ -19,6 +19,49 @@ enum class EPCGSettingDefaultValueExtraFlags : uint64
 
 ENUM_CLASS_FLAGS(EPCGSettingDefaultValueExtraFlags);
 
+USTRUCT()
+struct FPCGPinDefaultValueInfo
+{
+	GENERATED_BODY()
+
+#if WITH_EDITOR
+	struct FInitParams
+	{
+		bool bStartActivated = true;
+		EPCGMetadataTypes InitialType = EPCGMetadataTypes::Double;
+		FString InitialValue = PCG::Private::MetadataTraits<double>::ZeroValueString();
+		TArray<EPCGMetadataTypes> AvailableTypes;
+	};
+
+	FPCGPinDefaultValueInfo() = default;
+
+	explicit FPCGPinDefaultValueInfo(FInitParams InParams)
+		: CurrentType(InParams.InitialType)
+		, AvailableTypes(MoveTemp(InParams.AvailableTypes))
+		, InitialType(InParams.InitialType)
+		, InitialValue(MoveTemp(InParams.InitialValue))
+	{}
+#endif // WITH_EDITOR
+
+	UPROPERTY()
+	EPCGMetadataTypes CurrentType = EPCGMetadataTypes::Unknown;
+
+#if WITH_EDITORONLY_DATA
+	/** Metadata types this pin accepts. Empty = all types are supported. */
+	UPROPERTY()
+	TArray<EPCGMetadataTypes> AvailableTypes;
+
+	UPROPERTY()
+	EPCGMetadataTypes InitialType = EPCGMetadataTypes::Unknown;
+
+	UPROPERTY()
+	FString InitialValue;
+
+	UPROPERTY()
+	EPCGSettingDefaultValueExtraFlags ExtraFlags = EPCGSettingDefaultValueExtraFlags::None;
+#endif // WITH_EDITORONLY_DATA
+};
+
 UINTERFACE(MinimalAPI)
 class UPCGSettingsDefaultValueProvider : public UInterface
 {
@@ -33,7 +76,7 @@ public:
 	/** One or more pins on this node has a 'default value' and can be adjusted via an inline constant. */
 	virtual bool DefaultValuesAreEnabled() const { return false; }
 
-	/** The specified pin can accomodate 'default value' inline constants. */
+	/** The specified pin can accommodate 'default value' inline constants. */
 	virtual bool IsPinDefaultValueEnabled(FName PinLabel) const { return false; }
 
 	/** The specified pin has a 'default value' currently activated. */
@@ -42,9 +85,6 @@ public:
 	/** Gets the current 'default value' type, if supported, for the pin. */
 	virtual EPCGMetadataTypes GetPinDefaultValueType(FName PinLabel) const { return EPCGMetadataTypes::Unknown; }
 
-	/** Whether the pin supports the provided metadata type. */
-	virtual bool IsPinDefaultValueMetadataTypeValid(FName PinLabel, EPCGMetadataTypes DataType) const { return false; }
-
 	/** Add an attribute to a given metadata for the initial default value. In most cases, this will be the Zero Value. */
 	virtual bool CreateInitialDefaultValueAttribute(FName PinLabel, UPCGMetadata* OutMetadata) const
 	{
@@ -57,6 +97,9 @@ public:
 	}
 
 #if WITH_EDITOR
+	/** Returns true if the pin supports the provided metadata type. */
+	virtual bool IsPinDefaultValueMetadataTypeValid(FName PinLabel, EPCGMetadataTypes DataType) const { return false; }
+
 	/** Resets all default values to their 'reset' value and deactivates them. */
 	virtual void ResetDefaultValues() {}
 
@@ -69,6 +112,9 @@ public:
 	/** Attempts a metadata type conversion of the pin's default value inline constant. */
 	virtual void ConvertPinDefaultValueMetadataType(FName PinLabel, EPCGMetadataTypes DataType) {}
 
+	/** Whether the default value can be toggled on/off by the user. */
+	virtual bool CanDefaultValueBeToggled() const { return true; }
+
 	/** Sets the default value to active. Must be overridden by the subclass. */
 	virtual void SetPinDefaultValueIsActivated(FName PinLabel, bool bIsActivated, bool bDirtySettings = true) PURE_VIRTUAL(IPCGSettingsDefaultValueProvider::SetPinDefaultValueIsActivated, return;);
 
@@ -83,6 +129,7 @@ public:
 #endif // WITH_EDITOR
 
 protected:
+	// @todo_pcg: GetPinInitialDefaultValueType should likely be Editor Only
 	/** For the initial 'default value' type of the pin. */
 	virtual EPCGMetadataTypes GetPinInitialDefaultValueType(FName PinLabel) const { return EPCGMetadataTypes::Unknown; }
-};
\ No newline at end of file
+};
diff --git a/Engine/Plugins/PCG/Source/PCG/Public/Metadata/PCGInlineConstantInterface.h b/Engine/Plugins/PCG/Source/PCG/Public/Metadata/PCGInlineConstantInterface.h
new file mode 100644
--- /dev/null
+++ b/Engine/Plugins/PCG/Source/PCG/Public/Metadata/PCGInlineConstantInterface.h
@@ -0,0 +1,142 @@
+// Copyright Epic Games, Inc. All Rights Reserved.
+
+#pragma once
+
+#include "Helpers/PCGDefaultValueContainer.h"
+#include "Metadata/PCGDefaultValueInterface.h"
+
+#include "PCGInlineConstantInterface.generated.h"
+
+class UPCGParamData;
+struct FPCGContext;
+
+/**
+ * Composable state struct for inline constant values on pins.
+ * Any UPCGSettings class implementing IPCGSettingsInlineConstant can embed this as a UPROPERTY member and override
+ * GetInlineConstantState()/GetMutableInlineConstantState() to get full inline constant support with minimal boilerplate.
+ */
+USTRUCT()
+struct FPCGInlineConstantState
+{
+	GENERATED_BODY()
+
+	/** One or more pins on this node has a 'default value' and can be adjusted via an inline constant. */
+	PCG_API bool DefaultValuesAreEnabled() const;
+
+	/** The specified pin can accommodate 'default value' inline constants. */
+	PCG_API bool IsPinDefaultValueEnabled(FName PinLabel) const;
+
+	/** The specified pin has a 'default value' currently activated. */
+	PCG_API bool IsPinDefaultValueActivated(FName PinLabel) const;
+
+	/** Get the current 'default value' type, if supported, for the pin. */
+	PCG_API EPCGMetadataTypes GetPinDefaultValueType(FName PinLabel) const;
+
+	/** Create a Param Data with the inline constant default value properties inserted as metadata. */
+	PCG_API const UPCGParamData* CreateDefaultValueParamData(FPCGContext* InContext, FName PropertyKey, FName AttributeName = NAME_None) const;
+
+#if WITH_EDITOR
+	/** Returns true if the pin supports the provided metadata type. */
+	PCG_API bool IsPinDefaultValueMetadataTypeValid(FName PinLabel, EPCGMetadataTypes DataType) const;
+
+	/** Returns true if the pin has been registered for default values. */
+	PCG_API bool PinIsMapped(FName PinLabel) const;
+
+	/** For the initial 'default value' type of the pin. */
+	PCG_API EPCGMetadataTypes GetPinInitialDefaultValueType(FName PinLabel) const;
+
+	/** Get the 'default value', if supported, for the pin. */
+	PCG_API FString GetPinDefaultValueAsString(FName PinLabel) const;
+
+	/** Get the initial 'default value' of the pin. */
+	PCG_API FString GetPinInitialDefaultValueString(FName PinLabel) const;
+
+	/** Extra flags related to displaying the value. */
+	PCG_API EPCGSettingDefaultValueExtraFlags GetDefaultValueExtraFlags(FName PinLabel) const;
+
+	/** Register a pin for inline constant support. Returns false if already registered. */
+	PCG_API bool AddDefaultValueToPin(FName PinLabel, FPCGPinDefaultValueInfo::FInitParams InitParams = FPCGPinDefaultValueInfo::FInitParams{});
+
+	/** Unregister a pin from inline constant support. Returns false if not registered. */
+	PCG_API bool RemoveDefaultValueFromPin(FName PinLabel);
+
+	/** Reset all values to their initial type/value. Returns true if any reset occurred. */
+	PCG_API bool ResetDefaultValues();
+
+	/** Reset a single pin's value. Returns true if reset occurred. */
+	PCG_API bool ResetDefaultValue(FName PinLabel);
+
+	/** Set the pin's default value string directly. Returns true if value changed. */
+	PCG_API bool SetPinDefaultValue(FName PinLabel, const FString& DefaultValue, bool bCreateIfNeeded = false);
+
+	/** Attempt a metadata type conversion. Returns true if conversion occurred. */
+	PCG_API bool ConvertPinDefaultValueMetadataType(FName PinLabel, EPCGMetadataTypes DataType);
+
+	/** Set the default value to active/inactive. Returns true if activation state changed. */
+	PCG_API bool SetPinDefaultValueIsActivated(FName PinLabel, bool bIsActivated);
+#endif // WITH_EDITOR
+
+private:
+#if WITH_EDITOR
+	void ResetDefaultValueInternal(FName PinLabel);
+#endif // WITH_EDITOR
+
+	UPROPERTY()
+	TMap<FName, FPCGPinDefaultValueInfo> DefaultValueMap;
+
+	UPROPERTY()
+	FPCGDefaultValueContainer DefaultValues;
+};
+
+/**
+ * Interface for settings that support inline constants via FPCGInlineConstantState.
+ * Inherits from IPCGSettingsDefaultValueProvider and overrides all methods with boilerplate implementations that
+ * delegate to the composable state struct.
+ *
+ * Inheriting this interface:
+ *   1. Embed an FPCGInlineConstantState as a UPROPERTY
+ *   2. Override GetInlineConstantState() / GetMutableInlineConstantState() to return it
+ *   3. Override GetPinInitialDefaultValueType() to specify eligible pins and their initial types
+ *
+ * Note: Do NOT inherit IPCGSettingsDefaultValueProvider manually alongside this interface.
+ */
+UINTERFACE(MinimalAPI)
+class UPCGSettingsInlineConstant : public UPCGSettingsDefaultValueProvider
+{
+	GENERATED_BODY()
+};
+
+class IPCGSettingsInlineConstant : public IPCGSettingsDefaultValueProvider
+{
+	GENERATED_BODY()
+
+public:
+	/** Creates a Param Data with the inline constant default value properties inserted as metadata. */
+	PCG_API const UPCGParamData* CreateDefaultValueParamData(FPCGContext* InContext, FName PropertyKey, FName AttributeName = NAME_None) const;
+
+	// ~Begin IPCGSettingsDefaultValueProvider interface
+	PCG_API virtual bool DefaultValuesAreEnabled() const override;
+	PCG_API virtual bool IsPinDefaultValueEnabled(FName PinLabel) const override;
+	PCG_API virtual bool IsPinDefaultValueActivated(FName PinLabel) const override;
+	PCG_API virtual EPCGMetadataTypes GetPinDefaultValueType(FName PinLabel) const override;
+
+#if WITH_EDITOR
+	PCG_API virtual bool IsPinDefaultValueMetadataTypeValid(FName PinLabel, EPCGMetadataTypes DataType) const override;
+	PCG_API virtual void ResetDefaultValues() override;
+	PCG_API virtual void ResetDefaultValue(FName PinLabel) override;
+	PCG_API virtual void SetPinDefaultValue(FName PinLabel, const FString& DefaultValue, bool bCreateIfNeeded = false) override;
+	PCG_API virtual void ConvertPinDefaultValueMetadataType(FName PinLabel, EPCGMetadataTypes DataType) override;
+	PCG_API virtual void SetPinDefaultValueIsActivated(FName PinLabel, bool bIsActivated, bool bDirtySettings = true) override;
+	PCG_API virtual FString GetPinDefaultValueAsString(FName PinLabel) const override;
+	PCG_API virtual FString GetPinInitialDefaultValueString(FName PinLabel) const override;
+	PCG_API virtual EPCGSettingDefaultValueExtraFlags GetDefaultValueExtraFlags(FName PinLabel) const override;
+#endif // WITH_EDITOR
+	// ~End IPCGSettingsDefaultValueProvider interface
+
+protected:
+	/** Return the embedded const FPCGInlineConstantState. Must be overridden by subclasses. */
+	virtual const FPCGInlineConstantState* GetInlineConstantState() const = 0;
+
+	/** Return the embedded mutable FPCGInlineConstantState. Must be overridden by subclasses. */
+	virtual FPCGInlineConstantState* GetMutableInlineConstantState() = 0;
+};
diff --git a/Engine/Plugins/PCG/Source/PCG/Public/PCGSettings.h b/Engine/Plugins/PCG/Source/PCG/Public/PCGSettings.h
--- a/Engine/Plugins/PCG/Source/PCG/Public/PCGSettings.h
+++ b/Engine/Plugins/PCG/Source/PCG/Public/PCGSettings.h
@@ -325,6 +325,8 @@ public:
 #if WITH_EDITOR
 	/** Puts node title on node body, reducing overall node size */
 	virtual bool ShouldDrawNodeCompact() const { return false; }
+	/** Whether to show/hide the title text/icon on the compact node. */
+	virtual bool ShouldShowCompactNodeTitle() const { return true; }
 	/** Returns the icon to use instead of text in compact node form */
 	virtual bool GetCompactNodeIcon(FName& OutCompactNodeIcon) const { return false; }
 	/** Returns whether the user can directly interact with the node name */
@@ -668,6 +670,11 @@ private:
 	/** Calculate Crc for these settings and save it. */
 	UE_API void CacheCrc();
 
+#if WITH_EDITOR
+	/** Self-subscribing handler that invalidates CachedCrc whenever OnSettingsChangedDelegate fires. */
+	void OnSettingsChanged(UPCGSettings* InSettings, EPCGChangeType ChangeType);
+#endif
+
 #if WITH_EDITORONLY_DATA
 	/** The cached Crc for these settings. */
 	FPCGCrc CachedCrc;
