[PCG] Remove ensure firing in Editor Tests

#Synced-CL 52194233

diff --git a/Engine/Plugins/PCG/Source/PCG/Private/Elements/Metadata/PCGMetadataOpElementBase.cpp b/Engine/Plugins/PCG/Source/PCG/Private/Elements/Metadata/PCGMetadataOpElementBase.cpp
--- a/Engine/Plugins/PCG/Source/PCG/Private/Elements/Metadata/PCGMetadataOpElementBase.cpp
+++ b/Engine/Plugins/PCG/Source/PCG/Private/Elements/Metadata/PCGMetadataOpElementBase.cpp
@@ -703,9 +703,8 @@ bool FPCGMetadataElementBase::PrepareDataInternal(FPCGContext* Context) const
 					// Need to make sure the param data is properly tracked by the context to prevent garbage collection
 					TimeSlicedContext->TrackObject(DefaultData.Data);
 
-					// @todo_pcg: To be removed once it's confirmed this never triggers.
 					UPCGMetadata* DefaultParamMetadata = CastChecked<UPCGParamData>(DefaultData.Data)->Metadata;
-					if (!ensure(DefaultParamMetadata->GetLocalItemCount() != 0))
+					if (DefaultParamMetadata->GetLocalItemCount() == 0)
 					{
 						DefaultParamMetadata->AddEntry();
 					}
