[PCG] Remove BGRA8 format restriction from SaveTextureToAsset (and minor fixes)

- Auto-detect output format from input texture pixel format. Supports BGRA8 (8-bit), RGBA16F (16-bit HDR), and RGBA32F (32-bit HDR). Fallback to BGRA8 for unsupported formats.
- Fix destructively updating render target resources on input textures.
- Fix GPU readback row pitch alignment.
- Fix exported UTexture2D having blank platform data.
- Call Modify() on the exported texture so re-exports to an existing asset mark the package dirty.
- Fix SaveTextureToAsset node being cacheable. Should always re-export.
- Add subtype to asset path output pin, to get the pin color!

#Synced-CL 52208564

diff --git a/Engine/Plugins/PCG/Source/PCG/Private/Elements/PCGSaveTextureToAsset.cpp b/Engine/Plugins/PCG/Source/PCG/Private/Elements/PCGSaveTextureToAsset.cpp
--- a/Engine/Plugins/PCG/Source/PCG/Private/Elements/PCGSaveTextureToAsset.cpp
+++ b/Engine/Plugins/PCG/Source/PCG/Private/Elements/PCGSaveTextureToAsset.cpp
@@ -6,10 +6,13 @@
 #include "PCGModule.h"
 #include "PCGParamData.h"
 #include "PCGTextureReadback.h"
+#include "Data/PCGRenderTargetData.h"
 #include "Data/PCGTextureData.h"
 #include "Metadata/PCGMetadata.h"
 
+#include "PixelFormat.h"
 #include "RHIStaticStates.h"
+#include "Async/ParallelFor.h"
 #include "Engine/Texture2D.h"
 
 #include UE_INLINE_GENERATED_CPP_BY_NAME(PCGSaveTextureToAsset)
@@ -21,6 +24,91 @@ namespace PCGSaveTextureToAssetConstants
 	const FName OutAssetPathLabel = TEXT("OutAssetPath");
 }
 
+namespace PCGSaveTextureToAssetHelpers
+{
+	/** Maps an EPixelFormat to an ETextureSourceFormat.
+	 * Returns TSF_Invalid if no direct mapping exists.
+	 */
+	ETextureSourceFormat GetTextureSourceFormat(EPixelFormat InFormat)
+	{
+		switch (InFormat)
+		{
+		case PF_G8:              return TSF_G8;
+		case PF_B8G8R8A8:        return TSF_BGRA8;
+		case PF_G16:             return TSF_G16;
+		case PF_R16F:            return TSF_R16F;
+		case PF_FloatRGBA:       return TSF_RGBA16F;
+		case PF_R32_FLOAT:       return TSF_R32F;
+		case PF_A32B32G32R32F:   return TSF_RGBA32F;
+		default:                 return TSF_Invalid;
+		}
+	}
+
+	/** Returns a readback-compatible EPixelFormat for the given input format.
+	 * Prefers the input format when it supports UAV and has a direct ETextureSourceFormat mapping.
+	 * Falls back to a format at matching precision and channel count otherwise.
+	 */
+	EPixelFormat GetReadbackFormat(EPixelFormat InFormat)
+	{
+		if (UE::PixelFormat::HasCapabilities(InFormat, EPixelFormatCapabilities::UAV) && GetTextureSourceFormat(InFormat) != TSF_Invalid)
+		{
+			return InFormat;
+		}
+
+		// Fallback: pick a format with matching precision. ETextureSourceFormat has no 2-channel entries, so 2/3-channel formats are promoted to 4-channel.
+		const FPixelFormatInfo& FormatInfo = GPixelFormats[InFormat];
+		const int32 BytesPerComponent = FormatInfo.BlockBytes / FMath::Max(FormatInfo.NumComponents, 1);
+
+		if (BytesPerComponent >= 4)
+		{
+			// 32-bit
+			return (FormatInfo.NumComponents <= 1) ? PF_R32_FLOAT : PF_A32B32G32R32F;
+		}
+		else if (BytesPerComponent >= 2 || IsFloatFormat(InFormat) || IsHDR(InFormat))
+		{
+			// 16-bit
+			return (FormatInfo.NumComponents <= 1) ? PF_R16F : PF_FloatRGBA;
+		}
+		else
+		{
+			// 8-bit
+			return (FormatInfo.NumComponents <= 1) ? PF_G8 : PF_B8G8R8A8;
+		}
+	}
+
+	/** Sets compression settings on a texture to match its source format.
+	 * TC_Default uses LDR compression which is lossy for HDR and grayscale formats.
+	 */
+	void ApplyCompressionSettings(UTexture2D* Texture, ETextureSourceFormat InFormat)
+	{
+		switch (InFormat)
+		{
+		// HDR formats: sRGB is forced off because the sRGB transfer function clamps to [0,1] and applies gamma, which is incompatible with linear float data.
+		case TSF_RGBA16F:
+		case TSF_RGBA32F:
+			Texture->CompressionSettings = TC_HDR_Compressed;
+			Texture->SRGB = false;
+			break;
+		case TSF_R16F:
+			Texture->CompressionSettings = TC_HalfFloat;
+			Texture->SRGB = false;
+			break;
+		case TSF_R32F:
+			Texture->CompressionSettings = TC_SingleFloat;
+			Texture->SRGB = false;
+			break;
+		case TSF_G8:
+		case TSF_G16:
+			Texture->CompressionSettings = TC_Grayscale;
+			break;
+		case TSF_BGRA8:
+		default:
+			Texture->CompressionSettings = TC_Default;
+			break;
+		}
+	}
+}
+
 TArray<FPCGPinProperties> UPCGSaveTextureToAssetSettings::InputPinProperties() const
 {
 	TArray<FPCGPinProperties> Properties;
@@ -33,7 +121,8 @@ TArray<FPCGPinProperties> UPCGSaveTextureToAssetSettings::InputPinProperties() c
 TArray<FPCGPinProperties> UPCGSaveTextureToAssetSettings::OutputPinProperties() const
 {
 	TArray<FPCGPinProperties> Properties;
-	Properties.Emplace(PCGSaveTextureToAssetConstants::OutAssetPathLabel, EPCGDataType::Param, /*bInAllowMultipleConnections=*/true, /*bAllowMultipleData=*/false);
+	FPCGPinProperties& Pin = Properties.Emplace_GetRef(PCGSaveTextureToAssetConstants::OutAssetPathLabel, EPCGDataType::Param, /*bInAllowMultipleConnections=*/true, /*bAllowMultipleData=*/false);
+	Pin.AllowedTypes.CustomSubtype = static_cast<int32>(EPCGMetadataTypes::SoftObjectPath);
 
 	return Properties;
 }
@@ -126,8 +215,14 @@ bool FPCGSaveTextureToAssetElement::ExecuteInternal(FPCGContext* InContext) cons
 
 		if (ExportedTexture)
 		{
-			// Initialize exported texture with the readback data.
-			ExportedTexture->Source.Init(Context->ReadbackWidth, Context->ReadbackHeight, /*NewNumSlices=*/1, /*NewNumMips=*/1, TSF_BGRA8, Context->RawReadbackData);
+			// Initialize exported texture with the readback data in the format matching the input texture.
+			const ETextureSourceFormat TexSrcFormat = PCGSaveTextureToAssetHelpers::GetTextureSourceFormat(Context->ReadbackFormat);
+
+			ExportedTexture->Modify();
+			ExportedTexture->Source.Init(Context->ReadbackWidth, Context->ReadbackHeight, /*NewNumSlices=*/1, /*NewNumMips=*/1, TexSrcFormat, Context->RawReadbackData);
+
+			PCGSaveTextureToAssetHelpers::ApplyCompressionSettings(ExportedTexture, TexSrcFormat);
+
 			ExportedTexture->UpdateResource();
 
 			Context->ExportedTexture = ExportedTexture;
@@ -189,7 +284,9 @@ bool FPCGSaveTextureToAssetElement::ReadbackInputTexture(FPCGSaveTextureToAssetC
 				return true;
 			}
 
-			if (!InContext->bUpdatedReadbackTextureResource)
+			// Render target data already has a fully initialized GPU resource.
+			// Calling UpdateResource() on it would re-create the resource and discard the content.
+			if (!InContext->bUpdatedReadbackTextureResource && !InContext->InputTextureData->IsA<UPCGRenderTargetData>())
 			{
 				InputTexture->UpdateResource();
 
@@ -217,10 +314,12 @@ bool FPCGSaveTextureToAssetElement::ReadbackInputTexture(FPCGSaveTextureToAssetC
 		Params.SourceSampler = TStaticSamplerState<SF_Point, AM_Clamp, AM_Clamp, AM_Clamp>::GetRHI();
 		Params.SourceTextureIndex = InContext->InputTextureData->GetTextureSlice();
 		Params.SourceDimensions = FIntPoint(TextureSize.X, TextureSize.Y);
+		Params.OutputFormat = PCGSaveTextureToAssetHelpers::GetReadbackFormat(InContext->InputTextureData->GetFormat());
 
+		InContext->ReadbackFormat = Params.OutputFormat;
 		InContext->bReadbackDispatched = true;
 
-		FPCGTextureReadbackInterface::Dispatch(Params, [ContextHandle = InContext->GetOrCreateHandle()](void* OutBuffer, int32 ReadbackWidth, int32 ReadbackHeight)
+		FPCGTextureReadbackInterface::Dispatch(Params, [ContextHandle = InContext->GetOrCreateHandle(), SourceDimensions = Params.SourceDimensions](void* OutBuffer, int32 ReadbackWidth, int32 ReadbackHeight)
 		{
 			TRACE_CPUPROFILER_EVENT_SCOPE(FPCGSaveTextureToAssetElement::ReadbackInputTexture::DispatchCallback);
 
@@ -232,13 +331,35 @@ bool FPCGSaveTextureToAssetElement::ReadbackInputTexture(FPCGSaveTextureToAssetC
 				return;
 			}
 
-			Context->ReadbackWidth = ReadbackWidth;
-			Context->ReadbackHeight = ReadbackHeight;
+			Context->ReadbackWidth = SourceDimensions.X;
+			Context->ReadbackHeight = SourceDimensions.Y;
 
-			// PCGTextureReadback always uses PF_B8G8R8A8. @todo_pcg: Maybe the read back format could be configurable.
-			const uint32 BufferSizeBytes = ReadbackWidth * ReadbackHeight * GPixelFormats[PF_B8G8R8A8].BlockBytes;
+			const int32 BytesPerPixel = GPixelFormats[Context->ReadbackFormat].BlockBytes;
+			const uint32 DstPitch = SourceDimensions.X * BytesPerPixel;
+			const uint32 SrcPitch = ReadbackWidth * BytesPerPixel;
+			const uint32 BufferSizeBytes = DstPitch * SourceDimensions.Y;
 			Context->RawReadbackData = static_cast<uint8*>(FMemory::Malloc(BufferSizeBytes));
-			FMemory::Memcpy(Context->RawReadbackData, OutBuffer, BufferSizeBytes);
+
+			if (ReadbackWidth == SourceDimensions.X)
+			{
+				FMemory::Memcpy(Context->RawReadbackData, OutBuffer, BufferSizeBytes);
+			}
+			else if (ensure(ReadbackWidth > SourceDimensions.X))
+			{
+				// Readback rows may be padded for GPU alignment (e.g. D3D12 256-byte row pitch).
+				// Copy row-by-row, stripping the padding to produce a tightly packed buffer.
+				ParallelFor(SourceDimensions.Y, [&](int32 Row)
+				{
+					FMemory::Memcpy(
+						Context->RawReadbackData + Row * DstPitch,
+						static_cast<uint8*>(OutBuffer) + Row * SrcPitch,
+						DstPitch);
+				});
+			}
+			else
+			{
+				FMemory::Memzero(Context->RawReadbackData, BufferSizeBytes);
+			}
 
 			Context->bReadbackComplete = true;
 			Context->bIsPaused = false;
diff --git a/Engine/Plugins/PCG/Source/PCG/Public/Elements/PCGSaveTextureToAsset.h b/Engine/Plugins/PCG/Source/PCG/Public/Elements/PCGSaveTextureToAsset.h
--- a/Engine/Plugins/PCG/Source/PCG/Public/Elements/PCGSaveTextureToAsset.h
+++ b/Engine/Plugins/PCG/Source/PCG/Public/Elements/PCGSaveTextureToAsset.h
@@ -5,13 +5,14 @@
 #include "PCGAssetExporter.h"
 #include "PCGContext.h"
 #include "PCGSettings.h"
+#include "PixelFormat.h"
 
 #include "PCGSaveTextureToAsset.generated.h"
 
 class UPCGTexture2DSingleBaseData;
 class UTexture2D;
 
-/** Save the input texture to a UTexture2D asset (format is always BGRA8). Outputs a soft object path attribute for the saved texture. Editor only. */
+/** Save the input texture to a UTexture2D asset. Output format is auto-detected from the input texture (BGRA8, RGBA16F, or RGBA32F). Outputs a soft object path attribute for the saved texture. Editor only. */
 UCLASS(MinimalAPI, BlueprintType, ClassGroup = (Procedural))
 class UPCGSaveTextureToAssetSettings : public UPCGSettings
 {
@@ -22,7 +23,7 @@ public:
 #if WITH_EDITOR
 	virtual FName GetDefaultNodeName() const override { return FName(TEXT("SaveTextureToAsset")); }
 	virtual FText GetDefaultNodeTitle() const override { return NSLOCTEXT("PCGSaveTextureToAssetElement", "NodeTitle", "Save Texture to Asset"); }
-	virtual FText GetNodeTooltipText() const override { return NSLOCTEXT("PCGSaveTextureToAssetElement", "NodeTooltip", "Save the input texture to a UTexture2D asset (format is always BGRA8). Outputs a soft object path attribute for the saved texture. Note: This node is editor only."); }
+	virtual FText GetNodeTooltipText() const override { return NSLOCTEXT("PCGSaveTextureToAssetElement", "NodeTooltip", "Save the input texture to a UTexture2D asset. Output format is auto-detected from the input texture (BGRA8, RGBA16F, or RGBA32F). Note: This node is editor only."); }
 	virtual EPCGSettingsType GetType() const override { return EPCGSettingsType::InputOutput; }
 #endif
 
@@ -52,6 +53,7 @@ public:
 	int32 ReadbackWidth = INDEX_NONE;
 	int32 ReadbackHeight = INDEX_NONE;
 	bool bUpdatedReadbackTextureResource = false;
+	EPixelFormat ReadbackFormat = PF_B8G8R8A8;
 	bool bReadbackDispatched = false;
 	bool bReadbackComplete = false;
 };
@@ -61,6 +63,7 @@ class FPCGSaveTextureToAssetElement : public IPCGElementWithCustomContext<FPCGSa
 public:
 	/** Creating asset and initializing texture resource must occur on main thread. */
 	virtual bool CanExecuteOnlyOnMainThread(FPCGContext* Context) const override { return true; }
+	virtual bool IsCacheable(const UPCGSettings* InSettings) const override { return false; }
 
 protected:
 	virtual bool ExecuteInternal(FPCGContext* InContext) const override;
diff --git a/Engine/Plugins/PCG/Source/PCGCompute/Private/PCGTextureReadback.cpp b/Engine/Plugins/PCG/Source/PCGCompute/Private/PCGTextureReadback.cpp
--- a/Engine/Plugins/PCG/Source/PCGCompute/Private/PCGTextureReadback.cpp
+++ b/Engine/Plugins/PCG/Source/PCGCompute/Private/PCGTextureReadback.cpp
@@ -91,7 +91,7 @@ void FPCGTextureReadbackInterface::Dispatch_RenderThread(FRHICommandListImmediat
 	}
 
 	FRHITextureCreateDesc TargetTextureDesc =
-		FRHITextureCreateDesc::Create2D(TEXT("PCGTexture Readback Compute Target"), Params.SourceDimensions.X, Params.SourceDimensions.Y, EPixelFormat::PF_B8G8R8A8)
+		FRHITextureCreateDesc::Create2D(TEXT("PCGTexture Readback Compute Target"), Params.SourceDimensions.X, Params.SourceDimensions.Y, Params.OutputFormat)
 			.SetClearValue(FClearValueBinding::None)
 			.SetFlags(ETextureCreateFlags::UAV |
 				ETextureCreateFlags::ShaderResource |
@@ -145,6 +145,10 @@ void FPCGTextureReadbackInterface::Dispatch_RenderThread(FRHICommandListImmediat
 
 		if (GPUTextureReadBackPtr->IsReady())
 		{
+			// NOTE: Lock() returns the row pitch in pixels as ReadbackWidth, which may be larger than the source texture width due to GPU alignment (e.g. D3D12 256-byte row pitch).
+			// Call sites must capture SourceDimensions and use ReadbackWidth only as a row stride when unpacking the buffer.
+			// See PCGTextureData and PCGSaveTextureToAsset for examples.
+			// @todo_pcg: Maybe we should just unpack the buffer here, but it would cost some perf.
 			int32 ReadbackWidth = 0, ReadbackHeight = 0;
 			void* OutBuffer = GPUTextureReadBackPtr->Lock(ReadbackWidth, &ReadbackHeight);
 
diff --git a/Engine/Plugins/PCG/Source/PCGCompute/Public/PCGTextureReadback.h b/Engine/Plugins/PCG/Source/PCGCompute/Public/PCGTextureReadback.h
--- a/Engine/Plugins/PCG/Source/PCGCompute/Public/PCGTextureReadback.h
+++ b/Engine/Plugins/PCG/Source/PCGCompute/Public/PCGTextureReadback.h
@@ -22,6 +22,9 @@ struct FPCGTextureReadbackDispatchParams
 
 	/** Optional texture index. Should be 0 if the SourceTexture is not a texture array. */
 	uint32 SourceTextureIndex = 0;
+
+	/** Pixel format for the staging texture used during readback. Must be a 4-channel uncompressed format (e.g. PF_B8G8R8A8, PF_FloatRGBA, PF_A32B32G32R32F). */
+	EPixelFormat OutputFormat = PF_B8G8R8A8;
 };
  
 /**
