[PCG] Support GPU compute nodes in standalone graphs

#Synced-CL 52213696

diff --git a/Engine/Plugins/PCG/Source/PCG/Private/Compute/Elements/PCGComputeGraphElement.cpp b/Engine/Plugins/PCG/Source/PCG/Private/Compute/Elements/PCGComputeGraphElement.cpp
--- a/Engine/Plugins/PCG/Source/PCG/Private/Compute/Elements/PCGComputeGraphElement.cpp
+++ b/Engine/Plugins/PCG/Source/PCG/Private/Compute/Elements/PCGComputeGraphElement.cpp
@@ -3,6 +3,7 @@
 #include "Compute/Elements/PCGComputeGraphElement.h"
 
 #include "PCGComponent.h"
+#include "PCGDefaultExecutionSource.h"
 #include "PCGGraph.h"
 #include "PCGModule.h"
 #include "Components/PCGProceduralISMComponent.h"
@@ -24,10 +25,44 @@
 #include "Engine/World.h"
 #include "Logging/LogVerbosity.h"
 
+#if WITH_EDITOR
+#include "Editor.h"
+#endif
+
 #include UE_INLINE_GENERATED_CPP_BY_NAME(PCGComputeGraphElement)
 
 #define LOCTEXT_NAMESPACE "PCGComputeGraphElement"
 
+namespace PCGComputeGraphElementHelpers
+{
+	/** Resolve a scene for compute dispatch. Prefers the execution source's world scene, falls back to the editor world scene for standalone graphs. */
+	static FSceneInterface* GetSceneForCompute(const FPCGContext* InContext)
+	{
+		if (!InContext->ExecutionSource.IsValid())
+		{
+			return nullptr;
+		}
+
+		if (UWorld* World = InContext->ExecutionSource->GetExecutionState().GetWorld())
+		{
+			return World->Scene;
+		}
+
+#if WITH_EDITOR
+		// Fallback to editor scene for standalone graph execution (UPCGDefaultExecutionSource has no world by design).
+		if (GEditor && Cast<UPCGDefaultExecutionSource>(InContext->ExecutionSource.Get()))
+		{
+			if (UWorld* EditorWorld = GEditor->GetEditorWorldContext().World())
+			{
+				return EditorWorld->Scene;
+			}
+		}
+#endif
+
+		return nullptr;
+	}
+}
+
 #if WITH_EDITOR
 #define PCG_COMPUTE_GRAPH_BREAKPOINT() \
 	if (Context->ComputeGraph && Context->ComputeGraph->bBreakDebugger && Context->IsExecutingGraphInspected()) \
@@ -143,9 +178,11 @@ bool FPCGComputeGraphElement::ExecuteInternal(FPCGContext* InContext) const
 		return true;
 	}
 
-	if (!InContext->ExecutionSource->GetExecutionState().GetWorld())
+	FSceneInterface* Scene = PCGComputeGraphElementHelpers::GetSceneForCompute(InContext);
+
+	if (!Scene)
 	{
-		UE_LOG(LogPCG, Error, TEXT("FPCGComputeGraphElement: Execution source is expected to have a world."));
+		UE_LOG(LogPCG, Error, TEXT("FPCGComputeGraphElement: Could not resolve a scene for compute dispatch."));
 		return true;
 	}
 
@@ -527,12 +564,10 @@ bool FPCGComputeGraphElement::ExecuteInternal(FPCGContext* InContext) const
 	case EPCGComputeGraphExecutionPhase::ScheduleComputeGraph: // Fallthrough
 	{
 		TRACE_CPUPROFILER_EVENT_SCOPE(ScheduleComputeGraph);
-		IPCGGraphExecutionSource* ExecutionSource = Context->ExecutionSource.Get();
-		check(ExecutionSource&& ExecutionSource->GetExecutionState().GetWorld());
 
 		const bool bGraphEnqueued = Context->ComputeGraphInstance->EnqueueWork(
 			Context->ComputeGraph.Get(),
-			InContext->ExecutionSource->GetExecutionState().GetWorld()->Scene,
+			Scene,
 			Context->ComputeGraph->GetRequiredExecutionGroup(),
 			FName(*InContext->ExecutionSource->GetExecutionState().GetDebugName()),
 			/*InFallbackDelegate=*/FSimpleDelegate::CreateLambda([ContextHandle = Context->GetOrCreateHandle()]()
@@ -714,10 +749,10 @@ void FPCGComputeGraphElement::AbortInternal(FPCGContext* InContext) const
 {
 	if (FPCGComputeGraphContext* Context = static_cast<FPCGComputeGraphContext*>(InContext))
 	{
-		const IPCGGraphExecutionSource* Source = InContext->ExecutionSource.Get();
-		if (Context->ComputeGraphInstance && Context->DataBinding && Source && Source->GetExecutionState().GetWorld())
+		FSceneInterface* ComputeScene = PCGComputeGraphElementHelpers::GetSceneForCompute(InContext);
+		if (Context->ComputeGraphInstance && Context->DataBinding && ComputeScene)
 		{
-			ComputeFramework::AbortWork(Source->GetExecutionState().GetWorld()->Scene, Context->DataBinding);
+			ComputeFramework::AbortWork(ComputeScene, Context->DataBinding);
 		}
 	}
 
