[PCG] Add CopyTo function on accessors to copy an accessor to the other without knowing their types

#rnx
#Synced-CL 52163112

diff --git a/Engine/Plugins/PCG/Source/PCG/Private/Metadata/Accessors/IPCGAttributeAccessor.cpp b/Engine/Plugins/PCG/Source/PCG/Private/Metadata/Accessors/IPCGAttributeAccessor.cpp
--- a/Engine/Plugins/PCG/Source/PCG/Private/Metadata/Accessors/IPCGAttributeAccessor.cpp
+++ b/Engine/Plugins/PCG/Source/PCG/Private/Metadata/Accessors/IPCGAttributeAccessor.cpp
@@ -2,80 +2,187 @@
 
 #include "Metadata/Accessors/IPCGAttributeAccessor.h"
 
+#include "Utils/PCGLogErrors.h"
+
+#define LOCTEXT_NAMESPACE "PCGAttributeAccessor"
+
 namespace PCGAttributeAccessorPrivate
 {
 	constexpr int32 MaxAlignment = alignof(FTransform);
 	constexpr int32 AlignedBlockSize = MaxAlignment;
 	constexpr int32 MaxInlineMem = 1024 / AlignedBlockSize; // 1KB
-	
+
 	// Making sure we can accommodate our big types that have different alignment like FTransform or FVector4
 	using FRawBuffer = TArray<TAlignedBytes<AlignedBlockSize, MaxAlignment>, TInlineAllocator<MaxInlineMem>>;
-	
+
 	int32 GetAllocationSize(const int32 ElementSize, const int32 Count)
 	{
 		return FMath::DivideAndRoundUp(ElementSize * Count, AlignedBlockSize);
 	}
-}
 
-bool IPCGAttributeAccessor::GetRange_Internal(PCG::Private::FOutValues OutValues, int16 ExpectedType, int32 Count, int32 Index, const IPCGAttributeAccessorKeys& Keys, EPCGAttributeAccessorFlags Flags, PCG::Private::EPCGAttributeAccessorOperation OpType) const
-{
-	using namespace PCG::Private;
-	
-	switch (OpType)
-	{
-	case EPCGAttributeAccessorOperation::SameType:
-	{
-		return GetRangeVirtual(MoveTemp(OutValues), Count, Index, Keys);
-	}
-	case EPCGAttributeAccessorOperation::Construct: // Fall-through
-	case EPCGAttributeAccessorOperation::Broadcast:
+	bool CallbackWithTempBuffer(
+		const FPCGMetadataAttributeDesc& BufferDesc,
+		const FPCGMetadataAttributeDesc& ReadDesc,
+		const FPCGMetadataAttributeDesc& WriteDesc,
+		PCG::Private::EPCGAttributeAccessorOperation OpType,
+		int32 Count,
+		TFunctionRef<bool(FRawBuffer&, PCG::Private::FTransformFunc)> Callback)
 	{
-		check(OutValues.IsType<FOutValuesByValue>() && IsPCGType(ExpectedType));
-			
-		PCGAttributeAccessorPrivate::FRawBuffer Buffer{}; // 1KB
-		const uint16 UnderlyingValueType = static_cast<uint16>(UnderlyingDesc.ValueType);
-			
-		const int32 ElementSize = GetElementSize(UnderlyingValueType);
+		check(&BufferDesc == &ReadDesc || &BufferDesc == &WriteDesc);
+		using namespace PCG::Private;
+
+		FRawBuffer Buffer{}; // 1KB
+		const uint16 BufferValueType = static_cast<uint16>(BufferDesc.ValueType);
+		const uint16 ReadValueType = static_cast<uint16>(ReadDesc.ValueType);
+		const uint16 WriteValueType = static_cast<uint16>(WriteDesc.ValueType);
+
+		const int32 ElementSize = GetElementSize(BufferValueType);
 		if (!ensure(ElementSize > 0))
 		{
 			return false;
 		}
-			
+
 		Buffer.SetNumUninitialized(PCGAttributeAccessorPrivate::GetAllocationSize(ElementSize, Count));
-		if (FConstructMemoryFunc ConstructFunc = GetConstructMemoryFunc(UnderlyingValueType))
+		if (FConstructMemoryFunc ConstructFunc = GetConstructMemoryFunc(BufferValueType))
 		{
 			(*ConstructFunc)(Buffer.GetData(), Count);
 		}
 
-		const bool bSuccess = GetRangeVirtual(FOutValues{TInPlaceType<FOutValuesByValue>{}, Buffer.GetData()}, Count, Index, Keys);
-			
-		if (bSuccess)
-		{
-			FTransformFunc TransformFunc = OpType == EPCGAttributeAccessorOperation::Broadcast 
-				? GetBroadcastTransformFunc(UnderlyingValueType, ExpectedType) 
-				: GetConstructibleTransformFunc(ExpectedType, UnderlyingValueType);
-			
-			check(TransformFunc);
-			(*TransformFunc)(OutValues.Get<FOutValuesByValue>().OutValues, Buffer.GetData(), Count);
-		}
-			
-		if (FDestroyMemoryFunc DestroyFunc = GetDestroyMemoryFunc(UnderlyingValueType))
+		FTransformFunc TransformFunc = OpType == EPCGAttributeAccessorOperation::Broadcast
+			? GetBroadcastTransformFunc(ReadValueType, WriteValueType)
+			: GetConstructibleTransformFunc(WriteValueType, ReadValueType);
+
+		const bool bSuccess = Callback(Buffer, TransformFunc);
+
+		if (FDestroyMemoryFunc DestroyFunc = GetDestroyMemoryFunc(BufferValueType))
 		{
 			(*DestroyFunc)(Buffer.GetData(), Count);
 		}
-			
+
 		return bSuccess;
 	}
+}
+
+bool IPCGAttributeAccessor::CopyTo(const IPCGAttributeAccessorKeys& SrcKeys, IPCGAttributeAccessor& DestAccessor, IPCGAttributeAccessorKeys& DestKeys, int32 Index, int32 Count, EPCGAttributeAccessorFlags Flags) const
+{
+	// Can't copy to read only accessors
+	if (DestAccessor.IsReadOnly())
+	{
+		PCGLog::LogErrorOnGraph(LOCTEXT("ReadOnlyDest", "[PCGAttributeAccessor::CopyTo] Destination accessor is read only."));
+		return false;
+	}
+	
+	// Or outside range
+	if (Index >= DestKeys.GetNum())
+	{
+		PCGLog::LogErrorOnGraph(FText::Format(LOCTEXT("OutsideRange", "[PCGAttributeAccessor::CopyTo] Try to write outside of destination accessor range. (Index = {0}, MaxIndex = {1})."), Index, DestKeys.GetNum() - 1));
+		return false;
+	}
+	
+	// Can only copy matching containers
+	const FPCGMetadataAttributeDesc& ReadDesc = GetUnderlyingDesc();
+	const FPCGMetadataAttributeDesc& WriteDesc = DestAccessor.GetUnderlyingDesc();
+	if (ReadDesc.ContainerTypes != WriteDesc.ContainerTypes)
+	{
+		PCGLog::LogErrorOnGraph(FText::Format(LOCTEXT("MismatchContainer", "[PCGAttributeAccessor::CopyTo] Mismatch between containers {0} vs {1}"), FText::FromString(ReadDesc.GetTypeString()), FText::FromString(WriteDesc.GetTypeString())));
+		return false;
+	}
+
+	const PCG::Private::EPCGAttributeAccessorOperation OpType = PCG::Private::GetOpType(ReadDesc, WriteDesc, Flags);
+
+	const bool bIsSingleValue = ReadDesc.IsSingleValue();
+	if (OpType == PCG::Private::EPCGAttributeAccessorOperation::Invalid || (!bIsSingleValue && OpType != PCG::Private::EPCGAttributeAccessorOperation::SameType))
+	{
+		PCGLog::LogErrorOnGraph(FText::Format(LOCTEXT("ImpossibleConversion", "[PCGAttributeAccessor::CopyTo] Can't convert {0} to {1}"), FText::FromString(ReadDesc.GetTypeString()), FText::FromString(WriteDesc.GetTypeString())));
+		return false;
+	}
+
+	if (bIsSingleValue)
+	{
+		auto ReadCallback = [this, Count, Index, &SrcKeys, &DestAccessor, &DestKeys, Flags, &ReadDesc, &WriteDesc, OpType](PCGAttributeAccessorPrivate::FRawBuffer& ReadBuffer, PCG::Private::FTransformFunc TransformFunc)
+		{
+			if (GetRangeVirtual(PCG::Private::FOutValues{TInPlaceType<PCG::Private::FOutValuesByValue>{}, ReadBuffer.GetData()}, Count, Index, SrcKeys))
+			{
+				// If we have a transform func, we need 2 buffers
+				if (TransformFunc)
+				{
+					auto WriteCallback = [this, Count, Index, &DestAccessor, &DestKeys, Flags, &ReadBuffer](PCGAttributeAccessorPrivate::FRawBuffer& WriteBuffer, PCG::Private::FTransformFunc TransformFunc)
+					{
+						check(TransformFunc);
+						(*TransformFunc)(WriteBuffer.GetData(), ReadBuffer.GetData(), Count);
+						return DestAccessor.SetRangeVirtual(PCG::Private::FInValues{TInPlaceType<PCG::Private::FInValuesByValue>{}, WriteBuffer.GetData()}, Count, Index, DestKeys, Flags);
+					};
+
+					return PCGAttributeAccessorPrivate::CallbackWithTempBuffer(WriteDesc, ReadDesc, WriteDesc, OpType, Count, WriteCallback);
+				}
+				else
+				{
+					return DestAccessor.SetRangeVirtual(PCG::Private::FInValues{TInPlaceType<PCG::Private::FInValuesByValue>{}, ReadBuffer.GetData()}, Count, Index, DestKeys, Flags);
+				}
+			}
+			else
+			{
+				return false;
+			}
+		};
+
+		return PCGAttributeAccessorPrivate::CallbackWithTempBuffer(ReadDesc, ReadDesc, WriteDesc, OpType, Count, ReadCallback);
+	}
+	else if (ReadDesc.IsArray())
+	{
+		check(OpType == PCG::Private::EPCGAttributeAccessorOperation::SameType);
+		TArray<TTuple<const void*, int32>, TInlineAllocator<256>> Temp;
+		Temp.SetNumUninitialized(Count);
+
+		return GetRangeVirtual(PCG::Private::FOutValues{TInPlaceType<PCG::Private::FOutValuesAsArray>{}, Temp}, Count, Index, SrcKeys)
+			&& DestAccessor.SetRangeVirtual(PCG::Private::FInValues{TInPlaceType<PCG::Private::FInValuesAsArray>{}, Temp}, Count, Index, DestKeys, Flags);
+	}
+	// @todo_pcg: Set and maps
+
+	return false;
+}
+
+bool IPCGAttributeAccessor::GetRange_Internal(PCG::Private::FOutValues OutValues, const FPCGMetadataAttributeDesc& WriteDesc, int32 Count, int32 Index, const IPCGAttributeAccessorKeys& Keys, EPCGAttributeAccessorFlags Flags, PCG::Private::EPCGAttributeAccessorOperation OpType) const
+{
+	using namespace PCG::Private;
+
+	switch (OpType)
+	{
+	case EPCGAttributeAccessorOperation::SameType:
+	{
+		return GetRangeVirtual(MoveTemp(OutValues), Count, Index, Keys);
+	}
+	case EPCGAttributeAccessorOperation::Construct: // Fall-through
+	case EPCGAttributeAccessorOperation::Broadcast:
+	{
+		check(OutValues.IsType<FOutValuesByValue>() && IsPCGType(static_cast<uint16>(UnderlyingDesc.ValueType)));
+
+		auto Callback = [this, &OutValues, Count, Index, &Keys](PCGAttributeAccessorPrivate::FRawBuffer& Buffer, FTransformFunc TransformFunc)
+		{
+			if (GetRangeVirtual(FOutValues{TInPlaceType<FOutValuesByValue>{}, Buffer.GetData()}, Count, Index, Keys))
+			{
+				check(TransformFunc);
+				(*TransformFunc)(OutValues.Get<FOutValuesByValue>().OutValues, Buffer.GetData(), Count);
+				return true;
+			}
+			else
+			{
+				return false;
+			}
+		};
+
+		return PCGAttributeAccessorPrivate::CallbackWithTempBuffer(UnderlyingDesc, UnderlyingDesc, WriteDesc, OpType, Count, Callback);
+	}
 	default:
 		checkNoEntry();
 		return false;
 	}
 }
 
-bool IPCGAttributeAccessor::SetRange_Internal(PCG::Private::FInValues InValues, int16 ExpectedType, int32 Count, int32 Index, IPCGAttributeAccessorKeys& Keys, EPCGAttributeAccessorFlags Flags, PCG::Private::EPCGAttributeAccessorOperation OpType)
+bool IPCGAttributeAccessor::SetRange_Internal(PCG::Private::FInValues InValues, const FPCGMetadataAttributeDesc& ReadDesc, int32 Count, int32 Index, IPCGAttributeAccessorKeys& Keys, EPCGAttributeAccessorFlags Flags, PCG::Private::EPCGAttributeAccessorOperation OpType)
 {
 	using namespace PCG::Private;
-	
+
 	switch (OpType)
 	{
 	case EPCGAttributeAccessorOperation::SameType:
@@ -85,41 +192,22 @@ bool IPCGAttributeAccessor::SetRange_Internal(PCG::Private::FInValues InValues,
 	case EPCGAttributeAccessorOperation::Construct: // Fall-through
 	case EPCGAttributeAccessorOperation::Broadcast:
 	{
-		check(InValues.IsType<FInValuesByValue>() && IsPCGType(ExpectedType))
-		
-		PCGAttributeAccessorPrivate::FRawBuffer Buffer{}; // 1KB
-		const uint16 UnderlyingValueType = static_cast<uint16>(UnderlyingDesc.ValueType);
-			
-		const int32 ElementSize = GetElementSize(UnderlyingValueType);
-		if (!ensure(ElementSize > 0))
-		{
-			return false;
-		}
+		check(InValues.IsType<FInValuesByValue>() && IsPCGType(static_cast<uint16>(UnderlyingDesc.ValueType)));
 
-		Buffer.SetNumUninitialized(PCGAttributeAccessorPrivate::GetAllocationSize(ElementSize, Count));
-		if (FConstructMemoryFunc ConstructFunc = GetConstructMemoryFunc(UnderlyingValueType))
+		auto Callback = [this, &InValues, Count, Index, &Keys, Flags](PCGAttributeAccessorPrivate::FRawBuffer& Buffer, FTransformFunc TransformFunc)
 		{
-			(*ConstructFunc)(Buffer.GetData(), Count);
-		}
+			check(TransformFunc);
+			(*TransformFunc)(Buffer.GetData(), InValues.Get<FInValuesByValue>().InValues, Count);
 
-		FTransformFunc TransformFunc = OpType == EPCGAttributeAccessorOperation::Broadcast
-			? GetBroadcastTransformFunc(ExpectedType, UnderlyingValueType)
-			: GetConstructibleTransformFunc(UnderlyingValueType, ExpectedType);
-		
-		check(TransformFunc);
-		(*TransformFunc)(Buffer.GetData(), InValues.Get<FInValuesByValue>().InValues, Count);
-	
-		const bool bSuccess = SetRangeVirtual(FInValues{TInPlaceType<FInValuesByValue>{}, Buffer.GetData(), Count}, Count, Index, Keys, Flags);
-		
-		if (FDestroyMemoryFunc DestroyFunc = GetDestroyMemoryFunc(UnderlyingValueType))
-		{
-			(*DestroyFunc)(Buffer.GetData(), Count);
-		}
-		
-		return bSuccess;
+			return SetRangeVirtual(FInValues{TInPlaceType<FInValuesByValue>{}, Buffer.GetData(), Count}, Count, Index, Keys, Flags);
+		};
+
+		return PCGAttributeAccessorPrivate::CallbackWithTempBuffer(UnderlyingDesc, ReadDesc, UnderlyingDesc, OpType, Count, Callback);
 	}
 	default:
 		checkNoEntry();
 		return false;
 	}
 }
+
+#undef LOCTEXT_NAMESPACE
\ No newline at end of file
diff --git a/Engine/Plugins/PCG/Source/PCG/Private/Metadata/PCGMetadataCommon.cpp b/Engine/Plugins/PCG/Source/PCG/Private/Metadata/PCGMetadataCommon.cpp
--- a/Engine/Plugins/PCG/Source/PCG/Private/Metadata/PCGMetadataCommon.cpp
+++ b/Engine/Plugins/PCG/Source/PCG/Private/Metadata/PCGMetadataCommon.cpp
@@ -79,6 +79,79 @@ void FPCGMetadataAttributeDesc::FixLegacyTypeId()
 	FixType(KeyType, KeyTypeObject);
 }
 
+bool FPCGMetadataAttributeDesc::IsSingleValue() const
+{
+	return ContainerTypes.IsEmpty() || ContainerTypes[0] == EPCGMetadataAttributeContainerTypes::None;
+}
+
+bool FPCGMetadataAttributeDesc::IsArray() const
+{
+	return !ContainerTypes.IsEmpty() && ContainerTypes[0] == EPCGMetadataAttributeContainerTypes::Array;
+}
+
+bool FPCGMetadataAttributeDesc::IsSet() const
+{
+	return !ContainerTypes.IsEmpty() && ContainerTypes[0] == EPCGMetadataAttributeContainerTypes::Set;
+}
+
+bool FPCGMetadataAttributeDesc::IsMap() const
+{
+	return !ContainerTypes.IsEmpty() && ContainerTypes[0] == EPCGMetadataAttributeContainerTypes::Map;
+}
+
+FString FPCGMetadataAttributeDesc::GetTypeString() const
+{
+	const UEnum* TypeEnum = StaticEnum<EPCGMetadataTypes>();
+	check(TypeEnum);
+	
+	TStringBuilder<256> Builder;
+	
+	auto FillTypeStr = [TypeEnum, &Builder](EPCGMetadataTypes InType, const UObject* InTypeObject)
+	{
+		Builder += TypeEnum->GetNameStringByValue(static_cast<int64>(InType));
+
+		if (InType == EPCGMetadataTypes::Struct || InType == EPCGMetadataTypes::Enum
+			|| InType == EPCGMetadataTypes::Object || InType == EPCGMetadataTypes::SoftObject 
+			|| InType == EPCGMetadataTypes::Class || InType == EPCGMetadataTypes::SoftClass)
+		{
+			Builder += TEXT("<");
+			Builder += InTypeObject ? InTypeObject->GetName() : TEXT("UNKNOWN");
+			Builder += TEXT(">");
+		}
+	};
+	
+	if (IsSingleValue())
+	{
+		FillTypeStr(ValueType, ValueTypeObject);
+	}
+	else if (IsArray())
+	{
+		Builder += TEXT("Array<");
+		FillTypeStr(ValueType, ValueTypeObject);
+		Builder += TEXT(">");
+	}
+	else if (IsSet())
+	{
+		Builder += TEXT("Set<");
+		FillTypeStr(ValueType, ValueTypeObject);
+		Builder += TEXT(">");
+	}
+	else if (IsMap())
+	{
+		Builder += TEXT("Map<");
+		FillTypeStr(KeyType, KeyTypeObject);
+		Builder += TEXT(",");
+		FillTypeStr(ValueType, ValueTypeObject);
+		Builder += TEXT(">");
+	}
+	else
+	{
+		checkNoEntry();
+	}
+	
+	return Builder.ToString();
+}
+
 int32 GetTypeHash(const FPCGMetadataAttributeDesc& Desc)
 {
 	int32 Hash = GetTypeHash(Desc.Name);
diff --git a/Engine/Plugins/PCG/Source/PCG/Public/Metadata/Accessors/IPCGAttributeAccessor.h b/Engine/Plugins/PCG/Source/PCG/Public/Metadata/Accessors/IPCGAttributeAccessor.h
--- a/Engine/Plugins/PCG/Source/PCG/Public/Metadata/Accessors/IPCGAttributeAccessor.h
+++ b/Engine/Plugins/PCG/Source/PCG/Public/Metadata/Accessors/IPCGAttributeAccessor.h
@@ -48,35 +48,26 @@ namespace PCG::Private
 		Invalid
 	};
 	
-	template <typename T, bool bGet>
-	EPCGAttributeAccessorOperation GetOpType(const FPCGMetadataAttributeDesc& AccessorDesc, EPCGAttributeAccessorFlags Flags)
+	inline EPCGAttributeAccessorOperation GetOpType(const FPCGMetadataAttributeDesc& ReadDesc, const FPCGMetadataAttributeDesc& WriteDesc, EPCGAttributeAccessorFlags Flags)
 	{
-		const FPCGMetadataAttributeDesc* FirstDesc = &AccessorDesc;
-		const FPCGMetadataAttributeDesc* SecondDesc = &GetDefaultAttributeDesc<T>();
-		if constexpr (!bGet)
-		{
-			// Reverse order
-			Swap(FirstDesc, SecondDesc);
-		}
-		
-		if (FirstDesc->IsSameType(*SecondDesc))
+		if (ReadDesc.IsSameType(WriteDesc))
 		{
 			return EPCGAttributeAccessorOperation::SameType;
 		}
-		else if (!!(Flags & EPCGAttributeAccessorFlags::AllowConstructible) && PCG::Private::IsConstructible(*SecondDesc, *FirstDesc))
+		else if (!!(Flags & EPCGAttributeAccessorFlags::AllowConstructible) && PCG::Private::IsConstructible(WriteDesc, ReadDesc))
 		{
 			return EPCGAttributeAccessorOperation::Construct;
 		}
-		else if (!!(Flags & EPCGAttributeAccessorFlags::AllowBroadcast) && PCG::Private::IsBroadcastable(*FirstDesc, *SecondDesc))
+		else if (!!(Flags & EPCGAttributeAccessorFlags::AllowBroadcast) && PCG::Private::IsBroadcastable(ReadDesc, WriteDesc))
 		{
 			return EPCGAttributeAccessorOperation::Broadcast;
 		}
 		
-		if constexpr (bGet && std::is_same_v<T, FString>)
+		if (IsEquivalentDesc<FString>(WriteDesc))
 		{
 			// Special case - always allow broadcasting from soft object path to string, so that legacy code that grabbed soft path attributes
 			// as FStrings will still work.
-			if (IsEquivalentDesc<FSoftObjectPath>(AccessorDesc) || IsEquivalentDesc<FSoftClassPath>(AccessorDesc))
+			if (IsEquivalentDesc<FSoftObjectPath>(ReadDesc) || IsEquivalentDesc<FSoftClassPath>(ReadDesc))
 			{
 				return EPCGAttributeAccessorOperation::Broadcast;
 			}
@@ -216,6 +207,18 @@ public:
 		return SetRange_Internal(InValues, Index, Keys, Flags);
 	}
 
+	/**
+	* Copy an accessor into another.
+	* @param SrcKeys - Identification to know how to retrieve the value in this accessor
+	* @param DestAccessor - Accessor to write into
+	* @param DestKeys - Identification to know how to write the value in destination accessor
+	* @param Index - The index to start looking for in the keys
+	* @param Count - Number of values to copy
+	* @param Flags - Optional flag to allow for specific operations. Cf EPCGAttributeAccessorFlags.
+	* @return true if the copy succeeded, false otherwise
+	*/
+	PCG_API bool CopyTo(const IPCGAttributeAccessorKeys& SrcKeys, IPCGAttributeAccessor& DestAccessor, IPCGAttributeAccessorKeys& DestKeys, int32 Index, int32 Count, EPCGAttributeAccessorFlags Flags = EPCGAttributeAccessorFlags::StrictType) const;
+
 protected:
 	IPCGAttributeAccessor(bool bInReadOnly, int16 InUnderlyingType)
 		: bReadOnly(bInReadOnly)
@@ -229,7 +232,8 @@ protected:
 	template <typename T> requires (!std::is_pointer_v<T>)
 	bool GetRange_Internal(TArrayView<T> OutValues, int32 Index, const IPCGAttributeAccessorKeys& Keys, EPCGAttributeAccessorFlags Flags) const
 	{
-		PCG::Private::EPCGAttributeAccessorOperation OpType = PCG::Private::GetOpType<T, /*bGet=*/true>(UnderlyingDesc, Flags);
+		const FPCGMetadataAttributeDesc& WriteDesc = PCG::Private::GetDefaultAttributeDesc<T>();
+		PCG::Private::EPCGAttributeAccessorOperation OpType = PCG::Private::GetOpType(UnderlyingDesc, WriteDesc, Flags);
 		
 		if (OpType == PCG::Private::EPCGAttributeAccessorOperation::Invalid)
 		{
@@ -238,7 +242,7 @@ protected:
 		
 		return GetRange_Internal(
 			PCG::Private::FOutValues{TInPlaceType<PCG::Private::FOutValuesByValue>{}, OutValues.GetData()},
-			static_cast<int16>(PCG::Private::MetadataTypes<T>::Id),
+			WriteDesc,
 			OutValues.Num(),
 			Index,
 			Keys,
@@ -251,7 +255,8 @@ protected:
 	template <typename T>
 	bool GetRange_Internal(TArrayView<const T*> OutValues, int32 Index, const IPCGAttributeAccessorKeys& Keys, EPCGAttributeAccessorFlags Flags) const
 	{
-		PCG::Private::EPCGAttributeAccessorOperation OpType = PCG::Private::GetOpType<T, /*bGet=*/true>(UnderlyingDesc, Flags);
+		const FPCGMetadataAttributeDesc& WriteDesc = PCG::Private::GetDefaultAttributeDesc<T>();
+		PCG::Private::EPCGAttributeAccessorOperation OpType = PCG::Private::GetOpType(UnderlyingDesc, WriteDesc, Flags);
 		
 		if (OpType != PCG::Private::EPCGAttributeAccessorOperation::SameType)
 		{
@@ -263,7 +268,7 @@ protected:
 		Temp.SetNumUninitialized(OutValues.Num());
 		const bool bSuccess = GetRange_Internal(
 			PCG::Private::FOutValues{ TInPlaceType<PCG::Private::FOutValuesByPtr>{}, Temp },
-			static_cast<int16>(PCG::Private::MetadataTypes<T>::Id),
+			WriteDesc,
 			OutValues.Num(),
 			Index,
 			Keys,
@@ -286,7 +291,8 @@ protected:
 	template <typename T> 
 	bool GetRange_Internal(TArrayView<TConstArrayView<T>> OutValues, int32 Index, const IPCGAttributeAccessorKeys& Keys, EPCGAttributeAccessorFlags Flags) const
 	{
-		PCG::Private::EPCGAttributeAccessorOperation OpType = PCG::Private::GetOpType<TArray<T>, /*bGet=*/true>(UnderlyingDesc, Flags);
+		const FPCGMetadataAttributeDesc& WriteDesc = PCG::Private::GetDefaultAttributeDesc<TArray<T>>();
+		PCG::Private::EPCGAttributeAccessorOperation OpType = PCG::Private::GetOpType(UnderlyingDesc, WriteDesc, Flags);
 		
 		if (OpType != PCG::Private::EPCGAttributeAccessorOperation::SameType)
 		{
@@ -297,7 +303,7 @@ protected:
 		Temp.SetNumUninitialized(OutValues.Num());
 		const bool bSuccess = GetRange_Internal(
 			PCG::Private::FOutValues{TInPlaceType<PCG::Private::FOutValuesAsArray>{}, Temp},
-			static_cast<int16>(PCG::Private::MetadataTypes<T>::Id),
+			WriteDesc,
 			OutValues.Num(),
 			Index,
 			Keys,
@@ -320,7 +326,8 @@ protected:
 	template <typename T> 
 	bool GetRange_Internal(TArrayView<PCG::TScriptSetWrapper<T>> OutValues, int32 Index, const IPCGAttributeAccessorKeys& Keys, EPCGAttributeAccessorFlags Flags) const
 	{
-		PCG::Private::EPCGAttributeAccessorOperation OpType = PCG::Private::GetOpType<TSet<T>, /*bGet=*/true>(UnderlyingDesc, Flags);
+		const FPCGMetadataAttributeDesc& WriteDesc = PCG::Private::GetDefaultAttributeDesc<TSet<T>>();
+		PCG::Private::EPCGAttributeAccessorOperation OpType = PCG::Private::GetOpType(UnderlyingDesc, WriteDesc, Flags);
 		
 		if (OpType != PCG::Private::EPCGAttributeAccessorOperation::SameType)
 		{
@@ -336,7 +343,7 @@ protected:
 
 		return GetRange_Internal(
 			PCG::Private::FOutValues{TInPlaceType<PCG::Private::FOutValuesAsSet>{}, Temp},
-			static_cast<int16>(PCG::Private::MetadataTypes<T>::Id),
+			WriteDesc,
 			OutValues.Num(),
 			Index,
 			Keys,
@@ -349,7 +356,8 @@ protected:
 	template <typename KeyType, typename ValueType> 
 	bool GetRange_Internal(TArrayView<PCG::TScriptMapWrapper<KeyType, ValueType>> OutValues, int32 Index, const IPCGAttributeAccessorKeys& Keys, EPCGAttributeAccessorFlags Flags) const
 	{
-		PCG::Private::EPCGAttributeAccessorOperation OpType = PCG::Private::GetOpType<TMap<KeyType, ValueType>, /*bGet=*/true>(UnderlyingDesc, Flags);
+		const FPCGMetadataAttributeDesc& WriteDesc = PCG::Private::GetDefaultAttributeDesc<TMap<KeyType, ValueType>>();
+		PCG::Private::EPCGAttributeAccessorOperation OpType = PCG::Private::GetOpType(UnderlyingDesc, WriteDesc, Flags);
 		
 		if (OpType != PCG::Private::EPCGAttributeAccessorOperation::SameType)
 		{
@@ -365,7 +373,7 @@ protected:
 
 		return GetRange_Internal(
 			PCG::Private::FOutValues{TInPlaceType<PCG::Private::FOutValuesAsMap>{}, Temp},
-			static_cast<int16>(PCG::Private::MetadataTypes<ValueType>::Id),
+			WriteDesc,
 			OutValues.Num(),
 			Index,
 			Keys,
@@ -377,7 +385,8 @@ protected:
 	template <typename T>
 	bool SetRange_Internal(TArrayView<const T> InValues, int32 Index, IPCGAttributeAccessorKeys& Keys, EPCGAttributeAccessorFlags Flags = EPCGAttributeAccessorFlags::StrictType)
 	{
-		PCG::Private::EPCGAttributeAccessorOperation OpType = PCG::Private::GetOpType<T, /*bGet=*/false>(UnderlyingDesc, Flags);
+		const FPCGMetadataAttributeDesc& ReadDesc = PCG::Private::GetDefaultAttributeDesc<T>();
+		PCG::Private::EPCGAttributeAccessorOperation OpType = PCG::Private::GetOpType(ReadDesc, UnderlyingDesc, Flags);
 		
 		if (OpType == PCG::Private::EPCGAttributeAccessorOperation::Invalid)
 		{
@@ -386,7 +395,7 @@ protected:
 
 		return SetRange_Internal(
 			PCG::Private::FInValues{TInPlaceType<PCG::Private::FInValuesByValue>{}, InValues.GetData(), InValues.Num()},
-			static_cast<int16>(PCG::Private::MetadataTypes<T>::Id),
+			ReadDesc,
 			InValues.Num(),
 			Index,
 			Keys,
@@ -403,7 +412,8 @@ protected:
 	{
 		using T = typename ArrayContainer::ElementType;
 		
-		PCG::Private::EPCGAttributeAccessorOperation OpType = PCG::Private::GetOpType<TArray<T>, /*bGet=*/false>(UnderlyingDesc, Flags);
+		const FPCGMetadataAttributeDesc& ReadDesc = PCG::Private::GetDefaultAttributeDesc<TArray<T>>();
+		PCG::Private::EPCGAttributeAccessorOperation OpType = PCG::Private::GetOpType(ReadDesc, UnderlyingDesc, Flags);
 		
 		if (OpType != PCG::Private::EPCGAttributeAccessorOperation::SameType)
 		{
@@ -416,7 +426,7 @@ protected:
 
 		return SetRange_Internal(
 			PCG::Private::FInValues{TInPlaceType<PCG::Private::FInValuesAsArray>{}, InValuesTuple},
-			static_cast<int16>(PCG::Private::MetadataTypes<T>::Id),
+			ReadDesc,
 			InValues.Num(),
 			Index,
 			Keys,
@@ -431,7 +441,8 @@ protected:
 	template <typename T, typename ...Extra>
 	bool SetRange_Internal(TConstArrayView<TSet<T, Extra...>> InValues, int32 Index, IPCGAttributeAccessorKeys& Keys, EPCGAttributeAccessorFlags Flags)
 	{
-		PCG::Private::EPCGAttributeAccessorOperation OpType = PCG::Private::GetOpType<TSet<T>, /*bGet=*/false>(UnderlyingDesc, Flags);
+		const FPCGMetadataAttributeDesc& ReadDesc = PCG::Private::GetDefaultAttributeDesc<TSet<T>>();
+		PCG::Private::EPCGAttributeAccessorOperation OpType = PCG::Private::GetOpType(ReadDesc, UnderlyingDesc, Flags);
 		
 		if (OpType != PCG::Private::EPCGAttributeAccessorOperation::SameType)
 		{
@@ -450,7 +461,7 @@ protected:
 		
 		return SetRange_Internal(
 			PCG::Private::FInValues{TInPlaceType<PCG::Private::FInValuesAsSet>{}, InSetValues},
-			static_cast<int16>(PCG::Private::MetadataTypes<T>::Id),
+			ReadDesc,
 			InValues.Num(),
 			Index,
 			Keys,
@@ -465,7 +476,8 @@ protected:
 	template <typename KeyType, typename ValueType, typename ...Extra>
 	bool SetRange_Internal(TConstArrayView<TMap<KeyType, ValueType, Extra...>> InValues, int32 Index, IPCGAttributeAccessorKeys& Keys, EPCGAttributeAccessorFlags Flags)
 	{
-		PCG::Private::EPCGAttributeAccessorOperation OpType = PCG::Private::GetOpType<TMap<KeyType, ValueType>, /*bGet=*/false>(UnderlyingDesc, Flags);
+		const FPCGMetadataAttributeDesc& ReadDesc = PCG::Private::GetDefaultAttributeDesc<TMap<KeyType, ValueType>>();
+		PCG::Private::EPCGAttributeAccessorOperation OpType = PCG::Private::GetOpType(ReadDesc, UnderlyingDesc, Flags);
 		
 		if (OpType != PCG::Private::EPCGAttributeAccessorOperation::SameType)
 		{
@@ -484,7 +496,7 @@ protected:
 		
 		return SetRange_Internal(
 			PCG::Private::FInValues{TInPlaceType<PCG::Private::FInValuesAsMap>{}, InMapValues},
-			static_cast<int16>(PCG::Private::MetadataTypes<ValueType>::Id),
+			ReadDesc,
 			InValues.Num(),
 			Index,
 			Keys,
@@ -493,10 +505,10 @@ protected:
 		);
 	}
 	
-	PCG_API bool GetRange_Internal(PCG::Private::FOutValues OutValues, int16 ExpectedType, int32 Count, int32 Index, const IPCGAttributeAccessorKeys& Keys, EPCGAttributeAccessorFlags Flags, PCG::Private::EPCGAttributeAccessorOperation OpType) const;
+	PCG_API bool GetRange_Internal(PCG::Private::FOutValues OutValues, const FPCGMetadataAttributeDesc& WriteDesc, int32 Count, int32 Index, const IPCGAttributeAccessorKeys& Keys, EPCGAttributeAccessorFlags Flags, PCG::Private::EPCGAttributeAccessorOperation OpType) const;
 	virtual bool GetRangeVirtual(PCG::Private::FOutValues OutValues, int32 Count, int32 Index, const IPCGAttributeAccessorKeys& Keys) const = 0;
 	
-	PCG_API bool SetRange_Internal(PCG::Private::FInValues InValues, int16 ExpectedType, int32 Count, int32 Index, IPCGAttributeAccessorKeys& Keys, EPCGAttributeAccessorFlags Flags, PCG::Private::EPCGAttributeAccessorOperation OpType);
+	PCG_API bool SetRange_Internal(PCG::Private::FInValues InValues, const FPCGMetadataAttributeDesc& ReadDesc, int32 Count, int32 Index, IPCGAttributeAccessorKeys& Keys, EPCGAttributeAccessorFlags Flags, PCG::Private::EPCGAttributeAccessorOperation OpType);
 	virtual bool SetRangeVirtual(PCG::Private::FInValues InValues, int32 Count, int32 Index, IPCGAttributeAccessorKeys& Keys, EPCGAttributeAccessorFlags Flags) = 0;
 
 	bool bReadOnly = true;
diff --git a/Engine/Plugins/PCG/Source/PCG/Public/Metadata/PCGMetadataCommon.h b/Engine/Plugins/PCG/Source/PCG/Public/Metadata/PCGMetadataCommon.h
--- a/Engine/Plugins/PCG/Source/PCG/Public/Metadata/PCGMetadataCommon.h
+++ b/Engine/Plugins/PCG/Source/PCG/Public/Metadata/PCGMetadataCommon.h
@@ -299,6 +299,21 @@ struct FPCGMetadataAttributeDesc
 	
 	// For already predefined structs, replace `EPCGMetadataTypes::Struct` by its predefined type.
 	PCG_API void FixLegacyTypeId();
+	
+	/** Returns true if the attribute is not a container (Array/Set/Map). */
+	PCG_API bool IsSingleValue() const;
+
+	/** Returns true if the attribute is an array. */
+	PCG_API bool IsArray() const;
+
+	/** Returns true if the attribute is a set. */
+	PCG_API bool IsSet() const;
+
+	/** Returns true if the attribute is a map. */
+	PCG_API bool IsMap() const;
+
+	/** Returns the string representing the type of this descriptor. */
+	FString GetTypeString() const;
 };
 
 /**
diff --git a/Engine/Plugins/PCG/Tests/Private/Metadata/Accessors/PCGAttributeAccessorTests.cpp b/Engine/Plugins/PCG/Tests/Private/Metadata/Accessors/PCGAttributeAccessorTests.cpp
--- a/Engine/Plugins/PCG/Tests/Private/Metadata/Accessors/PCGAttributeAccessorTests.cpp
+++ b/Engine/Plugins/PCG/Tests/Private/Metadata/Accessors/PCGAttributeAccessorTests.cpp
@@ -25,6 +25,10 @@
 //  - GetWithConstructibleType:          Get/Set with constructible type conversion (float<->double) using AllowConstructible flag
 //  - GetWithBroadcastableType:          Get/Set with broadcastable type conversion (float->FVector/FString) using AllowBroadcast flag
 //  - SoftObjectPathToStringWithStrictType: FSoftObjectPath/FSoftClassPath to FString conversion with StrictType flag
+//  - CopyToSameType:                     CopyTo copies values between same-type accessors (float, vector, string)
+//  - CopyToPartialRange:                 CopyTo copies a sub-range leaving untouched entries at their defaults
+//  - CopyToConstructibleType:            CopyTo with AllowConstructible copies between constructible types (float->double)
+//  - CopyToBroadcastableType:            CopyTo with AllowBroadcast copies between broadcastable types (float->FVector)
 // =============================================================================
 
 #include "PCGTestsCommon.h"
@@ -883,3 +887,230 @@ TEST_CASE_METHOD(PCGTests::FPCGBaseTestWithBasicContext, "PCG::Accessor::Attribu
 		REQUIRE_FALSE(SoftPathAccessor->Get(TempFloat, *Keys, EPCGAttributeAccessorFlags::StrictType));
 	}
 }
+
+// Tests that CopyTo copies values between accessors of the same underlying type (float->float, vector->vector, string->string).
+TEST_CASE_METHOD(PCGTests::FPCGBaseTestWithBasicContext, "PCG::Accessor::Attribute::CopyToSameType", "[PCG][Accessor][Attribute]")
+{
+	using namespace PCGAttributeAccessorTests;
+
+	AttributeData SrcData(CreatePointData());
+
+	TUniquePtr<const IPCGAttributeAccessorKeys> SrcKeys = PCGAttributeAccessorHelpers::CreateConstKeys(SrcData.PointData, FPCGAttributePropertySelector::CreateAttributeSelector(TEXT("FloatAttribute")));
+	REQUIRE(SrcKeys);
+	REQUIRE_EQUAL(SrcKeys->GetNum(), ValuesSize+1);
+
+	// Destination data with entries (values at defaults)
+	auto CopyAndValidate = [this, &SrcKeys]<typename T>(const IPCGAttributeAccessor& SrcAccessor, FName Name, T DestDefaultValue, const T(&Values)[3])
+	{
+		UPCGBasePointData* DestPointData = CreatePointData();
+		FPCGMetadataAttribute<T>* DestAttr = DestPointData->MutableMetadata()->CreateAttribute<T>(Name, DestDefaultValue, false, false);
+		DestPointData->SetNumPoints(ValuesSize+1);
+		TPCGValueRange<int64> DestMetadataEntryRange = DestPointData->GetMetadataEntryValueRange();
+
+		TUniquePtr<IPCGAttributeAccessor> DestAccessor = PCGAttributeAccessorHelpers::CreateAccessor(DestAttr, DestPointData->MutableMetadata());
+		TUniquePtr<IPCGAttributeAccessorKeys> DestKeys = PCGAttributeAccessorHelpers::CreateKeys(DestPointData, FPCGAttributePropertySelector::CreateAttributeSelector(Name));
+		REQUIRE(DestKeys);
+		REQUIRE_EQUAL(DestKeys->GetNum(), ValuesSize+1);
+		
+		REQUIRE(SrcAccessor.CopyTo(*SrcKeys, *DestAccessor, *DestKeys, 0, ValuesSize));
+		for (int32 i = 0; i < ValuesSize; ++i)
+		{
+			// Reverse order for values.
+			REQUIRE_EQUAL(DestAttr->GetValueFromItemKey(DestMetadataEntryRange[i]), Values[ValuesSize-i-1]);
+		}
+		
+		REQUIRE_EQUAL(DestAttr->GetValueFromItemKey(DestMetadataEntryRange[3]), DestDefaultValue);
+	};
+
+	SECTION("Float")
+	{
+		CopyAndValidate(*SrcData.FloatAccessor, TEXT("DestFloat"), 0.1f, FloatValues);
+	}
+
+	SECTION("Vector")
+	{
+		CopyAndValidate(*SrcData.VectorAccessor, TEXT("DestVector"), FVector::OneVector, VectorValues);
+	}
+
+	SECTION("String")
+	{
+		CopyAndValidate(*SrcData.StringAccessor, TEXT("DestString"), FString{TEXT("Journey before Destination")}, StringValues);
+	}
+}
+
+// Tests that CopyTo can copy a sub-range starting at a non-zero index, leaving entries outside the range at their attribute defaults.
+TEST_CASE_METHOD(PCGTests::FPCGBaseTestWithBasicContext, "PCG::Accessor::Attribute::CopyToPartialRange", "[PCG][Accessor][Attribute]")
+{
+	using namespace PCGAttributeAccessorTests;
+
+	// Source with sequential mapping
+	AttributeData SrcData(CreatePointData());
+	TUniquePtr<const IPCGAttributeAccessorKeys> SrcKeys = PCGAttributeAccessorHelpers::CreateConstKeys(SrcData.PointData, FPCGAttributePropertySelector::CreateAttributeSelector(TEXT("FloatAttribute")));
+	REQUIRE(SrcKeys);
+
+	// Destination with distinctive default
+	constexpr float DestDefault = -1.0f;
+	UPCGBasePointData* DestPointData = CreatePointData();
+	FPCGMetadataAttribute<float>* DestFloatAttr = DestPointData->MutableMetadata()->CreateAttribute<float>(TEXT("DestFloat"), DestDefault, true, false);
+
+	DestPointData->SetNumPoints(ValuesSize + 1);
+	TPCGValueRange<int64> DestMetadataEntryRange = DestPointData->GetMetadataEntryValueRange();
+
+	TUniquePtr<IPCGAttributeAccessor> DestFloatAccessor = PCGAttributeAccessorHelpers::CreateAccessor(DestFloatAttr, DestPointData->MutableMetadata());
+	TUniquePtr<IPCGAttributeAccessorKeys> DestKeys = PCGAttributeAccessorHelpers::CreateKeys(DestPointData, FPCGAttributePropertySelector::CreateAttributeSelector(TEXT("DestFloat")));
+	REQUIRE(DestKeys);
+
+	SECTION("Copy from index 1 with count 2")
+	{
+		REQUIRE(SrcData.FloatAccessor->CopyTo(*SrcKeys, *DestFloatAccessor, *DestKeys, 1, 2));
+		// Entry 0 and 3 should be untouched (default)
+		REQUIRE_EQUAL(DestFloatAttr->GetValueFromItemKey(DestMetadataEntryRange[0]), DestDefault);
+		REQUIRE_EQUAL(DestFloatAttr->GetValueFromItemKey(DestMetadataEntryRange[3]), DestDefault);
+		// Entries 1 and 2 should have source values (reverse order)
+		REQUIRE_EQUAL(DestFloatAttr->GetValueFromItemKey(DestMetadataEntryRange[1]), FloatValues[1]);
+		REQUIRE_EQUAL(DestFloatAttr->GetValueFromItemKey(DestMetadataEntryRange[2]), FloatValues[0]);
+	}
+
+	SECTION("Copy only first element")
+	{
+		REQUIRE(SrcData.FloatAccessor->CopyTo(*SrcKeys, *DestFloatAccessor, *DestKeys, 0, 1));
+		// Entry 0 should have source value (reverse order)
+		REQUIRE_EQUAL(DestFloatAttr->GetValueFromItemKey(DestMetadataEntryRange[0]), FloatValues[2]);
+		// Entries 1, 2 and 3 should be untouched (default)
+		REQUIRE_EQUAL(DestFloatAttr->GetValueFromItemKey(DestMetadataEntryRange[1]), DestDefault);
+		REQUIRE_EQUAL(DestFloatAttr->GetValueFromItemKey(DestMetadataEntryRange[2]), DestDefault);
+		REQUIRE_EQUAL(DestFloatAttr->GetValueFromItemKey(DestMetadataEntryRange[3]), DestDefault);
+	}
+}
+
+// Tests that CopyTo with AllowConstructible copies between constructible types (float->double), fails without the flag, and works with larger element counts.
+TEST_CASE_METHOD(PCGTests::FPCGBaseTestWithBasicContext, "PCG::Accessor::Attribute::CopyToConstructibleType", "[PCG][Accessor][Attribute]")
+{
+	using namespace PCGAttributeAccessorTests;
+	
+	auto ExecuteTest = [this](int32 Count, EPCGAttributeAccessorFlags Flags, bool bShouldFail)
+	{
+		// Source: float
+		UPCGBasePointData* SrcPointData = CreatePointData();
+		FPCGMetadataAttribute<float>* SrcFloatAttr = SrcPointData->MutableMetadata()->CreateAttribute<float>(TEXT("SrcFloat"), 0.0f, true, false);
+		SrcPointData->SetNumPoints(Count);
+		TPCGValueRange<int64> SrcEntryRange = SrcPointData->GetMetadataEntryValueRange();
+		TArray<float> Values;
+		Values.SetNumUninitialized(Count);
+
+		FRandomStream RandomStream(42);
+
+		for (int32 i = 0; i < Count; ++i)
+		{
+			SrcEntryRange[i] = SrcPointData->MutableMetadata()->AddEntry();
+			Values[i] = RandomStream.FRand();
+			SrcFloatAttr->SetValue(SrcEntryRange[i], Values[i]);
+		}
+			
+		TUniquePtr<IPCGAttributeAccessor> SrcFloatAccessor = PCGAttributeAccessorHelpers::CreateAccessor(SrcFloatAttr, SrcPointData->MutableMetadata());
+		TUniquePtr<const IPCGAttributeAccessorKeys> SrcKeys = PCGAttributeAccessorHelpers::CreateConstKeys(SrcPointData, FPCGAttributePropertySelector::CreateAttributeSelector(TEXT("SrcFloat")));
+		REQUIRE(SrcKeys);
+
+		// Destination: double
+		UPCGBasePointData* DestPointData = CreatePointData();
+		FPCGMetadataAttribute<double>* DestDoubleAttr = DestPointData->MutableMetadata()->CreateAttribute<double>(TEXT("DestDouble"), 0.0, true, false);
+		DestPointData->SetNumPoints(Count);
+		TPCGValueRange<int64> DestEntryRange = DestPointData->GetMetadataEntryValueRange();
+			
+		TUniquePtr<IPCGAttributeAccessor> DestDoubleAccessor = PCGAttributeAccessorHelpers::CreateAccessor(DestDoubleAttr, DestPointData->MutableMetadata());
+		TUniquePtr<IPCGAttributeAccessorKeys> DestKeys = PCGAttributeAccessorHelpers::CreateKeys(DestPointData, FPCGAttributePropertySelector::CreateAttributeSelector(TEXT("DestDouble")));
+		REQUIRE(DestKeys);
+			
+		REQUIRE(SrcFloatAccessor->CopyTo(*SrcKeys, *DestDoubleAccessor, *DestKeys, 0, Count, Flags) != bShouldFail);
+		if (!bShouldFail)
+		{
+			for (int32 i = 0; i < Count; ++i)
+			{
+				REQUIRE_EQUAL(DestDoubleAttr->GetValueFromItemKey(DestEntryRange[i]), static_cast<double>(Values[i]));
+			}
+		}
+	};
+	
+	SECTION("CopyTo float->double without AllowConstructible fails")
+	{
+		PCGTests::FPCGTestsLogOutputDevice LogCapture(/*bSuppressErrors=*/true);
+		ExecuteTest(ValuesSize, EPCGAttributeAccessorFlags::StrictType, /*bShouldFail=*/true);
+		REQUIRE(LogCapture.NbMessageReceived.Contains(ELogVerbosity::Error));
+	}
+
+	SECTION("CopyTo float->double with AllowConstructible succeeds")
+	{
+		ExecuteTest(ValuesSize, EPCGAttributeAccessorFlags::AllowConstructible, /*bShouldFail=*/false);
+	}
+
+	SECTION("CopyTo float->double with large count")
+	{
+		ExecuteTest(1024, EPCGAttributeAccessorFlags::AllowConstructible, /*bShouldFail=*/false);
+	}
+}
+
+// Tests that CopyTo with AllowBroadcast copies between broadcastable types (float->FVector), fails without the flag, and works with larger element counts.
+TEST_CASE_METHOD(PCGTests::FPCGBaseTestWithBasicContext, "PCG::Accessor::Attribute::CopyToBroadcastableType", "[PCG][Accessor][Attribute]")
+{
+	using namespace PCGAttributeAccessorTests;
+
+	auto ExecuteTest = [this](int32 Count, EPCGAttributeAccessorFlags Flags, bool bShouldFail)
+	{
+		// Source: float
+		UPCGBasePointData* SrcPointData = CreatePointData();
+		FPCGMetadataAttribute<float>* SrcFloatAttr = SrcPointData->MutableMetadata()->CreateAttribute<float>(TEXT("SrcFloat"), 0.0f, true, false);
+		SrcPointData->SetNumPoints(Count);
+		TPCGValueRange<int64> SrcEntryRange = SrcPointData->GetMetadataEntryValueRange();
+		TArray<float> Values;
+		Values.SetNumUninitialized(Count);
+
+		FRandomStream RandomStream(42);
+
+		for (int32 i = 0; i < Count; ++i)
+		{
+			SrcEntryRange[i] = SrcPointData->MutableMetadata()->AddEntry();
+			Values[i] = RandomStream.FRand();
+			SrcFloatAttr->SetValue(SrcEntryRange[i], Values[i]);
+		}
+	
+		TUniquePtr<IPCGAttributeAccessor> SrcFloatAccessor = PCGAttributeAccessorHelpers::CreateAccessor(SrcFloatAttr, SrcPointData->MutableMetadata());
+		TUniquePtr<const IPCGAttributeAccessorKeys> SrcKeys = PCGAttributeAccessorHelpers::CreateConstKeys(SrcPointData, FPCGAttributePropertySelector::CreateAttributeSelector(TEXT("SrcFloat")));
+		REQUIRE(SrcKeys);
+
+		// Destination: FVector
+		UPCGBasePointData* DestPointData = CreatePointData();
+		FPCGMetadataAttribute<FVector>* DestVectorAttr = DestPointData->MutableMetadata()->CreateAttribute<FVector>(TEXT("DestVector"), FVector::ZeroVector, false, false);
+		DestPointData->SetNumPoints(ValuesSize);
+		TPCGValueRange<int64> DestEntryRange = DestPointData->GetMetadataEntryValueRange();
+
+		TUniquePtr<IPCGAttributeAccessor> DestVectorAccessor = PCGAttributeAccessorHelpers::CreateAccessor(DestVectorAttr, DestPointData->MutableMetadata());
+		TUniquePtr<IPCGAttributeAccessorKeys> DestKeys = PCGAttributeAccessorHelpers::CreateKeys(DestPointData, FPCGAttributePropertySelector::CreateAttributeSelector(TEXT("DestVector")));
+		REQUIRE(DestKeys);
+			
+		REQUIRE(SrcFloatAccessor->CopyTo(*SrcKeys, *DestVectorAccessor, *DestKeys, 0, ValuesSize, Flags) != bShouldFail);
+		if (!bShouldFail)
+		{
+			for (int32 i = 0; i < ValuesSize; ++i)
+			{
+				REQUIRE_EQUAL(DestVectorAttr->GetValueFromItemKey(DestEntryRange[i]), FVector(Values[i]));
+			}
+		}
+	};
+
+	SECTION("CopyTo float->FVector without AllowBroadcast fails")
+	{
+		PCGTests::FPCGTestsLogOutputDevice LogCapture(/*bSuppressErrors=*/true);
+		ExecuteTest(ValuesSize, EPCGAttributeAccessorFlags::StrictType, /*bShouldFail=*/true);
+		REQUIRE(LogCapture.NbMessageReceived.Contains(ELogVerbosity::Error));
+	}
+
+	SECTION("CopyTo float->FVector with AllowBroadcast succeeds")
+	{
+		ExecuteTest(ValuesSize, EPCGAttributeAccessorFlags::AllowBroadcast, /*bShouldFail=*/false);
+	}
+
+	SECTION("CopyTo float->FVector with large count")
+	{
+		ExecuteTest(1024, EPCGAttributeAccessorFlags::AllowBroadcast, /*bShouldFail=*/false);
+	}
+}
\ No newline at end of file
