[PCG] Fix Code Editor giving no visual indication when it's locked, such as when no Custom HLSL node is selected, or when a node has a source asset assigned (so inline source disabled).

#rnx
#Synced-CL 52135071

diff --git a/Engine/Plugins/PCG/Source/PCGEditor/Private/Widgets/SPCGCodeEditorTextBox.cpp b/Engine/Plugins/PCG/Source/PCGEditor/Private/Widgets/SPCGCodeEditorTextBox.cpp
--- a/Engine/Plugins/PCG/Source/PCGEditor/Private/Widgets/SPCGCodeEditorTextBox.cpp
+++ b/Engine/Plugins/PCG/Source/PCGEditor/Private/Widgets/SPCGCodeEditorTextBox.cpp
@@ -148,7 +148,7 @@ void SPCGCodeEditorTextBox::Construct(const FArguments& InArgs)
 		[
 			SNew(SBorder)
 			.BorderImage(FPCGEditorStyle::Get().GetBrush("TextEditor.Border"))
-			.BorderBackgroundColor(FLinearColor(0.55f, 0.55f, 0.55f, 1.0f))
+			.BorderBackgroundColor(this, &SPCGCodeEditorTextBox::GetBorderBackgroundColor)
 			[
 				SNew(SGridPanel)
 				.FillColumn(0,1.0f)
@@ -758,4 +758,10 @@ bool SPCGCodeEditorTextBox::IsReadOnly() const
 	return true;
 }
 
+FSlateColor SPCGCodeEditorTextBox::GetBorderBackgroundColor() const
+{
+	// Hand tweaked values.
+	return IsReadOnly() ? FLinearColor(0.72f, 0.72f, 0.72f, 1.0f) : FLinearColor(0.4f, 0.4f, 0.4f, 1.0f);
+}
+
 #undef LOCTEXT_NAMESPACE
diff --git a/Engine/Plugins/PCG/Source/PCGEditor/Private/Widgets/SPCGCodeEditorTextBox.h b/Engine/Plugins/PCG/Source/PCGEditor/Private/Widgets/SPCGCodeEditorTextBox.h
--- a/Engine/Plugins/PCG/Source/PCGEditor/Private/Widgets/SPCGCodeEditorTextBox.h
+++ b/Engine/Plugins/PCG/Source/PCGEditor/Private/Widgets/SPCGCodeEditorTextBox.h
@@ -102,6 +102,7 @@ private:
 	void HandleAutoIndent() const;
 
 	bool IsReadOnly() const;
+	FSlateColor GetBorderBackgroundColor() const;
 
 private:
 	TSharedPtr<SVerticalBox> TabBody;
