[PCG] Minor updates to improve usage of IPCGSettingsInlineConstant

#rnx
#Synced-CL 52262938

diff --git a/Engine/Plugins/PCG/Source/PCG/Private/Metadata/PCGInlineConstantInterface.cpp b/Engine/Plugins/PCG/Source/PCG/Private/Metadata/PCGInlineConstantInterface.cpp
--- a/Engine/Plugins/PCG/Source/PCG/Private/Metadata/PCGInlineConstantInterface.cpp
+++ b/Engine/Plugins/PCG/Source/PCG/Private/Metadata/PCGInlineConstantInterface.cpp
@@ -10,6 +10,32 @@
 
 #define LOCTEXT_NAMESPACE "PCGInlineConstantInterface"
 
+namespace PCGInlineConstant::Helpers
+{
+	bool IsValidDefaultValueType(const EPCGMetadataTypes Type)
+	{
+		return Type != EPCGMetadataTypes::Unknown && PCGMetadataHelpers::MetadataTypeSupportsDefaultValues(Type);
+	}
+
+	FString GetZeroValueSerializedString(const EPCGMetadataTypes InType)
+	{
+		return PCGMetadataAttribute::CallbackWithRightType(static_cast<uint16>(InType), []<typename T>(const T&)
+		{
+			const T Value = PCG::Private::MetadataTraits<T>::ZeroValue();
+			if constexpr (requires { TBaseStructure<T>::Get(); })
+			{
+				FString Result;
+				TBaseStructure<T>::Get()->ExportText(Result, &Value, /*Defaults=*/nullptr, /*OwnerObject=*/nullptr, PPF_None, /*ExportRootScope=*/nullptr);
+				return Result;
+			}
+			else
+			{
+				return LexToString(Value);
+			}
+		});
+	}
+}
+
 bool FPCGInlineConstantState::DefaultValuesAreEnabled() const
 {
 	return true;
@@ -79,20 +105,7 @@ FString FPCGInlineConstantState::GetPinInitialDefaultValueString(const FName Pin
 {
 	if (const FPCGPinDefaultValueInfo* DefaultValueInfo = PinLabel != NAME_None ? DefaultValueMap.Find(PinLabel) : nullptr)
 	{
-		return PCGMetadataAttribute::CallbackWithRightType(static_cast<uint16>(DefaultValueInfo->InitialType), []<typename T>(const T&)
-		{
-			const T Value = PCG::Private::MetadataTraits<T>::ZeroValue();
-			if constexpr (requires { TBaseStructure<T>::Get(); })
-			{
-				FString Result;
-				TBaseStructure<T>::Get()->ExportText(Result, &Value, /*Defaults=*/nullptr, /*OwnerObject=*/nullptr, PPF_None, /*ExportRootScope=*/nullptr);
-				return Result;
-			}
-			else
-			{
-				return LexToString(Value);
-			}
-		});
+		return PCGInlineConstant::Helpers::GetZeroValueSerializedString(DefaultValueInfo->InitialType);
 	}
 
 	return FString{};
@@ -272,21 +285,37 @@ bool IPCGSettingsInlineConstant::IsPinDefaultValueEnabled(const FName PinLabel)
 		return true;
 	}
 
-	// Fallback here supporting pins that are eligible if their initial type supports default values (pre-activation).
-	const EPCGMetadataTypes InitialType = GetPinInitialDefaultValueType(PinLabel);
-	return InitialType != EPCGMetadataTypes::Unknown && PCGMetadataHelpers::MetadataTypeSupportsDefaultValues(InitialType);
+	// Unregistered pins are eligible if their initial type supports default values.
+	return PCGInlineConstant::Helpers::IsValidDefaultValueType(GetPinInitialDefaultValueType(PinLabel));
 }
 
 bool IPCGSettingsInlineConstant::IsPinDefaultValueActivated(const FName PinLabel) const
 {
 	const FPCGInlineConstantState* State = GetInlineConstantState();
-	return State && State->IsPinDefaultValueActivated(PinLabel);
+	if (State && State->IsPinDefaultValueEnabled(PinLabel))
+	{
+		return State->IsPinDefaultValueActivated(PinLabel);
+	}
+
+	// Treat unregistered pins as activated if their initial type supports default values.
+	// First user edit materializes the entry in state via lazy init in SetPinDefaultValueIsActivated.
+	return PCGInlineConstant::Helpers::IsValidDefaultValueType(GetPinInitialDefaultValueType(PinLabel));
 }
 
 EPCGMetadataTypes IPCGSettingsInlineConstant::GetPinDefaultValueType(const FName PinLabel) const
 {
 	const FPCGInlineConstantState* State = GetInlineConstantState();
-	return State ? State->GetPinDefaultValueType(PinLabel) : EPCGMetadataTypes::Unknown;
+	if (State)
+	{
+		const EPCGMetadataTypes Type = State->GetPinDefaultValueType(PinLabel);
+		if (Type != EPCGMetadataTypes::Unknown)
+		{
+			return Type;
+		}
+	}
+
+	// Use the initial type for unregistered pins.
+	return GetPinInitialDefaultValueType(PinLabel);
 }
 
 #if WITH_EDITOR
@@ -380,14 +409,14 @@ void IPCGSettingsInlineConstant::SetPinDefaultValueIsActivated(const FName PinLa
 	if (!State->PinIsMapped(PinLabel))
 	{
 		const EPCGMetadataTypes InitialType = GetPinInitialDefaultValueType(PinLabel);
-		if (InitialType == EPCGMetadataTypes::Unknown || !PCGMetadataHelpers::MetadataTypeSupportsDefaultValues(InitialType))
+		if (!PCGInlineConstant::Helpers::IsValidDefaultValueType(InitialType))
 		{
 			return;
 		}
 
 		FPCGPinDefaultValueInfo::FInitParams InitParams;
 		InitParams.InitialType = InitialType;
-		InitParams.bStartActivated = false;
+		InitParams.bStartActivated = true;
 		State->AddDefaultValueToPin(PinLabel, MoveTemp(InitParams));
 	}
 
@@ -412,7 +441,19 @@ void IPCGSettingsInlineConstant::SetPinDefaultValueIsActivated(const FName PinLa
 FString IPCGSettingsInlineConstant::GetPinDefaultValueAsString(const FName PinLabel) const
 {
 	const FPCGInlineConstantState* State = GetInlineConstantState();
-	return State ? State->GetPinDefaultValueAsString(PinLabel) : FString{};
+	if (State && State->IsPinDefaultValueEnabled(PinLabel))
+	{
+		return State->GetPinDefaultValueAsString(PinLabel);
+	}
+
+	// Zero value string for unregistered pins.
+	const EPCGMetadataTypes InitialType = GetPinInitialDefaultValueType(PinLabel);
+	if (PCGInlineConstant::Helpers::IsValidDefaultValueType(InitialType))
+	{
+		return PCGInlineConstant::Helpers::GetZeroValueSerializedString(InitialType);
+	}
+
+	return FString{};
 }
 
 FString IPCGSettingsInlineConstant::GetPinInitialDefaultValueString(const FName PinLabel) const
