[PCG] Removed need for a full purge on execution source ended and added some validation. Fixed a missing resource collection in the subsystem too.

#rnx
#Synced-CL 52283469

diff --git a/Engine/Plugins/PCG/Source/PCG/Private/Subsystems/IPCGBaseSubsystem.cpp b/Engine/Plugins/PCG/Source/PCG/Private/Subsystems/IPCGBaseSubsystem.cpp
--- a/Engine/Plugins/PCG/Source/PCG/Private/Subsystems/IPCGBaseSubsystem.cpp
+++ b/Engine/Plugins/PCG/Source/PCG/Private/Subsystems/IPCGBaseSubsystem.cpp
@@ -284,6 +284,7 @@ uint32 IPCGBaseSubsystem::GetGraphCacheEntryCount(IPCGElement* InElement) const
 
 void IPCGBaseSubsystem::OnPCGSourceGenerationDone(IPCGGraphExecutionSource* InExecutionSource, EPCGGenerationStatus InStatus)
 {
+	check(IsInGameThread());
 	OnPCGSourceGenerationDoneDelegate.Broadcast(this, InExecutionSource, InStatus);
 
 	UPCGDefaultExecutionSource* ExecutionSource = Cast<UPCGDefaultExecutionSource>(InExecutionSource);
@@ -307,7 +308,7 @@ void IPCGBaseSubsystem::OnPCGSourceGenerationDone(IPCGGraphExecutionSource* InEx
 
 		// Mark source as invalid, request GC (non-blocking)
 		ExecutionSource->MarkAsGarbage();
-		GEngine->ForceGarbageCollection(true);
+		GEngine->ForceGarbageCollection();
 	}
 }
 
diff --git a/Engine/Plugins/PCG/Source/PCG/Private/Subsystems/PCGSubsystem.cpp b/Engine/Plugins/PCG/Source/PCG/Private/Subsystems/PCGSubsystem.cpp
--- a/Engine/Plugins/PCG/Source/PCG/Private/Subsystems/PCGSubsystem.cpp
+++ b/Engine/Plugins/PCG/Source/PCG/Private/Subsystems/PCGSubsystem.cpp
@@ -166,6 +166,9 @@ void UPCGSubsystem::AddReferencedObjects(UObject* InThis, FReferenceCollector& C
 	Super::AddReferencedObjects(InThis, Collector);
 
 	UPCGSubsystem* PCGSubsystem = CastChecked<UPCGSubsystem>(InThis);
+#if WITH_EDITOR
+	PCGSubsystem->IPCGBaseSubsystem::AddReferencedObjects(Collector);
+#endif // WITH_EDITOR
 	
 	if (FPCGRuntimeGenScheduler* RuntimeGenScheduler = PCGSubsystem->GetRuntimeGenScheduler())
 	{
