[PCG] Infer created attributes from HLSL Set calls to make adding/changing output attributes far easier.

The old "Attributes to Create" array can still be used, and is used for old nodes which had entries to preserve old behaviour.

[RN] [PCG] Added parsing of created attributes directly from Custom HLSL source, so that attributes no longer need manual configuration on output pin properties.

#Synced-CL 52135071

diff --git a/Engine/Plugins/PCG/Source/PCG/Private/Compute/Elements/PCGCustomHLSL.cpp b/Engine/Plugins/PCG/Source/PCG/Private/Compute/Elements/PCGCustomHLSL.cpp
--- a/Engine/Plugins/PCG/Source/PCG/Private/Compute/Elements/PCGCustomHLSL.cpp
+++ b/Engine/Plugins/PCG/Source/PCG/Private/Compute/Elements/PCGCustomHLSL.cpp
@@ -19,6 +19,7 @@
 
 #include "Internationalization/Regex.h"
 #include "Containers/StaticArray.h"
+#include "UObject/FortniteMainBranchObjectVersion.h"
 
 #if WITH_EDITOR
 #include "ScopedTransaction.h"
@@ -66,6 +67,21 @@ void UPCGCustomHLSLSettings::PostLoad()
 		NumElementsY = NumElements2D_DEPRECATED.Y;
 		NumElements2D_DEPRECATED = FIntPoint(0, 0);
 	}
+
+	if (GetLinkerCustomVersion(FFortniteMainBranchObjectVersion::GUID) < FFortniteMainBranchObjectVersion::PCGCustomHLSLAttributeCreationSettingsToggle)
+	{
+		// For old nodes: enable the toggle if any output pin already has manually authored attribute keys,
+		// so existing settings remain visible after the upgrade. New nodes default to false (hidden).
+		bPerPinAttributeCreationSettings = false;
+		for (const FPCGPinPropertiesGPU& OutputPin : OutputPins)
+		{
+			if (OutputPin.PropertiesGPU.CreatedKernelAttributeKeys.Num() > 0)
+			{
+				bPerPinAttributeCreationSettings = true;
+				break;
+			}
+		}
+	}
 #endif
 
 	// Note: We update here so that Custom HLSL nodes will have the correct pin settings & declarations on load.
@@ -1213,12 +1229,16 @@ void UPCGCustomHLSLSettings::UpdatePinSettings()
 		// Static initialization to zero is allowed for pins that the kernel is not controlling, and is not currently implemented for textures.
 		Properties.PropertiesGPU.bShowInitializeToZeroSetting = !bInitModeDrivenByKernel
 			&& (!!(Properties.AllowedTypes & EPCGDataType::Point) || !!(Properties.AllowedTypes & EPCGDataType::Param));
+
+		// Show the legacy "Attributes to Create" array only when the user has explicitly opted in.
+		Properties.PropertiesGPU.bShowCreatedAttributesSettings = bPerPinAttributeCreationSettings;
 	}
 }
 
 void UPCGCustomHLSLSettings::UpdateAttributeKeys()
 {
-	// Make sure all the output attributes are up-to-date on their identifiers.
+	// Keep identifiers on manually-authored CreatedKernelAttributeKeys in sync with their selectors.
+	// These authored entries are still respected (as a base that HLSL inference augments).
 	bool bMarkedDirty = false;
 	for (FPCGPinPropertiesGPU& OutputPin : OutputPins)
 	{
diff --git a/Engine/Plugins/PCG/Source/PCG/Private/Compute/Elements/PCGCustomHLSL.h b/Engine/Plugins/PCG/Source/PCG/Private/Compute/Elements/PCGCustomHLSL.h
--- a/Engine/Plugins/PCG/Source/PCG/Private/Compute/Elements/PCGCustomHLSL.h
+++ b/Engine/Plugins/PCG/Source/PCG/Private/Compute/Elements/PCGCustomHLSL.h
@@ -97,6 +97,9 @@ public:
 
 	/** Sets the kernel type and updates pin settings and declarations. */
 	void SetKernelType(EPCGKernelType InKernelType);
+
+	/** Sets bPerPinAttributeCreationSettings. When false (default), authored CreatedKernelAttributeKeys are ignored and only HLSL inference applies. */
+	void SetPerPinAttributeCreationSettings(bool bInValue) { bPerPinAttributeCreationSettings = bInValue; }
 #endif
 
 #if WITH_EDITOR
@@ -224,6 +227,10 @@ protected:
 	UPROPERTY(EditAnywhere, Category = "Settings")
 	bool bMuteUnwrittenPinDataErrors = false;
 
+	/** Expose a per-pin "Attributes to Create" array so attributes can be manually specified. Attributes written via Set calls in HLSL (e.g. Out_SetFloat) are now inferred automatically, so this is no longer required. */
+	UPROPERTY(EditAnywhere, AdvancedDisplay, Category = "Settings")
+	bool bPerPinAttributeCreationSettings = false;
+
 #if WITH_EDITORONLY_DATA
 	// Shader source and declarations are entirely editor-only, and should never be serialized outside of the editor.
 
diff --git a/Engine/Plugins/PCG/Source/PCG/Private/Compute/Elements/PCGCustomHLSLKernel.cpp b/Engine/Plugins/PCG/Source/PCG/Private/Compute/Elements/PCGCustomHLSLKernel.cpp
--- a/Engine/Plugins/PCG/Source/PCG/Private/Compute/Elements/PCGCustomHLSLKernel.cpp
+++ b/Engine/Plugins/PCG/Source/PCG/Private/Compute/Elements/PCGCustomHLSLKernel.cpp
@@ -184,8 +184,12 @@ void UPCGCustomHLSLKernel::InitializeInternal()
 	Super::InitializeInternal();
 
 	InitEntryPoint();
-	PopulateAttributeKeysFromPinSettings();
+
+	// Must run before PopulateAttributeKeys to populate ParsedSources.
 	ParseShaderSource();
+
+	// Combine CreatedKernelAttributeKeys + inferred HLSL keys into PinCreatedAttributeKeys.
+	PopulateAttributeKeys();
 }
 #endif // WITH_EDITOR
 
@@ -316,10 +320,14 @@ TSharedPtr<const FPCGDataCollectionDesc> UPCGCustomHLSLKernel::ComputeOutputBind
 		return nullptr;
 	}
 
-	// Add attributes that will be created for this pin on the GPU. This will stomp any existing attributes if they collide!
-	for (const FPCGKernelAttributeKey& CreatedKey : OutputPinProperties->PropertiesGPU.CreatedKernelAttributeKeys)
+	// Add attributes created for this output pin. PinCreatedAttributeKeys is pre-computed during editor init
+	// (in PopulateAttributeKeys) from both manually-authored CreatedKernelAttributeKeys and HLSL Set/AtomicAdd calls.
+	if (const FPCGKernelAttributeKeyList* CreatedKeyList = PinCreatedAttributeKeys.Find(InOutputPinLabel))
 	{
-		OutDataDesc->AddAttributeToAllData(CreatedKey, InBinding);
+		for (const FPCGKernelAttributeKey& CreatedKey : CreatedKeyList->Keys)
+		{
+			OutDataDesc->AddAttributeToAllData(CreatedKey, InBinding);
+		}
 	}
 
 	// Try to propagate string keys across node. Not trivial because there could be one or more string key attributes on input pins and on output pins,
@@ -805,25 +813,64 @@ void UPCGCustomHLSLKernel::InitEntryPoint()
 	}
 }
 
-void UPCGCustomHLSLKernel::PopulateAttributeKeysFromPinSettings()
+void UPCGCustomHLSLKernel::PopulateAttributeKeys()
 {
 	const UPCGCustomHLSLSettings* CustomHLSLSettings = CastChecked<UPCGCustomHLSLSettings>(GetSettings());
 
-	// Process each output pin for any new attributes they want to create.
-	for (const FPCGPinPropertiesGPU& OutputPin : CustomHLSLSettings->OutputPins)
+	PinCreatedAttributeKeys.Reset();
+
+	// Pass 1: seed from manually-authored CreatedKernelAttributeKeys on each output pin, but only when the user
+	// has explicitly enabled per-pin attribute creation settings. When disabled, only HLSL inference is used.
+	if (CustomHLSLSettings->bPerPinAttributeCreationSettings)
 	{
-		for (const FPCGKernelAttributeKey& AuthoredKey : OutputPin.PropertiesGPU.CreatedKernelAttributeKeys)
+		for (const FPCGPinPropertiesGPU& OutputPin : CustomHLSLSettings->OutputPins)
 		{
-			if (AuthoredKey.IsValid())
+			for (const FPCGKernelAttributeKey& AuthoredKey : OutputPin.PropertiesGPU.CreatedKernelAttributeKeys)
 			{
-				KernelAttributeKeys.AddUnique(AuthoredKey);
+				if (AuthoredKey.IsValid())
+				{
+					KernelAttributeKeys.AddUnique(AuthoredKey);
+					PinCreatedAttributeKeys.FindOrAdd(OutputPin.Label).Keys.AddUnique(AuthoredKey);
+				}
 			}
 		}
 	}
+
+	// Pass 2: augment with attributes inferred from HLSL Set/AtomicAdd calls in ParsedSources.
+	// ParseShaderSource() already added these to KernelAttributeKeys, but we also need them in PinCreatedAttributeKeys.
+	for (const FPCGCustomHLSLParsedSource& ParsedSource : ParsedSources)
+	{
+		for (const FPCGParsedAttributeFunction& AttrFunc : ParsedSource.AttributeFunctions)
+		{
+			const bool bIsSet = (AttrFunc.FunctionName == PCGCustomHLSLKernel::AttributeFunctionSetKeyword);
+			const bool bIsAtomicAdd = (AttrFunc.FunctionName == PCGCustomHLSLKernel::AttributeFunctionAtomicAddKeyword);
+			if (!bIsSet && !bIsAtomicAdd)
+			{
+				continue;
+			}
+
+			const FPCGKernelAttributeKey InferredKey(FPCGAttributePropertySelector::CreateSelectorFromString(AttrFunc.AttributeName), static_cast<EPCGKernelAttributeType>(AttrFunc.AttributeType));
+			if (!InferredKey.IsValid())
+			{
+				continue;
+			}
+
+			const FName PinLabel(*AttrFunc.PinLabel);
+			TArray<FPCGKernelAttributeKey>& PinKeys = PinCreatedAttributeKeys.FindOrAdd(PinLabel).Keys;
+			if (PinKeys.Contains(InferredKey))
+			{
+				continue;
+			}
+
+			PinKeys.Add(InferredKey);
+		}
+	}
 }
 
 void UPCGCustomHLSLKernel::ParseShaderSource()
 {
+	KernelAttributeKeys.Reset();
+
 	CreateParsedSources();
 
 	static const TArray<FString> AttributeTypeStrings =
diff --git a/Engine/Plugins/PCG/Source/PCG/Private/Compute/Elements/PCGCustomHLSLKernel.h b/Engine/Plugins/PCG/Source/PCG/Private/Compute/Elements/PCGCustomHLSLKernel.h
--- a/Engine/Plugins/PCG/Source/PCG/Private/Compute/Elements/PCGCustomHLSLKernel.h
+++ b/Engine/Plugins/PCG/Source/PCG/Private/Compute/Elements/PCGCustomHLSLKernel.h
@@ -9,6 +9,17 @@
 
 #include "PCGCustomHLSLKernel.generated.h"
 
+/** Helper struct for serializing per-pin created attribute key lists. */
+USTRUCT()
+struct FPCGKernelAttributeKeyList
+{
+    GENERATED_BODY()
+
+public:
+    UPROPERTY()
+    TArray<FPCGKernelAttributeKey> Keys;
+};
+
 USTRUCT()
 struct FPCGParsedAttributeFunction
 {
@@ -139,7 +150,7 @@ protected:
 	// ~End UPCGComputeKernel interface
 
 	void InitEntryPoint();
-	void PopulateAttributeKeysFromPinSettings();
+	void PopulateAttributeKeys();
 	void ParseShaderSource();
 	void CreateParsedSources();
 	void CollectDataLabels(const FPCGCustomHLSLParsedSource& InParsedSource);
@@ -177,6 +188,11 @@ protected:
 	UPROPERTY()
 	FPCGPinDataLabels PinDataLabels;
 
+	/** Per-output-pin created attribute keys, combining manually authored CreatedKernelAttributeKeys and
+	 *  keys inferred from HLSL Set/AtomicAdd calls. Populated during editor init, serialized for runtime use. */
+	UPROPERTY()
+	TMap<FName, FPCGKernelAttributeKeyList> PinCreatedAttributeKeys;
+
 	/** Properties that need to be fully allocated per pin. Note, we use an integer instead of EPCGPointNativeProperties directly because bitflag enums don't serialize correctly. */
 	UPROPERTY()
 	TMap<FName, /*EPCGPointNativeProperties*/int32> PinAllocatedProperties;
diff --git a/Engine/Plugins/PCG/Source/PCG/Private/Tests/Compute/PCGCustomHLSLAttributeTest.cpp b/Engine/Plugins/PCG/Source/PCG/Private/Tests/Compute/PCGCustomHLSLAttributeTest.cpp
new file mode 100644
--- /dev/null
+++ b/Engine/Plugins/PCG/Source/PCG/Private/Tests/Compute/PCGCustomHLSLAttributeTest.cpp
@@ -0,0 +1,555 @@
+// Copyright Epic Games, Inc. All Rights Reserved.
+
+#if WITH_EDITOR
+
+#include "Tests/Compute/PCGGPUTestCommon.h"
+
+#include "PCGCommon.h"
+#include "PCGGraph.h"
+#include "PCGParamData.h"
+#include "Compute/Elements/PCGCustomHLSL.h"
+#include "Compute/PCGPinPropertiesGPU.h"
+#include "Data/PCGBasePointData.h"
+#include "Elements/PCGCreateAttribute.h"
+#include "Elements/PCGCreatePoints.h"
+#include "Elements/PCGGather.h"
+#include "Metadata/PCGMetadataAttributeTpl.h"
+
+/**
+ * GPU test: CustomHLSL AttributeSetProcessor doubles a float attribute.
+ *
+ *   CreateAttributeSet (Value=5) -> CustomHLSL (Value *= 2) -> Gather
+ *
+ * Verifies the output attribute set has Value == 10.
+ */
+IMPLEMENT_CUSTOM_SIMPLE_AUTOMATION_TEST(
+    FPCGCustomHLSLAttributeSetDoubleValueTest,
+    FPCGTestBaseClass,
+    "Plugins.PCG.Compute.CustomHLSL.AttributeSetDoubleValue",
+    PCGTestsCommon::TestFlags | EAutomationTestFlags::NonNullRHI)
+
+bool FPCGCustomHLSLAttributeSetDoubleValueTest::RunTest(const FString& Parameters)
+{
+    FPCGGPUGraphTestRunner Runner;
+
+    // CreateAttributeSet: one element with float attribute "Value" = 5.
+    UPCGCreateAttributeSetSettings* CreateAttrSettings = nullptr;
+    UPCGNode* CreateAttrNode = Runner.Graph->AddNodeOfType<UPCGCreateAttributeSetSettings>(CreateAttrSettings);
+    CreateAttrSettings->AttributeTypes.Type = EPCGMetadataTypes::Float;
+    CreateAttrSettings->AttributeTypes.FloatValue = 5.0f;
+    CreateAttrSettings->OutputTarget.SetAttributeName(FName(TEXT("Value")));
+
+    // CustomHLSL AttributeSetProcessor: Value *= 2.
+    UPCGCustomHLSLSettings* HLSLSettings = nullptr;
+    UPCGNode* HLSLNode = Runner.Graph->AddNodeOfType<UPCGCustomHLSLSettings>(HLSLSettings);
+    HLSLSettings->SetKernelType(EPCGKernelType::AttributeSetProcessor);
+    HLSLSettings->SetSourceText(
+        TEXT("float v = In_GetFloat(In_DataIndex, ElementIndex, 'Value');\n")
+        TEXT("Out_SetFloat(Out_DataIndex, ElementIndex, 'Value', v * 2.0f);\n")
+    );
+
+    // Useful to add gather node to grab output data from.
+    UPCGGatherSettings* GatherSettings = nullptr;
+    UPCGNode* GatherNode = Runner.Graph->AddNodeOfType<UPCGGatherSettings>(GatherSettings);
+
+    Runner.Graph->AddEdge(CreateAttrNode, PCGPinConstants::DefaultOutputLabel, HLSLNode, PCGPinConstants::DefaultInputLabel);
+    Runner.Graph->AddEdge(HLSLNode, PCGPinConstants::DefaultOutputLabel, GatherNode, PCGPinConstants::DefaultInputLabel);
+
+    TArray<FPCGDataCollection> OutNodeData;
+    if (!Runner.Execute({ this, { GatherNode } }, OutNodeData))
+    {
+        return false;
+    }
+
+    return PCGGPUTestCommon::ValidateSingleTaggedData(this, OutNodeData[0], [&](const FPCGTaggedData& TaggedData)
+    {
+        const UPCGParamData* OutParams = Cast<UPCGParamData>(TaggedData.Data);
+        if (!TestNotNull(TEXT("Readback data is attribute set data"), OutParams))
+        {
+            return false;
+        }
+
+        const UPCGMetadata* Metadata = OutParams->ConstMetadata();
+        if (!TestNotNull(TEXT("Output has metadata"), Metadata))
+        {
+            return false;
+        }
+
+        const FPCGMetadataAttribute<float>* FloatAttr = Metadata->GetConstTypedAttribute<float>(FName(TEXT("Value")));
+        if (!TestNotNull(TEXT("Output has 'Value' float attribute"), FloatAttr))
+        {
+            return false;
+        }
+
+        const double ActualValue = FloatAttr->GetValueFromItemKey(0);
+        return TestNearlyEqual(*FString::Printf(TEXT("Value attribute doubled (expected 10.0, got %.1f)"), ActualValue), ActualValue, 10.0, 0.1);
+    });
+}
+
+/**
+ * GPU test: CustomHLSL AttributeSetProcessor passes a string attribute through unchanged.
+ *
+ *   CreateAttributeSet (StringAttr="Hello") -> CustomHLSL (read StringAttr, write StringAttr) -> Gather
+ *
+ * Verifies the output attribute set has StringAttr == "Hello".
+ */
+IMPLEMENT_CUSTOM_SIMPLE_AUTOMATION_TEST(
+    FPCGCustomHLSLAttributeSetStringAttributeTest,
+    FPCGTestBaseClass,
+    "Plugins.PCG.Compute.CustomHLSL.AttributeSetStringAttribute",
+    PCGTestsCommon::TestFlags | EAutomationTestFlags::NonNullRHI)
+
+bool FPCGCustomHLSLAttributeSetStringAttributeTest::RunTest(const FString& Parameters)
+{
+    FPCGGPUGraphTestRunner Runner;
+
+    // CreateAttributeSet: one element with string attribute "StringAttr" = "Hello".
+    UPCGCreateAttributeSetSettings* CreateAttrSettings = nullptr;
+    UPCGNode* CreateAttrNode = Runner.Graph->AddNodeOfType<UPCGCreateAttributeSetSettings>(CreateAttrSettings);
+    CreateAttrSettings->AttributeTypes.Type = EPCGMetadataTypes::String;
+    CreateAttrSettings->AttributeTypes.StringValue = TEXT("Hello");
+    CreateAttrSettings->OutputTarget.SetAttributeName(FName(TEXT("StringAttr")));
+
+    // CustomHLSL AttributeSetProcessor: read StringAttr and write it back unchanged.
+    UPCGCustomHLSLSettings* HLSLSettings = nullptr;
+    UPCGNode* HLSLNode = Runner.Graph->AddNodeOfType<UPCGCustomHLSLSettings>(HLSLSettings);
+    HLSLSettings->SetKernelType(EPCGKernelType::AttributeSetProcessor);
+    HLSLSettings->SetSourceText(
+        TEXT("int StringAttrVal = In_GetStringKey(In_DataIndex, ElementIndex, 'StringAttr');\n")
+        TEXT("Out_SetStringKey(Out_DataIndex, ElementIndex, 'StringAttr', StringAttrVal);\n")
+    );
+
+    // Useful to add gather node to grab output data from.
+    UPCGGatherSettings* GatherSettings = nullptr;
+    UPCGNode* GatherNode = Runner.Graph->AddNodeOfType<UPCGGatherSettings>(GatherSettings);
+
+    Runner.Graph->AddEdge(CreateAttrNode, PCGPinConstants::DefaultOutputLabel, HLSLNode, PCGPinConstants::DefaultInputLabel);
+    Runner.Graph->AddEdge(HLSLNode, PCGPinConstants::DefaultOutputLabel, GatherNode, PCGPinConstants::DefaultInputLabel);
+
+    TArray<FPCGDataCollection> OutNodeData;
+    if (!Runner.Execute({ this, { GatherNode } }, OutNodeData))
+    {
+        return false;
+    }
+
+    return PCGGPUTestCommon::ValidateSingleTaggedData(this, OutNodeData[0], [&](const FPCGTaggedData& TaggedData)
+    {
+        const UPCGParamData* OutParams = Cast<UPCGParamData>(TaggedData.Data);
+        if (!TestNotNull(TEXT("Readback data is attribute set data"), OutParams))
+        {
+            return false;
+        }
+
+        const UPCGMetadata* Metadata = OutParams->ConstMetadata();
+        if (!TestNotNull(TEXT("Output has metadata"), Metadata))
+        {
+            return false;
+        }
+
+        const FPCGMetadataAttribute<FString>* StringAttr = Metadata->GetConstTypedAttribute<FString>(FName(TEXT("StringAttr")));
+        if (!TestNotNull(TEXT("Output has 'StringAttr' string attribute"), StringAttr))
+        {
+            return false;
+        }
+
+        const FString ActualValue = StringAttr->GetValueFromItemKey(0);
+        return TestEqual(
+            *FString::Printf(TEXT("'StringAttr' attribute passed through (expected 'Hello', got '%s')"), *ActualValue),
+            ActualValue, FString(TEXT("Hello")));
+    });
+}
+
+/**
+ * GPU test: CustomHLSL infers created attributes from HLSL Set calls.
+ *
+ *   CreatePoints (2 points) -> CustomHLSL PointProcessor (writes one attribute of each type) -> Gather
+ *
+ * Verifies that attributes written via Out_Set* calls in HLSL are automatically inferred and present
+ * on the output point data with the correct C++ metadata type. No CreatedKernelAttributeKeys are
+ * configured — HLSL inference is the sole source.
+ */
+IMPLEMENT_CUSTOM_SIMPLE_AUTOMATION_TEST(
+    FPCGCustomHLSLInferredAttributeTypesTest,
+    FPCGTestBaseClass,
+    "Plugins.PCG.Compute.CustomHLSL.InferredAttributeTypes",
+    PCGTestsCommon::TestFlags | EAutomationTestFlags::NonNullRHI)
+
+bool FPCGCustomHLSLInferredAttributeTypesTest::RunTest(const FString& Parameters)
+{
+    FPCGGPUGraphTestRunner Runner;
+
+    // CreatePoints: 2 default points.
+    UPCGCreatePointsSettings* CreatePointsSettings = nullptr;
+    UPCGNode* CreatePointsNode = Runner.Graph->AddNodeOfType<UPCGCreatePointsSettings>(CreatePointsSettings);
+    CreatePointsSettings->PointsToCreate.Empty();
+    CreatePointsSettings->PointsToCreate.AddDefaulted();
+    CreatePointsSettings->PointsToCreate.AddDefaulted();
+    CreatePointsSettings->CoordinateSpace = EPCGCoordinateSpace::World;
+
+    // CustomHLSL PointProcessor: write one attribute of each supported non-string type.
+    // No CreatedKernelAttributeKeys — all attributes are inferred from the HLSL Set calls.
+    UPCGCustomHLSLSettings* HLSLSettings = nullptr;
+    UPCGNode* HLSLNode = Runner.Graph->AddNodeOfType<UPCGCustomHLSLSettings>(HLSLSettings);
+    HLSLSettings->KernelType = EPCGKernelType::PointProcessor;
+    HLSLSettings->SetSourceText(
+        TEXT("Out_SetBool(Out_DataIndex, ElementIndex, 'AttrBool', true);\n")
+        TEXT("Out_SetInt(Out_DataIndex, ElementIndex, 'AttrInt', 7);\n")
+        TEXT("Out_SetFloat(Out_DataIndex, ElementIndex, 'AttrFloat', 1.5f);\n")
+        TEXT("Out_SetFloat2(Out_DataIndex, ElementIndex, 'AttrFloat2', float2(1.0f, 2.0f));\n")
+        TEXT("Out_SetFloat3(Out_DataIndex, ElementIndex, 'AttrFloat3', float3(1.0f, 2.0f, 3.0f));\n")
+        TEXT("Out_SetFloat4(Out_DataIndex, ElementIndex, 'AttrFloat4', float4(1.0f, 2.0f, 3.0f, 4.0f));\n")
+        TEXT("Out_SetRotator(Out_DataIndex, ElementIndex, 'AttrRotator', float3(10.0f, 20.0f, 30.0f));\n")
+        TEXT("Out_SetQuat(Out_DataIndex, ElementIndex, 'AttrQuat', float4(0.0f, 0.0f, 0.0f, 1.0f));\n")
+        TEXT("Out_SetTransform(Out_DataIndex, ElementIndex, 'AttrTransform', float4x4(1,0,0,0, 0,1,0,0, 0,0,1,0, 0,0,0,1));\n")
+    );
+
+    UPCGGatherSettings* GatherSettings = nullptr;
+    UPCGNode* GatherNode = Runner.Graph->AddNodeOfType<UPCGGatherSettings>(GatherSettings);
+
+    Runner.Graph->AddEdge(CreatePointsNode, PCGPinConstants::DefaultOutputLabel, HLSLNode, PCGPinConstants::DefaultInputLabel);
+    Runner.Graph->AddEdge(HLSLNode, PCGPinConstants::DefaultOutputLabel, GatherNode, PCGPinConstants::DefaultInputLabel);
+
+    TArray<FPCGDataCollection> OutNodeData;
+    if (!Runner.Execute({ this, { GatherNode } }, OutNodeData))
+    {
+        return false;
+    }
+
+    return PCGGPUTestCommon::ValidateSingleTaggedData(this, OutNodeData[0], [&](const FPCGTaggedData& TaggedData)
+    {
+        const UPCGBasePointData* OutPoints = Cast<UPCGBasePointData>(TaggedData.Data);
+        if (!TestNotNull(TEXT("Readback data is point data"), OutPoints))
+        {
+            return false;
+        }
+
+        const UPCGMetadata* Metadata = OutPoints->ConstMetadata();
+        if (!TestNotNull(TEXT("Output has metadata"), Metadata))
+        {
+            return false;
+        }
+
+        bool bResult = true;
+        bResult &= TestNotNull(TEXT("Bool attribute inferred from HLSL"),      Metadata->GetConstTypedAttribute<bool>(FName("AttrBool")));
+        bResult &= TestNotNull(TEXT("Int attribute inferred from HLSL"),        Metadata->GetConstTypedAttribute<int32>(FName("AttrInt")));
+        bResult &= TestNotNull(TEXT("Float attribute inferred from HLSL"),      Metadata->GetConstTypedAttribute<float>(FName("AttrFloat")));
+        bResult &= TestNotNull(TEXT("Float2 attribute inferred from HLSL"),     Metadata->GetConstTypedAttribute<FVector2D>(FName("AttrFloat2")));
+        bResult &= TestNotNull(TEXT("Float3 attribute inferred from HLSL"),     Metadata->GetConstTypedAttribute<FVector>(FName("AttrFloat3")));
+        bResult &= TestNotNull(TEXT("Float4 attribute inferred from HLSL"),     Metadata->GetConstTypedAttribute<FVector4>(FName("AttrFloat4")));
+        bResult &= TestNotNull(TEXT("Rotator attribute inferred from HLSL"),    Metadata->GetConstTypedAttribute<FRotator>(FName("AttrRotator")));
+        bResult &= TestNotNull(TEXT("Quat attribute inferred from HLSL"),       Metadata->GetConstTypedAttribute<FQuat>(FName("AttrQuat")));
+        bResult &= TestNotNull(TEXT("Transform attribute inferred from HLSL"),  Metadata->GetConstTypedAttribute<FTransform>(FName("AttrTransform")));
+        return bResult;
+    });
+}
+
+/**
+ * GPU test: CustomHLSL per-pin attribute creation settings (enabled).
+ *
+ *   CustomHLSL PointGenerator (empty HLSL, bPerPinAttributeCreationSettings=true, all types in CreatedKernelAttributeKeys) -> Gather
+ *
+ * Verifies that when bPerPinAttributeCreationSettings is enabled, attributes configured in
+ * CreatedKernelAttributeKeys are present on the output even when HLSL writes no attributes.
+ *
+ * Note: EPCGKernelAttributeType::Name/StringKey is excluded because it requires input pins
+ * to carry existing string-key attributes for the GPU string table. A PointGenerator has no
+ * inputs, so including it would produce a "No incoming attributes to obtain string keys from"
+ * warning and the attribute would not be usable.
+ */
+IMPLEMENT_CUSTOM_SIMPLE_AUTOMATION_TEST(
+    FPCGCustomHLSLPerPinAttributeSettingsEnabledTest,
+    FPCGTestBaseClass,
+    "Plugins.PCG.Compute.CustomHLSL.PerPinAttributeSettingsEnabled",
+    PCGTestsCommon::TestFlags | EAutomationTestFlags::NonNullRHI)
+
+bool FPCGCustomHLSLPerPinAttributeSettingsEnabledTest::RunTest(const FString& Parameters)
+{
+    FPCGGPUGraphTestRunner Runner;
+
+    UPCGCustomHLSLSettings* HLSLSettings = nullptr;
+    UPCGNode* HLSLNode = Runner.Graph->AddNodeOfType<UPCGCustomHLSLSettings>(HLSLSettings);
+    HLSLSettings->SetKernelType(EPCGKernelType::PointGenerator);
+    HLSLSettings->SetSourceText(TEXT(""));
+    HLSLSettings->NumElements = 2;
+    HLSLSettings->SetPerPinAttributeCreationSettings(true);
+
+    // Add one attribute of each type to the default output pin.
+    check(HLSLSettings->OutputPins.Num() == 1);
+    TArray<FPCGKernelAttributeKey>& Keys = HLSLSettings->OutputPins[0].PropertiesGPU.CreatedKernelAttributeKeys;
+    Keys.Add(FPCGKernelAttributeKey(FName("AttrBool"),      EPCGKernelAttributeType::Bool));
+    Keys.Add(FPCGKernelAttributeKey(FName("AttrInt"),       EPCGKernelAttributeType::Int));
+    Keys.Add(FPCGKernelAttributeKey(FName("AttrFloat"),     EPCGKernelAttributeType::Float));
+    Keys.Add(FPCGKernelAttributeKey(FName("AttrFloat2"),    EPCGKernelAttributeType::Float2));
+    Keys.Add(FPCGKernelAttributeKey(FName("AttrFloat3"),    EPCGKernelAttributeType::Float3));
+    Keys.Add(FPCGKernelAttributeKey(FName("AttrFloat4"),    EPCGKernelAttributeType::Float4));
+    Keys.Add(FPCGKernelAttributeKey(FName("AttrRotator"),   EPCGKernelAttributeType::Rotator));
+    Keys.Add(FPCGKernelAttributeKey(FName("AttrQuat"),      EPCGKernelAttributeType::Quat));
+    Keys.Add(FPCGKernelAttributeKey(FName("AttrTransform"), EPCGKernelAttributeType::Transform));
+
+    UPCGGatherSettings* GatherSettings = nullptr;
+    UPCGNode* GatherNode = Runner.Graph->AddNodeOfType<UPCGGatherSettings>(GatherSettings);
+
+    Runner.Graph->AddEdge(HLSLNode, PCGPinConstants::DefaultOutputLabel, GatherNode, PCGPinConstants::DefaultInputLabel);
+
+    TArray<FPCGDataCollection> OutNodeData;
+    if (!Runner.Execute({ this, { GatherNode } }, OutNodeData))
+    {
+        return false;
+    }
+
+    return PCGGPUTestCommon::ValidateSingleTaggedData(this, OutNodeData[0], [&](const FPCGTaggedData& TaggedData)
+    {
+        const UPCGBasePointData* OutPoints = Cast<UPCGBasePointData>(TaggedData.Data);
+        if (!TestNotNull(TEXT("Readback data is point data"), OutPoints))
+        {
+            return false;
+        }
+
+        const UPCGMetadata* Metadata = OutPoints->ConstMetadata();
+        if (!TestNotNull(TEXT("Output has metadata"), Metadata))
+        {
+            return false;
+        }
+
+        bool bResult = true;
+        bResult &= TestNotNull(TEXT("Bool attribute created from per-pin settings"),      Metadata->GetConstTypedAttribute<bool>(FName("AttrBool")));
+        bResult &= TestNotNull(TEXT("Int attribute created from per-pin settings"),        Metadata->GetConstTypedAttribute<int32>(FName("AttrInt")));
+        bResult &= TestNotNull(TEXT("Float attribute created from per-pin settings"),      Metadata->GetConstTypedAttribute<float>(FName("AttrFloat")));
+        bResult &= TestNotNull(TEXT("Float2 attribute created from per-pin settings"),     Metadata->GetConstTypedAttribute<FVector2D>(FName("AttrFloat2")));
+        bResult &= TestNotNull(TEXT("Float3 attribute created from per-pin settings"),     Metadata->GetConstTypedAttribute<FVector>(FName("AttrFloat3")));
+        bResult &= TestNotNull(TEXT("Float4 attribute created from per-pin settings"),     Metadata->GetConstTypedAttribute<FVector4>(FName("AttrFloat4")));
+        bResult &= TestNotNull(TEXT("Rotator attribute created from per-pin settings"),    Metadata->GetConstTypedAttribute<FRotator>(FName("AttrRotator")));
+        bResult &= TestNotNull(TEXT("Quat attribute created from per-pin settings"),       Metadata->GetConstTypedAttribute<FQuat>(FName("AttrQuat")));
+        bResult &= TestNotNull(TEXT("Transform attribute created from per-pin settings"),  Metadata->GetConstTypedAttribute<FTransform>(FName("AttrTransform")));
+        return bResult;
+    });
+}
+
+/**
+ * GPU test: CustomHLSL per-pin attribute creation settings (disabled).
+ *
+ *   CustomHLSL PointGenerator (empty HLSL, bPerPinAttributeCreationSettings=false, all types in CreatedKernelAttributeKeys) -> Gather
+ *
+ * Verifies that when bPerPinAttributeCreationSettings is disabled (the default), attributes
+ * configured in CreatedKernelAttributeKeys are ignored — only HLSL inference applies (none here).
+ */
+IMPLEMENT_CUSTOM_SIMPLE_AUTOMATION_TEST(
+    FPCGCustomHLSLPerPinAttributeSettingsDisabledTest,
+    FPCGTestBaseClass,
+    "Plugins.PCG.Compute.CustomHLSL.PerPinAttributeSettingsDisabled",
+    PCGTestsCommon::TestFlags | EAutomationTestFlags::NonNullRHI)
+
+bool FPCGCustomHLSLPerPinAttributeSettingsDisabledTest::RunTest(const FString& Parameters)
+{
+    FPCGGPUGraphTestRunner Runner;
+
+    UPCGCustomHLSLSettings* HLSLSettings = nullptr;
+    UPCGNode* HLSLNode = Runner.Graph->AddNodeOfType<UPCGCustomHLSLSettings>(HLSLSettings);
+    HLSLSettings->SetKernelType(EPCGKernelType::PointGenerator);
+    HLSLSettings->SetSourceText(TEXT(""));
+    HLSLSettings->NumElements = 2;
+    HLSLSettings->SetPerPinAttributeCreationSettings(false);  // default: authored keys are ignored
+
+    // Same keys as the enabled test — none of these should appear in the output.
+    check(HLSLSettings->OutputPins.Num() == 1);
+    TArray<FPCGKernelAttributeKey>& Keys = HLSLSettings->OutputPins[0].PropertiesGPU.CreatedKernelAttributeKeys;
+    Keys.Add(FPCGKernelAttributeKey(FName("AttrBool"),      EPCGKernelAttributeType::Bool));
+    Keys.Add(FPCGKernelAttributeKey(FName("AttrInt"),       EPCGKernelAttributeType::Int));
+    Keys.Add(FPCGKernelAttributeKey(FName("AttrFloat"),     EPCGKernelAttributeType::Float));
+    Keys.Add(FPCGKernelAttributeKey(FName("AttrFloat2"),    EPCGKernelAttributeType::Float2));
+    Keys.Add(FPCGKernelAttributeKey(FName("AttrFloat3"),    EPCGKernelAttributeType::Float3));
+    Keys.Add(FPCGKernelAttributeKey(FName("AttrFloat4"),    EPCGKernelAttributeType::Float4));
+    Keys.Add(FPCGKernelAttributeKey(FName("AttrRotator"),   EPCGKernelAttributeType::Rotator));
+    Keys.Add(FPCGKernelAttributeKey(FName("AttrQuat"),      EPCGKernelAttributeType::Quat));
+    Keys.Add(FPCGKernelAttributeKey(FName("AttrTransform"), EPCGKernelAttributeType::Transform));
+
+    UPCGGatherSettings* GatherSettings = nullptr;
+    UPCGNode* GatherNode = Runner.Graph->AddNodeOfType<UPCGGatherSettings>(GatherSettings);
+
+    Runner.Graph->AddEdge(HLSLNode, PCGPinConstants::DefaultOutputLabel, GatherNode, PCGPinConstants::DefaultInputLabel);
+
+    TArray<FPCGDataCollection> OutNodeData;
+    if (!Runner.Execute({ this, { GatherNode } }, OutNodeData))
+    {
+        return false;
+    }
+
+    return PCGGPUTestCommon::ValidateSingleTaggedData(this, OutNodeData[0], [&](const FPCGTaggedData& TaggedData)
+    {
+        const UPCGBasePointData* OutPoints = Cast<UPCGBasePointData>(TaggedData.Data);
+        if (!TestNotNull(TEXT("Readback data is point data"), OutPoints))
+        {
+            return false;
+        }
+
+        const UPCGMetadata* Metadata = OutPoints->ConstMetadata();
+        if (!TestNotNull(TEXT("Output has metadata"), Metadata))
+        {
+            return false;
+        }
+
+        bool bResult = true;
+        bResult &= TestNull(TEXT("Bool attribute absent when per-pin settings disabled"),      Metadata->GetConstTypedAttribute<bool>(FName("AttrBool")));
+        bResult &= TestNull(TEXT("Int attribute absent when per-pin settings disabled"),        Metadata->GetConstTypedAttribute<int32>(FName("AttrInt")));
+        bResult &= TestNull(TEXT("Float attribute absent when per-pin settings disabled"),      Metadata->GetConstTypedAttribute<float>(FName("AttrFloat")));
+        bResult &= TestNull(TEXT("Float2 attribute absent when per-pin settings disabled"),     Metadata->GetConstTypedAttribute<FVector2D>(FName("AttrFloat2")));
+        bResult &= TestNull(TEXT("Float3 attribute absent when per-pin settings disabled"),     Metadata->GetConstTypedAttribute<FVector>(FName("AttrFloat3")));
+        bResult &= TestNull(TEXT("Float4 attribute absent when per-pin settings disabled"),     Metadata->GetConstTypedAttribute<FVector4>(FName("AttrFloat4")));
+        bResult &= TestNull(TEXT("Rotator attribute absent when per-pin settings disabled"),    Metadata->GetConstTypedAttribute<FRotator>(FName("AttrRotator")));
+        bResult &= TestNull(TEXT("Quat attribute absent when per-pin settings disabled"),       Metadata->GetConstTypedAttribute<FQuat>(FName("AttrQuat")));
+        bResult &= TestNull(TEXT("Transform attribute absent when per-pin settings disabled"),  Metadata->GetConstTypedAttribute<FTransform>(FName("AttrTransform")));
+        return bResult;
+    });
+}
+
+/**
+ * GPU test: attribute created from both CreatedKernelAttributeKeys and HLSL Set call with no conflict.
+ *
+ *   CustomHLSL PointGenerator (bPerPinAttributeCreationSettings=true, AttrFloat authored + written in HLSL) -> Gather
+ *
+ * Verifies that when the same attribute (name + type) is declared in both CreatedKernelAttributeKeys and
+ * written via HLSL, the attribute is present exactly once with the correct type. The deduplication in
+ * PopulateAttributeKeys (AddUnique) should prevent duplicate entries.
+ */
+IMPLEMENT_CUSTOM_SIMPLE_AUTOMATION_TEST(
+    FPCGCustomHLSLAttributeCreationOverlapTest,
+    FPCGTestBaseClass,
+    "Plugins.PCG.Compute.CustomHLSL.AttributeCreationOverlap",
+    PCGTestsCommon::TestFlags | EAutomationTestFlags::NonNullRHI)
+
+bool FPCGCustomHLSLAttributeCreationOverlapTest::RunTest(const FString& Parameters)
+{
+    FPCGGPUGraphTestRunner Runner;
+
+    UPCGCustomHLSLSettings* HLSLSettings = nullptr;
+    UPCGNode* HLSLNode = Runner.Graph->AddNodeOfType<UPCGCustomHLSLSettings>(HLSLSettings);
+    HLSLSettings->SetKernelType(EPCGKernelType::PointGenerator);
+    HLSLSettings->NumElements = 2;
+    HLSLSettings->SetPerPinAttributeCreationSettings(true);
+
+    // SetSourceText MUST be called before adding keys to CreatedKernelAttributeKeys.
+    // PostEditChangeProperty triggered by SetSourceText calls UpdateAttributeKeys() which
+    // runs UpdateIdentifierFromSelector() on every key — keys constructed via FPCGKernelAttributeKey(FName, Type)
+    // have an empty Name selector, so their MetadataDomain gets set to Invalid, making them fail IsValid().
+    HLSLSettings->SetSourceText(
+        TEXT("Out_SetFloat(Out_DataIndex, ElementIndex, 'AttrFloat', 3.0f);\n")
+        TEXT("Out_SetBool(Out_DataIndex, ElementIndex, 'AttrBool', true);\n")
+    );
+
+    // Author AttrFloat and AttrInt in CreatedKernelAttributeKeys.
+    check(HLSLSettings->OutputPins.Num() == 1);
+    TArray<FPCGKernelAttributeKey>& Keys = HLSLSettings->OutputPins[0].PropertiesGPU.CreatedKernelAttributeKeys;
+    Keys.Add(FPCGKernelAttributeKey(FName("AttrFloat"), EPCGKernelAttributeType::Float));
+    Keys.Add(FPCGKernelAttributeKey(FName("AttrInt"),   EPCGKernelAttributeType::Int));
+
+    UPCGGatherSettings* GatherSettings = nullptr;
+    UPCGNode* GatherNode = Runner.Graph->AddNodeOfType<UPCGGatherSettings>(GatherSettings);
+    Runner.Graph->AddEdge(HLSLNode, PCGPinConstants::DefaultOutputLabel, GatherNode, PCGPinConstants::DefaultInputLabel);
+
+    TArray<FPCGDataCollection> OutNodeData;
+    if (!Runner.Execute({ this, { GatherNode } }, OutNodeData))
+    {
+        return false;
+    }
+
+    return PCGGPUTestCommon::ValidateSingleTaggedData(this, OutNodeData[0], [&](const FPCGTaggedData& TaggedData)
+    {
+        const UPCGBasePointData* OutPoints = Cast<UPCGBasePointData>(TaggedData.Data);
+        if (!TestNotNull(TEXT("Readback data is point data"), OutPoints))
+        {
+            return false;
+        }
+
+        const UPCGMetadata* Metadata = OutPoints->ConstMetadata();
+        if (!TestNotNull(TEXT("Output has metadata"), Metadata))
+        {
+            return false;
+        }
+
+        bool bResult = true;
+        // Authored only — should exist.
+        bResult &= TestNotNull(TEXT("AttrInt authored only — present"),              Metadata->GetConstTypedAttribute<int32>(FName("AttrInt")));
+        // Authored AND written in HLSL (same type) — should exist exactly once, no conflict.
+        bResult &= TestNotNull(TEXT("AttrFloat authored + HLSL written — present"),  Metadata->GetConstTypedAttribute<float>(FName("AttrFloat")));
+        // Inferred from HLSL only — should exist.
+        bResult &= TestNotNull(TEXT("AttrBool inferred from HLSL — present"),        Metadata->GetConstTypedAttribute<bool>(FName("AttrBool")));
+        return bResult;
+    });
+}
+
+/**
+ * GPU test: authored and inferred attributes with distinct names coexist without conflict.
+ *
+ *   CustomHLSL PointGenerator (bPerPinAttributeCreationSettings=true) -> Gather
+ *
+ * Authored in CreatedKernelAttributeKeys: "AttrA" as Float, "AttrB" as Int.
+ * Written in HLSL:                        "AttrC" as Bool,  "AttrD" as Float3.
+ *
+ * Verifies that all four attributes are present on the output — attributes from
+ * CreatedKernelAttributeKeys (pass 1) and HLSL inference (pass 2) are merged together.
+ *
+ * Note: same-name/different-type attributes are NOT supported by the runtime. Each attribute
+ * name must map to exactly one type. This test uses four distinct names to avoid that constraint.
+ */
+IMPLEMENT_CUSTOM_SIMPLE_AUTOMATION_TEST(
+    FPCGCustomHLSLAuthoredAndInferredCoexistTest,
+    FPCGTestBaseClass,
+    "Plugins.PCG.Compute.CustomHLSL.AuthoredAndInferredCoexist",
+    PCGTestsCommon::TestFlags | EAutomationTestFlags::NonNullRHI)
+
+bool FPCGCustomHLSLAuthoredAndInferredCoexistTest::RunTest(const FString& Parameters)
+{
+    FPCGGPUGraphTestRunner Runner;
+
+    UPCGCustomHLSLSettings* HLSLSettings = nullptr;
+    UPCGNode* HLSLNode = Runner.Graph->AddNodeOfType<UPCGCustomHLSLSettings>(HLSLSettings);
+    HLSLSettings->SetKernelType(EPCGKernelType::PointGenerator);
+    HLSLSettings->NumElements = 2;
+    HLSLSettings->SetPerPinAttributeCreationSettings(true);
+
+    // SetSourceText MUST be called before adding keys (see AttributeCreationOverlap for explanation).
+    HLSLSettings->SetSourceText(
+        TEXT("Out_SetBool(Out_DataIndex, ElementIndex, 'AttrC', true);\n")
+        TEXT("Out_SetFloat3(Out_DataIndex, ElementIndex, 'AttrD', float3(1.0f, 2.0f, 3.0f));\n")
+    );
+
+    // Author "AttrA" as Float and "AttrB" as Int.
+    check(HLSLSettings->OutputPins.Num() == 1);
+    TArray<FPCGKernelAttributeKey>& Keys = HLSLSettings->OutputPins[0].PropertiesGPU.CreatedKernelAttributeKeys;
+    Keys.Add(FPCGKernelAttributeKey(FName("AttrA"), EPCGKernelAttributeType::Float));
+    Keys.Add(FPCGKernelAttributeKey(FName("AttrB"), EPCGKernelAttributeType::Int));
+
+    UPCGGatherSettings* GatherSettings = nullptr;
+    UPCGNode* GatherNode = Runner.Graph->AddNodeOfType<UPCGGatherSettings>(GatherSettings);
+    Runner.Graph->AddEdge(HLSLNode, PCGPinConstants::DefaultOutputLabel, GatherNode, PCGPinConstants::DefaultInputLabel);
+
+    TArray<FPCGDataCollection> OutNodeData;
+    if (!Runner.Execute({ this, { GatherNode } }, OutNodeData))
+    {
+        return false;
+    }
+
+    return PCGGPUTestCommon::ValidateSingleTaggedData(this, OutNodeData[0], [&](const FPCGTaggedData& TaggedData)
+    {
+        const UPCGBasePointData* OutPoints = Cast<UPCGBasePointData>(TaggedData.Data);
+        if (!TestNotNull(TEXT("Readback data is point data"), OutPoints))
+        {
+            return false;
+        }
+
+        const UPCGMetadata* Metadata = OutPoints->ConstMetadata();
+        if (!TestNotNull(TEXT("Output has metadata"), Metadata))
+        {
+            return false;
+        }
+
+        bool bResult = true;
+        bResult &= TestNotNull(TEXT("AttrA (Float, authored) — present"),   Metadata->GetConstTypedAttribute<float>(FName("AttrA")));
+        bResult &= TestNotNull(TEXT("AttrB (Int, authored) — present"),     Metadata->GetConstTypedAttribute<int32>(FName("AttrB")));
+        bResult &= TestNotNull(TEXT("AttrC (Bool, HLSL inferred) — present"),   Metadata->GetConstTypedAttribute<bool>(FName("AttrC")));
+        bResult &= TestNotNull(TEXT("AttrD (Float3, HLSL inferred) — present"), Metadata->GetConstTypedAttribute<FVector>(FName("AttrD")));
+        return bResult;
+    });
+}
+
+#endif // WITH_EDITOR
diff --git a/Engine/Plugins/PCG/Source/PCG/Private/Tests/Compute/PCGCustomHLSLTest.cpp b/Engine/Plugins/PCG/Source/PCG/Private/Tests/Compute/PCGCustomHLSLTest.cpp
--- a/Engine/Plugins/PCG/Source/PCG/Private/Tests/Compute/PCGCustomHLSLTest.cpp
+++ b/Engine/Plugins/PCG/Source/PCG/Private/Tests/Compute/PCGCustomHLSLTest.cpp
@@ -6,15 +6,12 @@
 
 #include "PCGCommon.h"
 #include "PCGGraph.h"
-#include "PCGParamData.h"
 #include "Compute/Elements/PCGCustomHLSL.h"
 #include "Compute/PCGComputeCommon.h"
 #include "Compute/PCGPinPropertiesGPU.h"
 #include "Data/PCGBasePointData.h"
-#include "Elements/PCGCreateAttribute.h"
 #include "Elements/PCGCreatePoints.h"
 #include "Elements/PCGGather.h"
-#include "Metadata/PCGMetadataAttributeTpl.h"
 
 /**
  * GPU test: CustomHLSL PointProcessor translates point positions.
@@ -198,150 +195,6 @@ bool FPCGCustomHLSLChainedPointProcessorsTest::RunTest(const FString& Parameters
 	});
 }
 
-/**
- * GPU test: CustomHLSL AttributeSetProcessor doubles a float attribute.
- *
- *   CreateAttributeSet (Value=5) -> CustomHLSL (Value *= 2) -> Gather
- *
- * Verifies the output attribute set has Value == 10.
- */
-IMPLEMENT_CUSTOM_SIMPLE_AUTOMATION_TEST(
-	FPCGCustomHLSLAttributeSetDoubleValueTest,
-	FPCGTestBaseClass,
-	"Plugins.PCG.Compute.CustomHLSL.AttributeSetDoubleValue",
-	PCGTestsCommon::TestFlags | EAutomationTestFlags::NonNullRHI)
-
-bool FPCGCustomHLSLAttributeSetDoubleValueTest::RunTest(const FString& Parameters)
-{
-	FPCGGPUGraphTestRunner Runner;
-
-	// CreateAttributeSet: one element with float attribute "Value" = 5.
-	UPCGCreateAttributeSetSettings* CreateAttrSettings = nullptr;
-	UPCGNode* CreateAttrNode = Runner.Graph->AddNodeOfType<UPCGCreateAttributeSetSettings>(CreateAttrSettings);
-	CreateAttrSettings->AttributeTypes.Type = EPCGMetadataTypes::Float;
-	CreateAttrSettings->AttributeTypes.FloatValue = 5.0f;
-	CreateAttrSettings->OutputTarget.SetAttributeName(FName(TEXT("Value")));
-
-	// CustomHLSL AttributeSetProcessor: Value *= 2.
-	UPCGCustomHLSLSettings* HLSLSettings = nullptr;
-	UPCGNode* HLSLNode = Runner.Graph->AddNodeOfType<UPCGCustomHLSLSettings>(HLSLSettings);
-	HLSLSettings->SetKernelType(EPCGKernelType::AttributeSetProcessor);
-	HLSLSettings->SetSourceText(
-		TEXT("float v = In_GetFloat(In_DataIndex, ElementIndex, 'Value');\n")
-		TEXT("Out_SetFloat(Out_DataIndex, ElementIndex, 'Value', v * 2.0f);\n")
-	);
-
-	// Useful to add gather node to grab output data from.
-	UPCGGatherSettings* GatherSettings = nullptr;
-	UPCGNode* GatherNode = Runner.Graph->AddNodeOfType<UPCGGatherSettings>(GatherSettings);
-
-	Runner.Graph->AddEdge(CreateAttrNode, PCGPinConstants::DefaultOutputLabel, HLSLNode, PCGPinConstants::DefaultInputLabel);
-	Runner.Graph->AddEdge(HLSLNode, PCGPinConstants::DefaultOutputLabel, GatherNode, PCGPinConstants::DefaultInputLabel);
-
-	TArray<FPCGDataCollection> OutNodeData;
-	if (!Runner.Execute({ this, { GatherNode } }, OutNodeData))
-	{
-		return false;
-	}
-
-	return PCGGPUTestCommon::ValidateSingleTaggedData(this, OutNodeData[0], [&](const FPCGTaggedData& TaggedData)
-	{
-		const UPCGParamData* OutParams = Cast<UPCGParamData>(TaggedData.Data);
-		if (!TestNotNull(TEXT("Readback data is attribute set data"), OutParams))
-		{
-			return false;
-		}
-
-		const UPCGMetadata* Metadata = OutParams->ConstMetadata();
-		if (!TestNotNull(TEXT("Output has metadata"), Metadata))
-		{
-			return false;
-		}
-
-		const FPCGMetadataAttribute<float>* FloatAttr = Metadata->GetConstTypedAttribute<float>(FName(TEXT("Value")));
-		if (!TestNotNull(TEXT("Output has 'Value' float attribute"), FloatAttr))
-		{
-			return false;
-		}
-
-		const double ActualValue = FloatAttr->GetValueFromItemKey(0);
-		return TestNearlyEqual(*FString::Printf(TEXT("Value attribute doubled (expected 10.0, got %.1f)"), ActualValue), ActualValue, 10.0, 0.1);
-	});
-}
-
-/**
- * GPU test: CustomHLSL AttributeSetProcessor passes a string attribute through unchanged.
- *
- *   CreateAttributeSet (StringAttr="Hello") -> CustomHLSL (read StringAttr, write StringAttr) -> Gather
- *
- * Verifies the output attribute set has StringAttr == "Hello".
- */
-IMPLEMENT_CUSTOM_SIMPLE_AUTOMATION_TEST(
-	FPCGCustomHLSLAttributeSetStringAttributeTest,
-	FPCGTestBaseClass,
-	"Plugins.PCG.Compute.CustomHLSL.AttributeSetStringAttribute",
-	PCGTestsCommon::TestFlags | EAutomationTestFlags::NonNullRHI)
-
-bool FPCGCustomHLSLAttributeSetStringAttributeTest::RunTest(const FString& Parameters)
-{
-	FPCGGPUGraphTestRunner Runner;
-
-	// CreateAttributeSet: one element with string attribute "StringAttr" = "Hello".
-	UPCGCreateAttributeSetSettings* CreateAttrSettings = nullptr;
-	UPCGNode* CreateAttrNode = Runner.Graph->AddNodeOfType<UPCGCreateAttributeSetSettings>(CreateAttrSettings);
-	CreateAttrSettings->AttributeTypes.Type = EPCGMetadataTypes::String;
-	CreateAttrSettings->AttributeTypes.StringValue = TEXT("Hello");
-	CreateAttrSettings->OutputTarget.SetAttributeName(FName(TEXT("StringAttr")));
-
-	// CustomHLSL AttributeSetProcessor: read StringAttr and write it back unchanged.
-	UPCGCustomHLSLSettings* HLSLSettings = nullptr;
-	UPCGNode* HLSLNode = Runner.Graph->AddNodeOfType<UPCGCustomHLSLSettings>(HLSLSettings);
-	HLSLSettings->SetKernelType(EPCGKernelType::AttributeSetProcessor);
-	HLSLSettings->SetSourceText(
-		TEXT("int StringAttrVal = In_GetStringKey(In_DataIndex, ElementIndex, 'StringAttr');\n")
-		TEXT("Out_SetStringKey(Out_DataIndex, ElementIndex, 'StringAttr', StringAttrVal);\n")
-	);
-
-	// Useful to add gather node to grab output data from.
-	UPCGGatherSettings* GatherSettings = nullptr;
-	UPCGNode* GatherNode = Runner.Graph->AddNodeOfType<UPCGGatherSettings>(GatherSettings);
-
-	Runner.Graph->AddEdge(CreateAttrNode, PCGPinConstants::DefaultOutputLabel, HLSLNode, PCGPinConstants::DefaultInputLabel);
-	Runner.Graph->AddEdge(HLSLNode, PCGPinConstants::DefaultOutputLabel, GatherNode, PCGPinConstants::DefaultInputLabel);
-
-	TArray<FPCGDataCollection> OutNodeData;
-	if (!Runner.Execute({ this, { GatherNode } }, OutNodeData))
-	{
-		return false;
-	}
-
-	return PCGGPUTestCommon::ValidateSingleTaggedData(this, OutNodeData[0], [&](const FPCGTaggedData& TaggedData)
-	{
-		const UPCGParamData* OutParams = Cast<UPCGParamData>(TaggedData.Data);
-		if (!TestNotNull(TEXT("Readback data is attribute set data"), OutParams))
-		{
-			return false;
-		}
-
-		const UPCGMetadata* Metadata = OutParams->ConstMetadata();
-		if (!TestNotNull(TEXT("Output has metadata"), Metadata))
-		{
-			return false;
-		}
-
-		const FPCGMetadataAttribute<FString>* StringAttr = Metadata->GetConstTypedAttribute<FString>(FName(TEXT("StringAttr")));
-		if (!TestNotNull(TEXT("Output has 'StringAttr' string attribute"), StringAttr))
-		{
-			return false;
-		}
-
-		const FString ActualValue = StringAttr->GetValueFromItemKey(0);
-		return TestEqual(
-			*FString::Printf(TEXT("'StringAttr' attribute passed through (expected 'Hello', got '%s')"), *ActualValue),
-			ActualValue, FString(TEXT("Hello")));
-	});
-}
-
 /**
  * GPU test: CustomHLSL PointProcessor with pin labels containing spaces.
  *
diff --git a/Engine/Plugins/PCG/Source/PCG/Public/Compute/PCGPinPropertiesGPU.h b/Engine/Plugins/PCG/Source/PCG/Public/Compute/PCGPinPropertiesGPU.h
--- a/Engine/Plugins/PCG/Source/PCG/Public/Compute/PCGPinPropertiesGPU.h
+++ b/Engine/Plugins/PCG/Source/PCG/Public/Compute/PCGPinPropertiesGPU.h
@@ -157,8 +157,8 @@ struct FPCGPinPropertiesGPUStruct
 	UPROPERTY(BlueprintReadWrite, EditAnywhere, Category = Settings, meta = (EditCondition = "bAllowEditInitMode && InitializationMode == EPCGPinInitMode::FromInputPins && !bShowTexturePinSettings", EditConditionHides))
 	EPCGAttributeInheritanceMode AttributeInheritanceMode = EPCGAttributeInheritanceMode::CopyAttributeSetup;
 
-	/** Add entries to create new attributes on data emitted by this pin. */
-	UPROPERTY(EditAnywhere, DisplayName = "Attributes to Create", Category = Settings, meta = (EditCondition = "!bShowTexturePinSettings", EditConditionHides, HideEditConditionToggle))
+	/** Add entries to create new attributes on data emitted by this pin. This is optional, as of 5.8 attributes are automatically created when written from HLSL. */
+	UPROPERTY(EditAnywhere, DisplayName = "Attributes to Create", Category = Settings, meta = (EditCondition = "!bShowTexturePinSettings && bShowCreatedAttributesSettings", EditConditionHides, HideEditConditionToggle))
 	TArray<FPCGKernelAttributeKey> CreatedKernelAttributeKeys;
 
 #if WITH_EDITORONLY_DATA
@@ -180,6 +180,10 @@ struct FPCGPinPropertiesGPUStruct
 
 	UPROPERTY(Transient)
 	bool bShowInitializeToZeroSetting = false;
+
+	/** Driven by UPCGCustomHLSLSettings::bPerPinAttributeCreationSettings. Shows/hides the CreatedKernelAttributeKeys array. */
+	UPROPERTY(Transient)
+	bool bShowCreatedAttributesSettings = false;
 #endif // WITH_EDITORONLY_DATA
 };
 
diff --git a/Engine/Source/Runtime/Core/Public/UObject/FortniteMainBranchObjectVersions.inl b/Engine/Source/Runtime/Core/Public/UObject/FortniteMainBranchObjectVersions.inl
--- a/Engine/Source/Runtime/Core/Public/UObject/FortniteMainBranchObjectVersions.inl
+++ b/Engine/Source/Runtime/Core/Public/UObject/FortniteMainBranchObjectVersions.inl
@@ -805,5 +805,8 @@ FORTNITE_MAIN_VERSION(CompactDynamicDelegatePayloads, 263)
 // Added bone attachment serialization for UInstancedSkinnedMeshComponent
 FORTNITE_MAIN_VERSION(InstancedSkinnedMeshBoneAttachments, 264)
 
+// Added toggle to show/hide per-pin attribute creation settings on PCG Custom HLSL nodes
+FORTNITE_MAIN_VERSION(PCGCustomHLSLAttributeCreationSettingsToggle, 265)
+
 // -----<new versions can be added above this line>-------------------------------------------------
 #undef FORTNITE_MAIN_VERSION
