[PCG] Make `FPCGMetadataAttribute<T>` inherit from Generic attributes

* Strip out everything type specific to forward the calls to the generic ones

#rnx
#Synced-CL 52163112

diff --git a/Engine/Plugins/PCG/Source/PCG/Private/Helpers/PCGSettingsHelpers.cpp b/Engine/Plugins/PCG/Source/PCG/Private/Helpers/PCGSettingsHelpers.cpp
--- a/Engine/Plugins/PCG/Source/PCG/Private/Helpers/PCGSettingsHelpers.cpp
+++ b/Engine/Plugins/PCG/Source/PCG/Private/Helpers/PCGSettingsHelpers.cpp
@@ -17,9 +17,15 @@ namespace PCGSettingsHelpers
 	template<typename T>
 	bool GetParamValue(T& Value, const FPCGMetadataAttributeBase* InAttribute, PCGMetadataEntryKey InKey = 0)
 	{
-		if (InAttribute->GetTypeId() == PCG::Private::MetadataTypes<T>::Id)
+		if (!InAttribute->IsGeneric())
 		{
-			Value = static_cast<const FPCGMetadataAttribute<T>*>(InAttribute)->GetValueFromItemKey(InKey);
+			return false;
+		}
+		
+		const FPCGMetadataAttributeGeneric* Attribute = static_cast<const FPCGMetadataAttributeGeneric*>(InAttribute);
+		if (Attribute->IsOfType<T>())
+		{
+			Value = Attribute->GetValueFromItemKey<T>(InKey);
 			return true;
 		}
 		else
@@ -31,9 +37,15 @@ namespace PCGSettingsHelpers
 	template<typename T, typename U, typename... OtherTypes>
 	bool GetParamValue(T& Value, const FPCGMetadataAttributeBase* InAttribute, PCGMetadataEntryKey InKey = 0)
 	{
-		if (InAttribute->GetTypeId() == PCG::Private::MetadataTypes<U>::Id)
+		if (!InAttribute->IsGeneric())
+		{
+			return false;
+		}
+		
+		U TempValue{};
+		if (GetParamValue<U>(TempValue, InAttribute, InKey))
 		{
-			Value = (T)static_cast<const FPCGMetadataAttribute<U>*>(InAttribute)->GetValueFromItemKey(InKey);
+			Value = T{TempValue};
 			return true;
 		}
 		else
@@ -129,7 +141,7 @@ namespace PCGSettingsHelpers
 			FNumericProperty* UnderlyingProperty = EnumProperty->GetUnderlyingProperty();
 
 			int64 Value;
-			if (GetParamValue<int64, int32, uint32, int16, uint16, int8, uint8>(Value, MatchingAttribute))
+			if (GetParamValue<int64, int32, uint8>(Value, MatchingAttribute))
 			{
 				UnderlyingProperty->SetIntPropertyValue(EnumProperty->ContainerPtrToValuePtr<void>(Object, 0), Value);
 			}
diff --git a/Engine/Plugins/PCG/Source/PCG/Private/Metadata/Accessors/PCGAttributeAccessorHelpers.cpp b/Engine/Plugins/PCG/Source/PCG/Private/Metadata/Accessors/PCGAttributeAccessorHelpers.cpp
--- a/Engine/Plugins/PCG/Source/PCG/Private/Metadata/Accessors/PCGAttributeAccessorHelpers.cpp
+++ b/Engine/Plugins/PCG/Source/PCG/Private/Metadata/Accessors/PCGAttributeAccessorHelpers.cpp
@@ -358,33 +358,13 @@ TUniquePtr<const IPCGAttributeAccessor> PCGAttributeAccessorHelpers::CreateConst
 
 TUniquePtr<const IPCGAttributeAccessor> PCGAttributeAccessorHelpers::CreateConstAccessor(const FPCGMetadataAttributeBase* InAttribute, const FPCGMetadataDomain* InMetadata, bool bQuiet)
 {
-	if (!InMetadata || !InAttribute)
+	if (!InMetadata || !InAttribute || !InAttribute->IsGeneric())
 	{
 		return {};
 	}
 	else
 	{
-		auto MakeConstAccessor = [InAttribute, InMetadata]<typename T>(T) -> TUniquePtr<const IPCGAttributeAccessor>
-		{
-			if (InAttribute->IsGeneric())
-			{
-				return MakeUnique<const FPCGAttributeGenericAccessor<T>>(static_cast<const FPCGMetadataAttributeGeneric*>(InAttribute), InMetadata);
-			}
-			else
-			{
-				return MakeUnique<const FPCGAttributeAccessor<T>>(static_cast<const FPCGMetadataAttribute<T>*>(InAttribute), InMetadata);
-			}
-		};
-		
-		TUniquePtr<const IPCGAttributeAccessor> OutAccessor = PCGMetadataAttribute::CallbackWithRightType(InAttribute->GetTypeId(), MakeConstAccessor);
-		if (!OutAccessor && InAttribute->IsGeneric())
-		{
-			// @todo_pcg: To be updated when we won't have typed accessors anymore.
-			// This creates a generic accessor that knows its type with the desc in InAttribute.
-			OutAccessor = MakeUnique<const FPCGAttributeGenericAccessor<uint8>>(static_cast<const FPCGMetadataAttributeGeneric*>(InAttribute), InMetadata);
-		}
-		
-		return OutAccessor;
+		return MakeUnique<const FPCGAttributeGenericAccessor>(static_cast<const FPCGMetadataAttributeGeneric*>(InAttribute), InMetadata);
 	}
 }
 
@@ -450,33 +430,13 @@ TUniquePtr<IPCGAttributeAccessor> PCGAttributeAccessorHelpers::CreateAccessor(FP
 
 TUniquePtr<IPCGAttributeAccessor> PCGAttributeAccessorHelpers::CreateAccessor(FPCGMetadataAttributeBase* InAttribute, FPCGMetadataDomain* InMetadata, bool bQuiet)
 {
-	if (!InMetadata || !InAttribute)
+	if (!InMetadata || !InAttribute || !InAttribute->IsGeneric())
 	{
 		return {};
 	}
 	else
 	{
-		auto MakeAccessor = [InAttribute, InMetadata]<typename T>(T) -> TUniquePtr<IPCGAttributeAccessor>
-		{
-			if (InAttribute->IsGeneric())
-			{
-				return MakeUnique<FPCGAttributeGenericAccessor<T>>(static_cast<FPCGMetadataAttributeGeneric*>(InAttribute), InMetadata);
-			}
-			else
-			{
-				return MakeUnique<FPCGAttributeAccessor<T>>(static_cast<FPCGMetadataAttribute<T>*>(InAttribute), InMetadata);
-			}
-		};
-
-		TUniquePtr<IPCGAttributeAccessor> OutAccessor = PCGMetadataAttribute::CallbackWithRightType(InAttribute->GetTypeId(), MakeAccessor);
-		if (!OutAccessor && InAttribute->IsGeneric())
-		{
-			// @todo_pcg: To be updated when we won't have typed accessors anymore.
-			// This creates a generic accessor that knows its type with the desc in InAttribute.
-			OutAccessor = MakeUnique<FPCGAttributeGenericAccessor<uint8>>(static_cast<FPCGMetadataAttributeGeneric*>(InAttribute), InMetadata);
-		}
-
-		return OutAccessor;
+		return MakeUnique<FPCGAttributeGenericAccessor>(static_cast<FPCGMetadataAttributeGeneric*>(InAttribute), InMetadata);
 	}
 }
 
@@ -568,7 +528,7 @@ TUniquePtr<IPCGAttributeAccessor> PCGAttributeAccessorHelpers::CreateAccessorWit
 
 			FPCGMetadataAttribute<AttributeType>* Attribute = MetadataDomain->CreateAttribute<AttributeType>(AttributeName, DefaultValue, /*bAllowInterpolation=*/true, /*bOverrideParent=*/false);
 
-			return Attribute ? MakeUnique<FPCGAttributeAccessor<AttributeType>>(Attribute, MetadataDomain) : nullptr;
+			return Attribute ? MakeUnique<FPCGAttributeGenericAccessor>(Attribute, MetadataDomain) : nullptr;
 		};
 
 		Result = PCGMetadataAttribute::CallbackWithRightType(MatchingType, std::move(CreateAttributeAndAccessor));
diff --git a/Engine/Plugins/PCG/Source/PCG/Public/Metadata/Accessors/PCGAttributeAccessor.h b/Engine/Plugins/PCG/Source/PCG/Public/Metadata/Accessors/PCGAttributeAccessor.h
--- a/Engine/Plugins/PCG/Source/PCG/Public/Metadata/Accessors/PCGAttributeAccessor.h
+++ b/Engine/Plugins/PCG/Source/PCG/Public/Metadata/Accessors/PCGAttributeAccessor.h
@@ -12,39 +12,37 @@
 #define LOCTEXT_NAMESPACE "PCGAttributeAccessor"
 
 /**
-* Templated accessor class for attributes. Will wrap around a typed attribute.
+* Templated accessor class for attributes. Will wrap around a generic attribute.
 * Key supported: MetadataEntryKey and Points
 */
-template <typename T, typename UnderlyingAccessor>
-class FPCGAttributeAccessorBase : public IPCGAttributeAccessorT<FPCGAttributeAccessorBase<T, UnderlyingAccessor>>
+class FPCGAttributeGenericAccessor : public IPCGAttributeAccessor
 {
 public:
-	using Type = T;
-	using Super = IPCGAttributeAccessorT<FPCGAttributeAccessorBase<T, UnderlyingAccessor>>;
-	
-protected:
+	using Super = IPCGAttributeAccessor;
+
 	// Can't write if metadata is null
-	FPCGAttributeAccessorBase(FPCGMetadataAttributeBase* InAttribute, FPCGMetadataDomain* InMetadataDomain, bool bForceReadOnly = false)
-		: Super(/*bInReadOnly=*/ InMetadataDomain == nullptr || bForceReadOnly)
+	FPCGAttributeGenericAccessor(FPCGMetadataAttributeGeneric* InAttribute, FPCGMetadataDomain* InMetadataDomain, bool bForceReadOnly = false)
+		: Super(/*bInReadOnly=*/ InMetadataDomain == nullptr || bForceReadOnly, static_cast<int16>(InAttribute->GetAttributeDesc().ValueType))
 		, Attribute(InAttribute)
 		, MetadataDomain(InMetadataDomain)
 	{
 		check(InAttribute && (!InMetadataDomain || InAttribute->GetMetadataDomain() == InMetadataDomain));
+		UnderlyingDesc = InAttribute->GetAttributeDesc();
 	}
 
-	FPCGAttributeAccessorBase(const FPCGMetadataAttributeBase* InAttribute, const FPCGMetadataDomain* InMetadataDomain, bool bForceReadOnly = false)
-		: Super(/*bInReadOnly=*/ true)
-		, Attribute(const_cast<FPCGMetadataAttributeBase*>(InAttribute))
+	FPCGAttributeGenericAccessor(const FPCGMetadataAttributeGeneric* InAttribute, const FPCGMetadataDomain* InMetadataDomain, bool bForceReadOnly = false)
+		: Super(/*bInReadOnly=*/ true, static_cast<int16>(InAttribute->GetAttributeDesc().ValueType))
+		, Attribute(const_cast<FPCGMetadataAttributeGeneric*>(InAttribute))
 		, MetadataDomain(const_cast<FPCGMetadataDomain*>(InMetadataDomain))
 	{
 		check(InAttribute && (!InMetadataDomain || InAttribute->GetMetadataDomain() == InMetadataDomain));
+		UnderlyingDesc = InAttribute->GetAttributeDesc();
 	}
-	
-public:
-	bool GetRangeImpl(TArrayView<T> OutValues, int32 Index, const IPCGAttributeAccessorKeys& Keys) const
+
+	virtual bool GetRangeVirtual(PCG::Private::FOutValues OutValues, int32 Count, int32 Index, const IPCGAttributeAccessorKeys& Keys) const override
 	{
 		TArray<const PCGMetadataEntryKey*, TInlineAllocator<256>> EntryKeyPtrs;
-		EntryKeyPtrs.SetNumUninitialized(OutValues.Num());
+		EntryKeyPtrs.SetNumUninitialized(Count);
 
 		TArrayView<const PCGMetadataEntryKey*> EntryKeysView(EntryKeyPtrs);
 		if (!Keys.GetKeys<PCGMetadataEntryKey>(Index, EntryKeysView))
@@ -55,13 +53,15 @@ public:
 		TArray<PCGMetadataEntryKey, TInlineAllocator<256>> EntryKeys;
 		EntryKeys.Reserve(EntryKeyPtrs.Num());
 		Algo::Transform(EntryKeyPtrs, EntryKeys, [](const PCGMetadataEntryKey* KeyPtr) { return *KeyPtr; });
-
-		static_cast<const UnderlyingAccessor&>(*this).GetValuesFromItemKeys(TArrayView<PCGMetadataEntryKey>(EntryKeys), OutValues);
+		
+		TArray<PCGMetadataValueKey> ValueKeys;
+		Attribute->GetValueKeys(PCGValueRangeHelpers::MakeConstValueRange<PCGMetadataEntryKey>(EntryKeys), ValueKeys);
+		Attribute->GetValues_Internal(ValueKeys, MoveTemp(OutValues));
 
 		return true;
 	}
-
-	bool SetRangeImpl(TArrayView<const T> InValues, int32 Index, IPCGAttributeAccessorKeys& Keys, EPCGAttributeAccessorFlags Flags)
+	
+	virtual bool SetRangeVirtual(PCG::Private::FInValues InValues, int32 Count, int32 Index, IPCGAttributeAccessorKeys& Keys, EPCGAttributeAccessorFlags Flags) override
 	{
 		if (EnumHasAnyFlags(Flags, EPCGAttributeAccessorFlags::AllowSetDefaultValue) && Keys.GetNum() == 1)
 		{
@@ -70,35 +70,40 @@ public:
 
 			if (EntryKey && *EntryKey == PCGInvalidEntryKey)
 			{
-				check(!InValues.IsEmpty())
-				static_cast<UnderlyingAccessor&>(*this).SetDefaultValue(InValues[0]);
+				check(Count > 0);
+				Attribute->AddValues_Internal(InValues, 1, /*bExistingIndex=*/PCG::Private::ExistingIndexForDefaultValue);
 				return true;
 			}
 		}
 
 		TArray<PCGMetadataEntryKey*, TInlineAllocator<512>> EntryKeys;
-		EntryKeys.SetNumUninitialized(InValues.Num());
+		EntryKeys.SetNumUninitialized(Count);
 		TArrayView<PCGMetadataEntryKey*> EntryKeysView(EntryKeys);
-
-		int32 Start = Index;
-
-		if (!bWasPrepared || PCG::Private::MetadataTraits<T>::CompressData)
+		
+		if (!bWasPrepared || Attribute->DoesCompressData())
 		{
-			if (!Prepare(Keys, InValues.Num(), /*bCanReuseEntryKeys=*/EnumHasAnyFlags(Flags, EPCGAttributeAccessorFlags::AllowReuseMetadataEntryKey), Index, /*bPreallocateValues=*/false, &EntryKeysView, &Start))
+			if (!Prepare(Keys, Count, /*bCanReuseEntryKeys=*/EnumHasAnyFlags(Flags, EPCGAttributeAccessorFlags::AllowReuseMetadataEntryKey), Index, /*bPreallocateValues=*/false, &EntryKeysView))
 			{
 				return false;
 			}
 
+			TArray<PCGMetadataValueKey> ValueKeys;
+			
+			if (Attribute->DoesCompressData() && ensure(InValues.IsType<PCG::Private::FInValuesByValue>()))
+			{
+				ValueKeys = Attribute->AddCompressedValues(InValues.Get<PCG::Private::FInValuesByValue>(), Count);
+			}
+			else
 			{
-				//TRACE_CPUPROFILER_EVENT_SCOPE(FPCGAttributeAccessor::SetRangeImpl::SetValues);
-				static_cast<UnderlyingAccessor&>(*this).SetValues(EntryKeys, InValues);
+				ValueKeys = Attribute->AddValues_Internal(InValues, Count);
 			}
+			
+			Attribute->SetValuesFromValueKeys(EntryKeys, ValueKeys);
 		}
 		else
 		{
 			check(StartIndex != INDEX_NONE);
-			//TRACE_CPUPROFILER_EVENT_SCOPE(FPCGAttributeAccessor::SetRangeImpl::SetValues);
-			static_cast<UnderlyingAccessor&>(*this).SetValues_TryLockless(EntryKeys, InValues, StartIndex + Start);
+			Attribute->AddValues_Internal(InValues, Count, /*bExistingIndex=*/StartIndex + Index, /*bLockless=*/true);
 		}
 
 		return true;
@@ -106,11 +111,9 @@ public:
 
 	virtual bool IsAttribute() const override { return true; }
 
-	bool Prepare(IPCGAttributeAccessorKeys& Keys, int32 Count, const bool bCanReuseEntryKeys, int32 Index, bool bPreallocateValues, TArrayView<PCGMetadataEntryKey*>* OutEntryKeys = nullptr, int32* OutStartIndex = nullptr)
+	bool Prepare(IPCGAttributeAccessorKeys& Keys, int32 Count, const bool bCanReuseEntryKeys, int32 Index, bool bPreallocateValues, TArrayView<PCGMetadataEntryKey*>* OutEntryKeys = nullptr)
 	{
-		//TRACE_CPUPROFILER_EVENT_SCOPE(FPCGAttributeAccessor::Prepare);
-		
-		if (!MetadataDomain || Super::IsReadOnly())
+		if (!MetadataDomain || IsReadOnly())
 		{
 			PCGLog::LogWarningOnGraph(LOCTEXT("ReadOnlyInvalidPrepare", "Try to prepare an attribute accessor with a read only accessor. This is invalid."));
 			return false;		
@@ -157,7 +160,6 @@ public:
 
 		if (!EntriesToSet.IsEmpty())
 		{
-			TRACE_CPUPROFILER_EVENT_SCOPE(FPCGAttributeAccessor::Prepare::AddEntriesInPlace);
 			if (bSupportsMultiEntries || (MetadataDomain->GetItemCountForChild() == 0 && ensure(EntriesToSet.Num() == 1)))
 			{
 				MetadataDomain->AddEntriesInPlace(EntriesToSet);
@@ -170,11 +172,7 @@ public:
 
 		if (bPreallocateValues && ensure(!bWasPrepared))
 		{
-			const int32 TempStartIndex = static_cast<UnderlyingAccessor&>(*this).PreallocateValues(EntryKeysView, /*bLockless=*/false);
-			if (OutStartIndex)
-			{
-				*OutStartIndex = TempStartIndex;
-			}
+			StartIndex = Attribute->PreallocateValues(EntryKeysView, /*bLockless=*/false);
 		}
 
 		return true;
@@ -187,192 +185,22 @@ public:
 			return;
 		}
 
-		static_cast<UnderlyingAccessor&>(*this).PrepareAttribute(Count);
-		bWasPrepared = Prepare(Keys, Count, bCanReuseEntryKeys, /*Index=*/0, /*bPreallocateValues=*/true, /*OutEntryKeys=*/nullptr, &StartIndex);
+		Attribute->Prepare(Count);
+		bWasPrepared = Prepare(Keys, Count, bCanReuseEntryKeys, /*Index=*/0, /*bPreallocateValues=*/true, /*OutEntryKeys=*/nullptr);
 	}
 
 protected:
-	FPCGMetadataAttributeBase* Attribute = nullptr;
+	FPCGMetadataAttributeGeneric* Attribute = nullptr;
 	FPCGMetadataDomain* MetadataDomain = nullptr;
 	bool bWasPrepared = false;
 	int32 StartIndex = INDEX_NONE;
 };
 
+/** Deprecated accessor, use the generic one.*/
 template <typename T>
-class FPCGAttributeAccessor : public FPCGAttributeAccessorBase<T, FPCGAttributeAccessor<T>>
+class UE_DEPRECATED(5.8, "Use the Generic Accessor.") FPCGAttributeAccessor : public FPCGAttributeGenericAccessor
 {
-public:
-	using Super = FPCGAttributeAccessorBase<T, FPCGAttributeAccessor<T>>;
-	
-	// Can't write if metadata is null
-	FPCGAttributeAccessor(FPCGMetadataAttribute<T>* InAttribute, FPCGMetadataDomain* InMetadataDomain, bool bForceReadOnly = false)
-		: Super(InAttribute, InMetadataDomain, bForceReadOnly)
-	{
-	}
-
-	FPCGAttributeAccessor(const FPCGMetadataAttribute<T>* InAttribute, const FPCGMetadataDomain* InMetadataDomain, bool bForceReadOnly = false)
-		: Super(InAttribute, InMetadataDomain, bForceReadOnly)
-	{
-	}
-	
-	void GetValuesFromItemKeys(TArrayView<PCGMetadataEntryKey> EntryKeys, TArrayView<T> OutValues) const
-	{
-		static_cast<const FPCGMetadataAttribute<T>*>(this->Attribute)->GetValuesFromItemKeys(EntryKeys, OutValues);
-	}
-
-	void SetDefaultValue(const T& Value)
-	{
-		static_cast<FPCGMetadataAttribute<T>*>(this->Attribute)->SetDefaultValue(Value);
-	}
-
-	void SetValues(TArrayView<const PCGMetadataEntryKey * const> ItemKeys, const TArrayView<const T>& InValues)
-	{
-		static_cast<FPCGMetadataAttribute<T>*>(this->Attribute)->SetValues(ItemKeys, InValues);
-	}
-
-	void SetValues_TryLockless(TArrayView<PCGMetadataEntryKey*> ItemKeys, const TArrayView<const T>& InValues, int32 StartIndex)
-	{
-		static_cast<FPCGMetadataAttribute<T>*>(this->Attribute)->SetValues_TryLockless(ItemKeys, InValues, StartIndex);
-	}
-
-	int32 PreallocateValues(TArrayView<PCGMetadataEntryKey*> EntryKeys, bool bLockless)
-	{
-		return static_cast<FPCGMetadataAttribute<T>*>(this->Attribute)->PreallocateValues(EntryKeys, bLockless);
-	}
-
-	void PrepareAttribute(int32 Count)
-	{
-		static_cast<FPCGMetadataAttribute<T>*>(this->Attribute)->Prepare(Count);
-	}
-};
-
-template <typename T>
-class FPCGAttributeGenericAccessor : public FPCGAttributeAccessorBase<T, FPCGAttributeGenericAccessor<T>>
-{
-public:
-	using Super = FPCGAttributeAccessorBase<T, FPCGAttributeGenericAccessor<T>>;
-	
-	// Can't write if metadata is null
-	FPCGAttributeGenericAccessor(FPCGMetadataAttributeGeneric* InAttribute, FPCGMetadataDomain* InMetadataDomain, bool bForceReadOnly = false)
-		: Super(InAttribute, InMetadataDomain, bForceReadOnly)
-	{
-		this->UnderlyingDesc = InAttribute->GetAttributeDesc();
-		this->UnderlyingType = static_cast<int16>(this->UnderlyingDesc.ValueType);
-	}
-
-	FPCGAttributeGenericAccessor(const FPCGMetadataAttributeGeneric* InAttribute, const FPCGMetadataDomain* InMetadataDomain, bool bForceReadOnly = false)
-		: Super(InAttribute, InMetadataDomain, bForceReadOnly)
-	{
-		this->UnderlyingDesc = InAttribute->GetAttributeDesc();
-		this->UnderlyingType = static_cast<int16>(this->UnderlyingDesc.ValueType);
-	}
-	
-	virtual bool GetRangeVirtual(PCG::Private::FOutValues OutValues, int32 Count, int32 Index, const IPCGAttributeAccessorKeys& Keys) const override
-	{
-		const FPCGMetadataAttributeGeneric* ThisAttribute = static_cast<const FPCGMetadataAttributeGeneric*>(this->Attribute);
-		
-		TArray<const PCGMetadataEntryKey*, TInlineAllocator<256>> EntryKeyPtrs;
-		EntryKeyPtrs.SetNumUninitialized(Count);
-
-		TArrayView<const PCGMetadataEntryKey*> EntryKeysView(EntryKeyPtrs);
-		if (!Keys.GetKeys<PCGMetadataEntryKey>(Index, EntryKeysView))
-		{
-			return false;
-		}
-
-		TArray<PCGMetadataEntryKey, TInlineAllocator<256>> EntryKeys;
-		EntryKeys.Reserve(EntryKeyPtrs.Num());
-		Algo::Transform(EntryKeyPtrs, EntryKeys, [](const PCGMetadataEntryKey* KeyPtr) { return *KeyPtr; });
-		
-		TArray<PCGMetadataValueKey> ValueKeys;
-		ThisAttribute->GetValueKeys(PCGValueRangeHelpers::MakeConstValueRange<PCGMetadataEntryKey>(EntryKeys), ValueKeys);
-		ThisAttribute->GetValues_Internal(ValueKeys, MoveTemp(OutValues));
-
-		return true;
-	}
-	
-	virtual bool SetRangeVirtual(PCG::Private::FInValues InValues, int32 Count, int32 Index, IPCGAttributeAccessorKeys& Keys, EPCGAttributeAccessorFlags Flags) override
-	{
-		FPCGMetadataAttributeGeneric* ThisAttribute = static_cast<FPCGMetadataAttributeGeneric*>(this->Attribute);
-		
-		if (EnumHasAnyFlags(Flags, EPCGAttributeAccessorFlags::AllowSetDefaultValue) && Keys.GetNum() == 1)
-		{
-			PCGMetadataEntryKey* EntryKey = nullptr;
-			Keys.GetKey(EntryKey);
-
-			if (EntryKey && *EntryKey == PCGInvalidEntryKey)
-			{
-				check(Count > 0);
-				ThisAttribute->AddValues_Internal(InValues, 1, /*bExistingIndex=*/PCG::Private::ExistingIndexForDefaultValue);
-				return true;
-			}
-		}
-
-		TArray<PCGMetadataEntryKey*, TInlineAllocator<512>> EntryKeys;
-		EntryKeys.SetNumUninitialized(Count);
-		TArrayView<PCGMetadataEntryKey*> EntryKeysView(EntryKeys);
-
-		// Not const because it can be modified by Prepare
-		int32 Start = Index;
-
-		if (!this->bWasPrepared || ThisAttribute->DoesCompressData())
-		{
-			if (!Super::Prepare(Keys, Count, /*bCanReuseEntryKeys=*/EnumHasAnyFlags(Flags, EPCGAttributeAccessorFlags::AllowReuseMetadataEntryKey), Index, /*bPreallocateValues=*/true, &EntryKeysView, &Start))
-			{
-				return false;
-			}
-
-			TArray<PCGMetadataValueKey> ValueKeys;
-			
-			if (ThisAttribute->DoesCompressData() && ensure(InValues.IsType<PCG::Private::FInValuesByValue>()))
-			{
-				ValueKeys = ThisAttribute->AddCompressedValues(InValues.Get<PCG::Private::FInValuesByValue>(), Count);
-			}
-			else
-			{
-				ValueKeys = ThisAttribute->AddValues_Internal(InValues, Count, /*bExistingIndex=*/Start);
-			}
-			
-			ThisAttribute->SetValuesFromValueKeys(EntryKeys, ValueKeys);
-		}
-		else
-		{
-			check(this->StartIndex != INDEX_NONE);
-			ThisAttribute->AddValues_Internal(InValues, Count, /*bExistingIndex=*/this->StartIndex + Start, /*bLockless=*/true);
-		}
-
-		return true;
-	}
-	
-	void GetValuesFromItemKeys(TArrayView<PCGMetadataEntryKey> EntryKeys, TArrayView<T> OutValues) const
-	{
-		static_cast<const FPCGMetadataAttributeGeneric*>(this->Attribute)->GetValuesFromItemKeys(EntryKeys, OutValues);
-	}
-
-	void SetDefaultValue(const T& Value)
-	{
-		static_cast<FPCGMetadataAttributeGeneric*>(this->Attribute)->SetDefaultValue(Value);
-	}
-
-	void SetValues(TArrayView<const PCGMetadataEntryKey * const> ItemKeys, const TArrayView<const T>& InValues)
-	{
-		static_cast<FPCGMetadataAttributeGeneric*>(this->Attribute)->SetValues(ItemKeys, InValues);
-	}
-
-	void SetValues_TryLockless(TArrayView<PCGMetadataEntryKey*> ItemKeys, const TArrayView<const T>& InValues, int32 StartIndex)
-	{
-		static_cast<FPCGMetadataAttributeGeneric*>(this->Attribute)->SetValues_TryLockless(ItemKeys, InValues, StartIndex);
-	}
-
-	int32 PreallocateValues(TArrayView<PCGMetadataEntryKey*> EntryKeys, bool bLockless)
-	{
-		return static_cast<FPCGMetadataAttributeGeneric*>(this->Attribute)->PreallocateValues(EntryKeys, bLockless);
-	}
-
-	void PrepareAttribute(int32 Count)
-	{
-		static_cast<FPCGMetadataAttributeGeneric*>(this->Attribute)->Prepare(Count);
-	}
+	using FPCGAttributeGenericAccessor::FPCGAttributeGenericAccessor;
 };
 
 #undef LOCTEXT_NAMESPACE
\ No newline at end of file
diff --git a/Engine/Plugins/PCG/Source/PCG/Public/Metadata/PCGMetadataAttributeGeneric.h b/Engine/Plugins/PCG/Source/PCG/Public/Metadata/PCGMetadataAttributeGeneric.h
--- a/Engine/Plugins/PCG/Source/PCG/Public/Metadata/PCGMetadataAttributeGeneric.h
+++ b/Engine/Plugins/PCG/Source/PCG/Public/Metadata/PCGMetadataAttributeGeneric.h
@@ -378,8 +378,7 @@ public:
 protected:
 	/** Attributes can't be created individually, they need to be created through a Metadata Domain. */
 	friend FPCGMetadataDomain;
-
-	template <typename T>
+	
 	friend class FPCGAttributeGenericAccessor;
 
 	// New attribute with no parent, name is contained in the InAttributeDesc.
@@ -576,7 +575,7 @@ protected:
 	PCG_API bool FindValues(TArrayView<PCGMetadataValueKey> OutValueKeys, const PCG::Private::FInValues& InValues, int32 Count) const;
 
 	template <typename T>
-	void InitForDeprecation(TConstArrayView<T> InValues, const T& InDefaultValue)
+	void InitForDeprecation(TArray<T> InValues, const T& InDefaultValue)
 	{
 		CachedDesc.Name = Name;
 
@@ -585,10 +584,8 @@ protected:
 			Helper->EmptyValues(InValues.Num());
 			SetDefaultValue<T>(InDefaultValue);
 
-			if (!InValues.IsEmpty())
-			{
-				AddValues<T>(InValues);
-			}
+			// FScriptArray and TArray are interchangeable, if and only if they are the exact same type
+			Helper->MoveAssign(&InValues);
 			
 			ValueKeyOffset = GetParent() ? GetParent()->GetValueKeyOffsetForChild() : 0;
 		}
diff --git a/Engine/Plugins/PCG/Source/PCG/Public/Metadata/PCGMetadataAttributeTpl.h b/Engine/Plugins/PCG/Source/PCG/Public/Metadata/PCGMetadataAttributeTpl.h
--- a/Engine/Plugins/PCG/Source/PCG/Public/Metadata/PCGMetadataAttributeTpl.h
+++ b/Engine/Plugins/PCG/Source/PCG/Public/Metadata/PCGMetadataAttributeTpl.h
@@ -2,22 +2,26 @@
 
 #pragma once
 
+#include "PCGMetadataAttributeGeneric.h"
 
-#include "Metadata/PCGMetadataCommon.h"
 #include "PCGModule.h"
-
 #include "Helpers/PCGMetadataHelpers.h"
 #include "Metadata/PCGMetadataAttribute.h"
 #include "Metadata/PCGMetadataAttributeTraits.h"
-#include "Misc/ScopeRWLock.h"
+#include "Metadata/PCGMetadataCommon.h"
 
 #include "Serialization/ArchiveCountMem.h"
+#include "UObject/FortniteMainBranchObjectVersion.h"
 
 class FPCGMetadataDomain;
 class UPCGMetadata;
 
+/**
+ * Not supposed to be used directly. Only supports legacy types, such as int32 or float.
+ * It does not support arrays or complex structs.
+ */
 template<typename T>
-class FPCGMetadataAttribute : public FPCGMetadataAttributeBase
+class FPCGMetadataAttribute : public FPCGMetadataAttributeGeneric
 {
 	template<typename U>
 	friend class FPCGMetadataAttribute;
@@ -25,206 +29,48 @@ class FPCGMetadataAttribute : public FPCGMetadataAttributeBase
 public:
 	/** TODO: All deprecated versions taking a UPCGMetadata will be added at the end of this declaration. For the review, it will be in place for "clarity" */
 	FPCGMetadataAttribute(FPCGMetadataDomain* InMetadata, FName InName, const FPCGMetadataAttributeBase* InParent, const T& InDefaultValue, bool bInAllowsInterpolation)
-		: FPCGMetadataAttributeBase(InMetadata, InName, InParent, bInAllowsInterpolation)
-		, DefaultValue(InDefaultValue)
+		: FPCGMetadataAttributeGeneric(PCG::Private::MakeAttributeDesc<T>(InName), InParent && InParent->IsGeneric() ? static_cast<const FPCGMetadataAttributeGeneric*>(InParent) : nullptr, InMetadata, bInAllowsInterpolation)
 	{
-		TypeId = PCG::Private::MetadataTypes<T>::Id;
-
-		// Make sure we don't parent with the wrong type id
-		check(!Parent || Parent->GetTypeId() == TypeId)
-
-		if (GetParent())
+		static_assert(PCG::Private::IsPCGType<T>(), "Only supports legacy types, such as int32 or float");
+		if (!InParent)
 		{
-			ValueKeyOffset = GetParent()->GetValueKeyOffsetForChild();
+			SetDefaultValue(InDefaultValue);
 		}
 	}
 
-	UE_DEPRECATED(5.6, "Use the version with FPCGMetadataDomain")
-	FPCGMetadataAttribute(UPCGMetadata* InMetadata, FName InName, const FPCGMetadataAttributeBase* InParent, const T& InDefaultValue, bool bInAllowsInterpolation)
-		: FPCGMetadataAttribute(PCGMetadataHelpers::GetDefaultMetadataDomain(InMetadata), InName, InParent, InDefaultValue, bAllowsInterpolation)
-	{}
-
 	// This constructor is used only during serialization
 	FPCGMetadataAttribute()
+		: FPCGMetadataAttributeGeneric(PCG::Private::MakeAttributeDesc<T>(NAME_None), nullptr, false)
 	{
-		TypeId = PCG::Private::MetadataTypes<T>::Id;
+		static_assert(PCG::Private::IsPCGType<T>(), "Only supports legacy types, such as int32 or float");
 	}
 
 	virtual void Serialize(FPCGMetadataDomain* InMetadata, FArchive& InArchive) override
 	{
-		FPCGMetadataAttributeBase::Serialize(InMetadata, InArchive);
-
-		InArchive << Values;
-		InArchive << DefaultValue;
+		InArchive.UsingCustomVersion(FFortniteMainBranchObjectVersion::GUID);
 		
-		// Initialize non-serialized members
-		if (InArchive.IsLoading())
+		if (InArchive.IsLoading() && InArchive.CustomVer(FFortniteMainBranchObjectVersion::GUID) < FFortniteMainBranchObjectVersion::ConvertFPCGMetadataAttributeToGenericAttributes)
 		{
-			ValueKeyOffset = GetParent() ? GetParent()->GetValueKeyOffsetForChild() : 0;
-		}
-	}
-
-	virtual void GetResourceSizeEx(FResourceSizeEx& CumulativeResourceSize) const override
-	{
-		PCG::TUniqueScopeLock ReadLock(ValueLock);
-		FArchiveCountMem Archive(nullptr);
-		const_cast<FPCGMetadataAttribute<T>*>(this)->Serialize(Metadata, Archive);
+			// Need to forward all the values coming from the previous serialization, and follow the same order of initialization
+			FPCGMetadataAttributeBase::Serialize(InMetadata, InArchive);
 
-		CumulativeResourceSize.AddDedicatedSystemMemoryBytes(Archive.GetNum());
-	}
+			TArray<T> PreviousValues;
+			T PreviousDefaultValue{};
+			InArchive << PreviousValues;
+			InArchive << PreviousDefaultValue;
 
-	virtual void Flatten() override
-	{
-		// Implementation notes:
-		// We don't need to flatten the EntryToValueKeyMap - this will have been taken care of in the metadata flatten
-
-		// Flatten values, from root to current attribute
-		if (Parent)
-		{
-			PCG::TUniqueScopeLock WriteLock(ValueLock);
-
-			TArray<const TArray<T>*> OriginalValues;
-			int32 ValueCount = 0;
-
-			const FPCGMetadataAttribute<T>* Current = static_cast<const FPCGMetadataAttribute<T>*>(this);
-			while (Current)
-			{
-				ValueCount += Current->Values.Num();
-				OriginalValues.Add(&Current->Values);
-				Current = static_cast<const FPCGMetadataAttribute<T>*>(Current->Parent);
-			}
-
-			TArray<T> FlattenedValues;
-			FlattenedValues.Reserve(ValueCount);
-
-			for (int32 ValuesIndex = OriginalValues.Num() - 1; ValuesIndex >= 0; --ValuesIndex)
-			{
-				FlattenedValues.Append(*OriginalValues[ValuesIndex]);
-			}
-
-			Values = MoveTemp(FlattenedValues);
-		}
-		
-		// Reset value offset, and lose parent
-		ValueKeyOffset = 0;
-		Parent = nullptr;
-	}
-
-	virtual void FlattenAndCompress(const TArrayView<const PCGMetadataEntryKey>& InEntryKeysToKeep) override
-	{
-		// No entries, we can just delete everything in the attribute.
-		if (InEntryKeysToKeep.IsEmpty())
-		{
-			Reset();
-			return;
-		}
-
-		TArray<PCGMetadataValueKey> AllValueKeys;
-		TArray<PCGMetadataValueKey> AllUniqueValueKeys;
-		AllValueKeys.Reserve(InEntryKeysToKeep.Num());
-		constexpr bool bUseValueKeys = PCG::Private::MetadataTraits<T>::CompressData;
-		if constexpr (bUseValueKeys)
-		{
-			AllUniqueValueKeys.Reserve(InEntryKeysToKeep.Num());
-		}
-
-		// First gather all value keys associated with the entry keys to keep.
-		// If we compress data, we also store the unique value keys used (that is not default).
-		for (const PCGMetadataEntryKey& EntryKey : InEntryKeysToKeep)
-		{
-			AllValueKeys.Add(GetValueKey(EntryKey));
-			if constexpr (bUseValueKeys)
-			{
-				if (AllValueKeys.Last() != PCGDefaultValueKey)
-				{
-					AllUniqueValueKeys.AddUnique(AllValueKeys.Last());
-				}
-			}
-		}
-
-		// Then for each value key (or unique values keys), gather the value in a NewValues array
-		// and also keep a mapping between old value key and new value key.
-		// Only done if the old value key is not the default one.
-		TMap<PCGMetadataValueKey, PCGMetadataValueKey> ValueKeyMapping;
-		TArray<T> NewValues;
-		const TArray<PCGMetadataValueKey>& AllValueKeysRef = !bUseValueKeys ? AllValueKeys : AllUniqueValueKeys;
-		NewValues.Reserve(AllValueKeysRef.Num());
-		ValueKeyMapping.Reserve(AllValueKeysRef.Num());
-		for (PCGMetadataValueKey ValueKey : AllValueKeysRef)
-		{
-			if (ValueKey != PCGDefaultValueKey)
-			{
-				ValueKeyMapping.Add(ValueKey, NewValues.Num());
-				NewValues.Add(GetValue(ValueKey));
-			}
-		}
-
-		// Move the new values in place of the old values.
-		{
-			PCG::TUniqueScopeLock WriteLock(ValueLock);
-			Values = std::move(NewValues);
-		}
-
-		// And finally, create a new entry to value mapping.
-		// Logic is that each entry to keep will have their "index" as new entry key
-		// (like if the entries to keep are [25, 47, 54], the new entries would be [0, 1, 2]).
-		// So the operation is:
-		// All pairs Old EK -> Old VK transform to New EK -> New VK.
-		TMap<PCGMetadataEntryKey, PCGMetadataValueKey> NewMap;
-		for (int32 i = 0; i < InEntryKeysToKeep.Num(); ++i)
-		{
-			PCGMetadataValueKey ValueKey = AllValueKeys[i];
-			if (ValueKey != PCGDefaultValueKey && ensure(InEntryKeysToKeep[i] != PCGInvalidEntryKey))
-			{
-				NewMap.Add(i, ValueKeyMapping[ValueKey]);
-			}
-		}
-
-		// And move the map.
-		{
-			PCG::TUniqueScopeLock WriteLock(EntryMapLock);
-			EntryToValueKeyMap = std::move(NewMap);
+			InitForDeprecation<T>(MoveTemp(PreviousValues), PreviousDefaultValue);
 		}
-
-		// At the end, reset value offset, and lose parent.
-		ValueKeyOffset = 0;
-		Parent = nullptr;
-	}
-
-	virtual void Reset() override
-	{
-		ValueKeyOffset = 0;
-		Parent = nullptr;
-
-		{
-			PCG::TUniqueScopeLock WriteLock(EntryMapLock);
-			EntryToValueKeyMap.Empty();
-		}
-
+		else
 		{
-			PCG::TUniqueScopeLock WriteLock(ValueLock);
-			Values.Empty();
+			FPCGMetadataAttributeGeneric::Serialize(InMetadata, InArchive);
 		}
 	}
 
-	const FPCGMetadataAttribute* GetParent() const { return static_cast<const FPCGMetadataAttribute*>(Parent); }
-
-	UE_DEPRECATED(5.6, "Use the version with FPCGMetadataDomain")
-	FPCGMetadataAttribute* TypedCopy(FName NewName, UPCGMetadata* InMetadata, bool bKeepParent, bool bCopyEntries = true, bool bCopyValues = true)
-	{
-		return static_cast<FPCGMetadataAttribute*>(TypedCopy(NewName, PCGMetadataHelpers::GetDefaultMetadataDomain(InMetadata), bKeepParent, bCopyEntries, bCopyValues));
-	}
-	
 	FPCGMetadataAttribute* TypedCopy(FName NewName, FPCGMetadataDomain* InMetadata, bool bKeepParent, bool bCopyEntries = true, bool bCopyValues = true)
 	{
 		return static_cast<FPCGMetadataAttribute*>(Copy(NewName, InMetadata, bKeepParent, bCopyEntries, bCopyValues));
 	}
-
-	UE_DEPRECATED(5.6, "Use the version with FPCGMetadataDomain")
-	virtual FPCGMetadataAttribute* Copy(FName NewName, UPCGMetadata* InMetadata, bool bKeepParent, bool bCopyEntries = true, bool bCopyValues = true) const override
-	{
-		return static_cast<FPCGMetadataAttribute*>(Copy(NewName, PCGMetadataHelpers::GetDefaultMetadataDomain(InMetadata), bKeepParent, bCopyEntries, bCopyValues));
-	}
 	
 	virtual FPCGMetadataAttributeBase* Copy(FName NewName, FPCGMetadataDomain* InMetadata, bool bKeepParent, bool bCopyEntries = true, bool bCopyValues = true) const override
 	{
@@ -233,13 +79,6 @@ public:
 
 	virtual FPCGMetadataAttributeBase* CopyToAnotherType(int16 TargetType) const override;
 
-	template<typename U>
-	UE_DEPRECATED(5.6, "Do not use anymore, use CopyToAnotherType")
-	FPCGMetadataAttributeBase* CopyInternal(FName NewName, UPCGMetadata* InMetadata, bool bKeepParent, bool bCopyEntries, bool bCopyValues) const
-	{
-		return CopyInternal<U>(NewName, PCGMetadataHelpers::GetDefaultMetadataDomain(InMetadata), bKeepParent, bCopyEntries, bCopyValues);
-	}
-
 protected:
 	// TODO: add enable if only on compatible types, but this has some repercussion on using metadata on types that aren't normally supported.
 	// TODO: The UPCGMetadata version was exposed, but should have not, mark the FPCGMetadataDomain version protected. Will be moved at the end, before commit.
@@ -258,176 +97,170 @@ protected:
 			UE_LOG(LogPCG, Error, TEXT("Try to create a new attribute with an invalid name: %s"), *NewName.ToString());
 			return nullptr;
 		}
-
-		U NewDefaultValue{};
-		PCG::Private::GetValueWithBroadcastAndConstructible(DefaultValue, NewDefaultValue);
 		
 		// This copies to a new attribute.
-		FPCGMetadataAttribute<U>* AttributeCopy = new FPCGMetadataAttribute<U>(InMetadata, NewName, bKeepParent ? this : nullptr, NewDefaultValue, bAllowsInterpolation);
-
-		// Gather the chain of parents if we don't keep the parent and we want to copy entries/values.
-		// We always have at least one item, "this".
-		TArray<const FPCGMetadataAttribute<T>*, TInlineAllocator<2>> Parents = { this };
-		if (!bKeepParent && (bCopyEntries || bCopyValues))
-		{
-			const FPCGMetadataDomain* CurrentMetadata = Metadata;
-			const FPCGMetadataAttribute<T>* Current = this;
-
-			const FPCGMetadataDomain* ParentMetadata = PCGMetadataHelpers::GetParentMetadata(CurrentMetadata);
-			while(ParentMetadata && Current->Parent)
-			{
-				CurrentMetadata = ParentMetadata;
-				Current = static_cast<const FPCGMetadataAttribute<T>*>(Current->Parent);
-				Parents.Add(Current);
-
-				ParentMetadata = PCGMetadataHelpers::GetParentMetadata(CurrentMetadata);
-			}
-		}
-
-		if (bCopyEntries)
+		FPCGMetadataAttribute<U>* AttributeCopy = new FPCGMetadataAttribute<U>(InMetadata, NewName, bKeepParent ? this : nullptr, U{}, bAllowsInterpolation);
+		
+		auto CopyFunctor = [](void* DestData, const void* SrcData, const int32 Count)
 		{
-			// We go backwards, since we need to preserve order (root -> this)
-			// Latest entry in our Parents array is the root.
-			for (int32 i = Parents.Num() - 1; i >= 0; --i)
+			if constexpr (std::is_same_v<T, U> && std::is_trivially_copyable_v<T>)
 			{
-				const FPCGMetadataAttribute<T>* Current = Parents[i];
-				
-				PCG::TSharedScopeLock ReadLock(Current->EntryMapLock);
-				AttributeCopy->EntryToValueKeyMap.Append(Current->EntryToValueKeyMap);
+				FMemory::Memcpy(DestData, SrcData, Count * sizeof(T));
 			}
-		}
-
-		if (bCopyValues)
-		{
-			// We go backwards, since we need to preserve order (root -> this)
-			// Latest entry in our Parents array is the root.
-			for (int32 i = Parents.Num() - 1; i >= 0; --i)
+			else if constexpr (PCG::Private::IsBroadcastableOrConstructible<T, U>())
 			{
-				const FPCGMetadataAttribute<T>* Current = Parents[i];
+				TArrayView<U> DestDataView(static_cast<U*>(DestData), Count);
+				TConstArrayView<T> SrcDataView(static_cast<const T*>(SrcData), Count);
 				
-				PCG::TSharedScopeLock ReadLock(Current->ValueLock);
-
-				if constexpr (std::is_same_v<T, U>)
+				for (int32 i = 0; i < Count; ++i)
 				{
-					AttributeCopy->Values.Append(Current->Values);
-				}
-				else
-				{
-					U NewValue{};
-
-					for (const T& Value : Current->Values)
-					{
-						PCG::Private::GetValueWithBroadcastAndConstructible(Value, NewValue);
-						AttributeCopy->Values.Add(NewValue);
-					}
-				}
-
-				// The expected value key offset is the one for this attribute (i == 0), and only if we
-				// keep the parent. Otherwise we don't have any parent, so offset should be kept at 0.
-				if (i == 0 && bKeepParent)
-				{
-					AttributeCopy->ValueKeyOffset = Current->ValueKeyOffset;
+					PCG::Private::GetValueWithBroadcastAndConstructible(SrcDataView[i], DestDataView[i]);
 				}
 			}
-		}
+			else
+			{
+				static_assert(!std::is_same_v<T, T>, "Copy Internal should never be called with a pair of types that are incompatible.");
+			}
+		};
+		
+		FPCGMetadataAttributeGeneric::CopyInternal(AttributeCopy, bKeepParent, bCopyEntries, bCopyValues, CopyFunctor);
 
 		return AttributeCopy;
 	}
 
 public:
-	virtual void SetValue(PCGMetadataEntryKey ItemKey, const FPCGMetadataAttributeBase* InAttribute, PCGMetadataEntryKey InEntryKey) override
-	{
-		check(ItemKey != PCGInvalidEntryKey);
-		if (InAttribute == this)
-		{
-			SetValueFromValueKey(ItemKey, GetValueKey(InEntryKey));
-		}
-		else if (InAttribute)
-		{
-			SetValue(ItemKey, static_cast<const FPCGMetadataAttribute<T>*>(InAttribute)->GetValueFromItemKey(InEntryKey));
-		}
-	}
-
 	virtual void SetZeroValue(PCGMetadataEntryKey ItemKey) override
 	{
 		check(ItemKey != PCGInvalidEntryKey);
 		ZeroValue(ItemKey);
 	}
-
+	
 	virtual void AccumulateValue(PCGMetadataEntryKey ItemKey, const FPCGMetadataAttributeBase* InAttribute, PCGMetadataEntryKey InEntryKey, float Weight) override
 	{
-		check(ItemKey != PCGInvalidEntryKey);
-		Accumulate(ItemKey, InAttribute, InEntryKey, Weight);
+		if constexpr (PCG::Private::MetadataTraits<T>::CanInterpolate)
+		{
+			check(ItemKey != PCGInvalidEntryKey);
+			
+			if (InAttribute && !InAttribute->IsGeneric())
+			{
+				return;
+			}
+		
+			const FPCGMetadataAttributeGeneric* Attribute = static_cast<const FPCGMetadataAttributeGeneric*>(InAttribute);
+		
+			if (Attribute && !Attribute->IsOfType<T>())
+			{
+				return;
+			}
+		
+			Accumulate(ItemKey, Attribute, InEntryKey, Weight);
+		}
 	}
 
 	virtual void SetWeightedValue(PCGMetadataEntryKey ItemKey, const FPCGMetadataAttributeBase* InAttribute, const TArrayView<const TPair<PCGMetadataEntryKey, float>>& InWeightedKeys) override
 	{
-		check(ItemKey != PCGInvalidEntryKey);
-		Accumulate(ItemKey, InAttribute, InWeightedKeys);
+		if constexpr (PCG::Private::MetadataTraits<T>::CanInterpolate)
+		{
+			check(ItemKey != PCGInvalidEntryKey);
+			
+			if (InAttribute && !InAttribute->IsGeneric())
+			{
+				return;
+			}
+		
+			const FPCGMetadataAttributeGeneric* Attribute = static_cast<const FPCGMetadataAttributeGeneric*>(InAttribute);
+		
+			if (Attribute && !Attribute->IsOfType<T>())
+			{
+				return;
+			}
+		
+			Accumulate(ItemKey, Attribute, InWeightedKeys);
+		}
+	}
+	
+	
+	//~ Note that this override is purely to disambiguate between all the different versions of SetValue. It just forward the call to the super method.
+	virtual void SetValue(PCGMetadataEntryKey ItemKey, const FPCGMetadataAttributeBase* InAttribute, PCGMetadataEntryKey InEntryKey) override
+	{
+		return FPCGMetadataAttributeGeneric::SetValue(ItemKey, InAttribute, InEntryKey);
 	}
 
 	virtual void SetValue(PCGMetadataEntryKey ItemKey, const FPCGMetadataAttributeBase* InAttributeA, PCGMetadataEntryKey InEntryKeyA, const FPCGMetadataAttributeBase* InAttributeB, PCGMetadataEntryKey InEntryKeyB, EPCGMetadataOp Op) override
 	{
 		check(ItemKey != PCGInvalidEntryKey);
 		bool bAppliedValue = false;
-
-		if (Op == EPCGMetadataOp::TargetValue && InAttributeB)
+		
+		if ((InAttributeA && !InAttributeA->IsGeneric()) || (InAttributeB && !InAttributeB->IsGeneric()))
+		{
+			return;
+		}
+		
+		const FPCGMetadataAttributeGeneric* AttributeA = static_cast<const FPCGMetadataAttributeGeneric*>(InAttributeA);
+		const FPCGMetadataAttributeGeneric* AttributeB = static_cast<const FPCGMetadataAttributeGeneric*>(InAttributeB);
+		
+		if ((AttributeA && !AttributeA->IsOfType<T>()) || (AttributeB && !AttributeB->IsOfType<T>()))
+		{
+			return;
+		}
+		
+		if (Op == EPCGMetadataOp::TargetValue && AttributeB)
 		{
 			// Take value of second attribute.
-			if (InAttributeB == this)
+			if (AttributeB == this)
 			{
 				SetValueFromValueKey(ItemKey, GetValueKey(InEntryKeyB));
 			}
 			else
 			{
-				SetValue(ItemKey, static_cast<const FPCGMetadataAttribute<T>*>(InAttributeB)->GetValueFromItemKey(InEntryKeyB));
+				SetValue(ItemKey, AttributeB->GetValueFromItemKey<T>(InEntryKeyB));
 			}
 
 			bAppliedValue = true;
 		}
-		else if (Op == EPCGMetadataOp::SourceValue && InAttributeA)
+		else if (Op == EPCGMetadataOp::SourceValue && AttributeA)
 		{
 			// Take value of first attribute.
-			if (InAttributeA == this)
+			if (AttributeA == this)
 			{
 				SetValueFromValueKey(ItemKey, GetValueKey(InEntryKeyA));
 			}
 			else
 			{
-				SetValue(ItemKey, static_cast<const FPCGMetadataAttribute<T>*>(InAttributeA)->GetValueFromItemKey(InEntryKeyA));
+				SetValue(ItemKey, AttributeA->GetValueFromItemKey<T>(InEntryKeyA));
 			}
 
 			bAppliedValue = true;
 		}
-		else if (InAttributeA && InAttributeB && bAllowsInterpolation)
+		else if (AttributeA && AttributeB && bAllowsInterpolation)
 		{
 			// Combine attributes using specified operation.
 			if (Op == EPCGMetadataOp::Min)
 			{
-				bAppliedValue = SetMin(ItemKey, InAttributeA, InEntryKeyA, InAttributeB, InEntryKeyB);
+				bAppliedValue = SetMin(ItemKey, AttributeA, InEntryKeyA, AttributeB, InEntryKeyB);
 			}
 			else if (Op == EPCGMetadataOp::Max)
 			{
-				bAppliedValue = SetMax(ItemKey, InAttributeA, InEntryKeyA, InAttributeB, InEntryKeyB);
+				bAppliedValue = SetMax(ItemKey, AttributeA, InEntryKeyA, AttributeB, InEntryKeyB);
 			}
 			else if (Op == EPCGMetadataOp::Sub)
 			{
-				bAppliedValue = SetSub(ItemKey, InAttributeA, InEntryKeyA, InAttributeB, InEntryKeyB);
+				bAppliedValue = SetSub(ItemKey, AttributeA, InEntryKeyA, AttributeB, InEntryKeyB);
 			}
 			else if (Op == EPCGMetadataOp::Add)
 			{
-				bAppliedValue = SetAdd(ItemKey, InAttributeA, InEntryKeyA, InAttributeB, InEntryKeyB);
+				bAppliedValue = SetAdd(ItemKey, AttributeA, InEntryKeyA, AttributeB, InEntryKeyB);
 			}
 			else if (Op == EPCGMetadataOp::Mul)
 			{
-				bAppliedValue = SetMul(ItemKey, InAttributeA, InEntryKeyA, InAttributeB, InEntryKeyB);
+				bAppliedValue = SetMul(ItemKey, AttributeA, InEntryKeyA, AttributeB, InEntryKeyB);
 			}
 			else if (Op == EPCGMetadataOp::Div)
 			{
-				bAppliedValue = SetDiv(ItemKey, InAttributeA, InEntryKeyA, InAttributeB, InEntryKeyB);
+				bAppliedValue = SetDiv(ItemKey, AttributeA, InEntryKeyA, AttributeB, InEntryKeyB);
 			}
 		}
-		else if (InAttributeA && InAttributeB && HasNonDefaultValue(ItemKey))
+		else if (AttributeA && AttributeB && HasNonDefaultValue(ItemKey))
 		{
 			// In this case, the current already has a value, in which case we should not update it
 			bAppliedValue = true;
@@ -437,279 +270,92 @@ public:
 		{
 			// Nothing to do
 		}
-		else if (InAttributeA)
+		else if (AttributeA)
 		{
-			if (InAttributeA == this)
+			if (AttributeA == this)
 			{
 				SetValueFromValueKey(ItemKey, GetValueKey(InEntryKeyA));
 			}
 			else
 			{
-				SetValue(ItemKey, static_cast<const FPCGMetadataAttribute<T>*>(InAttributeA)->GetValueFromItemKey(InEntryKeyA));
+				SetValue(ItemKey, AttributeA->GetValueFromItemKey<T>(InEntryKeyA));
 			}
 		}
-		else if (InAttributeB)
+		else if (AttributeB)
 		{
-			if (InAttributeB == this)
+			if (AttributeB == this)
 			{
 				SetValueFromValueKey(ItemKey, GetValueKey(InEntryKeyB));
 			}
 			else
 			{
-				SetValue(ItemKey, static_cast<const FPCGMetadataAttribute<T>*>(InAttributeB)->GetValueFromItemKey(InEntryKeyB));
+				SetValue(ItemKey, AttributeB->GetValueFromItemKey<T>(InEntryKeyB));
 			}
 		}
 	}
 
-	virtual bool AreValuesEqualForEntryKeys(PCGMetadataEntryKey EntryKey1, PCGMetadataEntryKey EntryKey2) const override
-	{
-		return AreValuesEqual(GetValueKey(EntryKey1), GetValueKey(EntryKey2));
-	}
-
-	virtual bool AreValuesEqual(PCGMetadataValueKey ValueKey1, PCGMetadataValueKey ValueKey2) const override
-	{
-		if constexpr (PCG::Private::MetadataTraits<T>::CompressData)
-		{
-			if (ValueKey1 == PCGInvalidEntryKey)
-			{
-				return PCG::Private::MetadataTraits<T>::Equal(DefaultValue, GetValue(ValueKey2));
-			}
-			else if (ValueKey2 == PCGInvalidEntryKey)
-			{
-				return PCG::Private::MetadataTraits<T>::Equal(GetValue(ValueKey1), DefaultValue);
-			}
-			else
-			{
-				return ValueKey1 == ValueKey2;
-			}
-		}
-		else
-		{
-			return ValueKey1 == ValueKey2 || PCG::Private::MetadataTraits<T>::Equal(GetValue(ValueKey1), GetValue(ValueKey2));
-		}
-	}
-
-	virtual bool IsEqualToDefaultValue(PCGMetadataValueKey ValueKey) const override
-	{
-		return PCG::Private::MetadataTraits<T>::Equal(GetValue(ValueKey), DefaultValue);
-	}
-
-	virtual void SetDefaultValueToFirstEntry() override
-	{
-		if (Values.Num() != 1)
-		{
-			return;
-		}
-
-		DefaultValue = Values[0];
-	}
-
-	PCGMetadataValueKey GetValueKeyOffsetForChild() const
-	{
-		PCG::TSharedScopeLock ReadLock(ValueLock);
-		return Values.Num() + ValueKeyOffset;
-	}
-
 	/** Adds the value, returns the value key for the given value */
 	PCGMetadataValueKey AddValue(const T& InValue)
 	{
-		PCGMetadataValueKey FoundValueKey = FindValue(InValue);
-
-		if (FoundValueKey == PCGNotFoundValueKey)
-		{
-			PCG::TUniqueScopeLock WriteLock(ValueLock);
-			return Values.Add(InValue) + ValueKeyOffset;
-		}
-		else
-		{
-			return FoundValueKey;
-		}
-	}
-
-	template<typename IT = T, typename TEnableIf<PCG::Private::MetadataTraits<IT>::CompressData>::Type* = nullptr>
-	TArray<PCGMetadataValueKey> AddValues(const TArrayView<const T>& InValues)
-	{
-		// Since we're getting raw values here, we might have duplicates
-		// so we should aim to remove duplicates here so we preserve our 'compress data' idea, otherwise it will break other foundational blocks (e.g. partition)
-		TArray<T, TInlineAllocator<256>> UniqueValues;
-
-		// Initially, fill with mapping to unique values so we can remap them at the end if needed
-		TArray<PCGMetadataValueKey> FoundValueKeys;
-		FoundValueKeys.Reserve(InValues.Num());
-
-		for (int ValueIndex = 0; ValueIndex < InValues.Num(); ++ValueIndex)
-		{
-			FoundValueKeys.Emplace(UniqueValues.AddUnique(InValues[ValueIndex]));
-		}
-
-		const bool bHasDuplicateValues = (UniqueValues.Num() != InValues.Num());
-
-		TArray<PCGMetadataValueKey> FoundUniqueValueKeys;
-		TArray<PCGMetadataValueKey>& FoundKeys = (bHasDuplicateValues ? FoundUniqueValueKeys : FoundValueKeys);
-
-		// Implementation note: when we don't have any duplicate values, the previously-set values in FoundValueKeys will be wiped out - this is intended
-		const bool bAtLeastOneValueNotFound = !FindValues(UniqueValues, FoundKeys);
-
-		if (bAtLeastOneValueNotFound)
-		{
-			PCG::TUniqueScopeLock WriteLock(ValueLock);
-			for (int ValueIndex = 0; ValueIndex < UniqueValues.Num(); ++ValueIndex)
-			{
-				if (FoundKeys[ValueIndex] == PCGNotFoundValueKey)
-				{
-					FoundKeys[ValueIndex] = Values.Add(UniqueValues[ValueIndex]) + ValueKeyOffset;
-				}
-			}
-		}
-
-		// Remap to full array if needed
-		if (bHasDuplicateValues)
-		{
-			for (PCGMetadataValueKey& ValueToRemap : FoundValueKeys)
-			{
-				ValueToRemap = FoundUniqueValueKeys[ValueToRemap];
-			}
-		}
-
-		return FoundValueKeys;
+		TArray<PCGMetadataValueKey> Result = FPCGMetadataAttributeGeneric::AddValues<T>(MakeArrayView<const T>(&InValue, 1));
+		return !Result.IsEmpty() ? Result[0] : PCGDefaultValueKey; 
 	}
 
-	template<typename IT = T, typename TEnableIf<!PCG::Private::MetadataTraits<IT>::CompressData>::Type* = nullptr>
 	TArray<PCGMetadataValueKey> AddValues(const TArrayView<const T>& InValues)
 	{
-		TArray<PCGMetadataValueKey> ValueKeys;
-
-		int FirstValueIndex = 0;
-
-		{
-			PCG::TUniqueScopeLock WriteLock(ValueLock);
-			FirstValueIndex = Values.Num() + ValueKeyOffset;
-			Values.Append(InValues);
-		}
-
-		ValueKeys.SetNum(InValues.Num());
-		for (int ValueIndex = 0; ValueIndex < InValues.Num(); ++ValueIndex)
-		{
-			ValueKeys[ValueIndex] = FirstValueIndex + ValueIndex;
-		}
-
-		return ValueKeys;
+		return FPCGMetadataAttributeGeneric::AddValues<T>(InValues);
 	}
 
 	void SetValue(PCGMetadataEntryKey ItemKey, const T& InValue)
 	{
-		check(ItemKey != PCGInvalidEntryKey);
-		SetValueFromValueKey(ItemKey, AddValue(InValue));
+		FPCGMetadataAttributeGeneric::SetValue<T>(ItemKey, InValue);
 	}
 
 	void SetValues(const TArrayView<const PCGMetadataEntryKey>& ItemKeys, const TArrayView<const T>& InValues)
 	{
-		SetValuesFromValueKeys(ItemKeys, AddValues(InValues));
+		FPCGMetadataAttributeGeneric::SetValues<T>(ItemKeys, InValues);
 	}
 
 	void SetValues(const TArrayView<const PCGMetadataEntryKey * const>& ItemKeys, const TArrayView<const T>& InValues)
 	{
-		SetValuesFromValueKeys(ItemKeys, AddValues(InValues));
+		FPCGMetadataAttributeGeneric::SetValues<T>(ItemKeys, InValues);
 	}
 	
 	void SetValues(TConstPCGValueRange<PCGMetadataEntryKey> ItemKeys, const TArrayView<const T>& InValues)
 	{
-		SetValuesFromValueKeys(ItemKeys, AddValues(InValues));
+		FPCGMetadataAttributeGeneric::SetValues<T>(ItemKeys, InValues);
 	}
 	
 	void SetValues(TPCGValueRange<PCGMetadataEntryKey> ItemKeys, const TArrayView<const T>& InValues)
 	{
-		SetValuesFromValueKeys(PCGValueRangeHelpers::MakeConstValueRange(ItemKeys), AddValues(InValues));
+		FPCGMetadataAttributeGeneric::SetValues<T>(ItemKeys, InValues);
 	}
 
-	template<typename U>
+	template<typename U> requires (!std::is_same_v<T, U> && PCG::Private::IsBroadcastableOrConstructible<U, T>())
 	void SetValue(PCGMetadataEntryKey ItemKey, const U& InValue)
 	{
-		check(ItemKey != PCGInvalidEntryKey);
-		SetValueFromValueKey(ItemKey, AddValue(T(InValue)));
+		T Value{};
+		PCG::Private::GetValueWithBroadcastAndConstructible(InValue, Value);
+		FPCGMetadataAttributeGeneric::SetValue<T>(ItemKey, Value);
+	}
+	
+	// Special version for enum classes, as enums are not directly constructible
+	template<typename U> requires (!std::is_same_v<T, U> && !PCG::Private::IsBroadcastableOrConstructible<U, T>() && std::is_enum_v<U> && std::is_convertible_v<std::underlying_type_t<U>, T>)
+	void SetValue(PCGMetadataEntryKey ItemKey, const U& InValue)
+	{
+		FPCGMetadataAttributeGeneric::SetValue<T>(ItemKey, static_cast<T>(InValue));
 	}
 
 	T GetValueFromItemKey(PCGMetadataEntryKey ItemKey) const
 	{
-		return GetValue(GetValueKey(ItemKey));
+		return FPCGMetadataAttributeGeneric::GetValueFromItemKey<T>(ItemKey);
 	}
 
 	T GetValue(PCGMetadataValueKey ValueKey) const
 	{
-		if (ValueKey == PCGDefaultValueKey)
-		{
-			return DefaultValue;
-		}
-		else if (ValueKey >= ValueKeyOffset)
-		{
-			int32 Index = ValueKey - ValueKeyOffset;
-			PCG::TSharedScopeLock ReadLock(ValueLock);
-			return Index < Values.Num() ? Values[Index] : DefaultValue;
-		}
-		else if (GetParent())
-		{
-			return GetParent()->GetValue(ValueKey);
-		}
-		else
-		{
-			return DefaultValue;
-		}
-	}
-
-protected:
-	void GetValues_Internal(const TArrayView<const PCGMetadataValueKey> ValueKeys, TArrayView<T> OutValues, TBitArray<>& UnretrievedValues) const
-	{
-		check(ValueKeys.Num() == OutValues.Num() && OutValues.Num() == UnretrievedValues.Num());
-
-		bool bFoundAllKeys = true;
-		TConstSetBitIterator<> It(UnretrievedValues);
-		if (!It)
-		{
-			return;
-		}
-
-		const FPCGMetadataAttribute* ThisParent = GetParent();
-
-		{
-			PCG::TSharedScopeLock ReadLock(ValueLock);
-
-			for (; It; ++It)
-			{
-				const int32 Index = It.GetIndex();
-				const PCGMetadataValueKey ValueKey = ValueKeys[Index];
-
-				auto RetrieveValue = [Index, &OutValues, &UnretrievedValues](T Value)
-					{
-						OutValues[Index] = std::move(Value);
-						UnretrievedValues[Index] = false;
-					};
-
-				if (ValueKey == PCGDefaultValueKey)
-				{
-					RetrieveValue(DefaultValue);
-				}
-				else if (ValueKey >= ValueKeyOffset)
-				{
-					int32 ValueIndex = ValueKey - ValueKeyOffset;
-					RetrieveValue(ValueIndex < Values.Num() ? Values[ValueIndex] : DefaultValue);
-				}
-				else if (!ThisParent)
-				{
-					RetrieveValue(DefaultValue);
-				}
-				else
-				{
-					bFoundAllKeys = false;
-				}
-			}
-		}
-
-		ensure(ThisParent || bFoundAllKeys);
-
-		if (ThisParent && !bFoundAllKeys)
-		{
-			ThisParent->GetValues_Internal(ValueKeys, OutValues, UnretrievedValues);
-		}
+		T Value{};
+		FPCGMetadataAttributeGeneric::GetValues<T>(MakeConstArrayView(&ValueKey, 1), MakeArrayView(&Value, 1));
+		return Value;
 	}
 
 public:
@@ -718,9 +364,7 @@ public:
 	*/
 	void GetValues(const TArrayView<const PCGMetadataValueKey> ValueKeys, TArrayView<T> OutValues) const
 	{
-		// Bitset with all unretrieved values. If we have any unretrieved value, we will ask the parent for those.
-		TBitArray<> UnretrievedValues(true, ValueKeys.Num());
-		GetValues_Internal(ValueKeys, OutValues, UnretrievedValues);
+		return FPCGMetadataAttributeGeneric::GetValues<T>(ValueKeys, OutValues);
 	}
 
 	/** 
@@ -731,19 +375,12 @@ public:
 	*/
 	void GetValuesFromItemKeys(TConstArrayView<PCGMetadataEntryKey> EntryKeys, TArrayView<T> OutValues) const
 	{
-		GetValuesFromItemKeys(PCGValueRangeHelpers::MakeConstValueRange(EntryKeys), OutValues);
+		FPCGMetadataAttributeGeneric::GetValuesFromItemKeys<T>(EntryKeys, OutValues);
 	}
 	
 	void GetValuesFromItemKeys(TConstPCGValueRange<PCGMetadataEntryKey> EntryKeys, TArrayView<T> OutValues) const
 	{
-		if (!ensure(EntryKeys.Num() == OutValues.Num()))
-		{
-			return;
-		}
-
-		TArray<PCGMetadataValueKey> ValueKeys;
-		GetValueKeys(EntryKeys, ValueKeys);
-		GetValues(ValueKeys, OutValues);
+		FPCGMetadataAttributeGeneric::GetValuesFromItemKeys<T>(EntryKeys, OutValues);
 	}
 
 	/** 
@@ -753,386 +390,194 @@ public:
 	*/
 	void GetValuesFromItemKeys(TArrayView<PCGMetadataEntryKey> EntryKeys, TArrayView<T> OutValues) const
 	{
-		if (!ensure(EntryKeys.Num() == OutValues.Num()))
-		{
-			return;
-		}
-
-		TArray<PCGMetadataValueKey> ValueKeys;
-		GetValueKeys(EntryKeys, ValueKeys);
-		GetValues(ValueKeys, OutValues);
-	}
-
-	/** Code related to finding values / compressing data */
-	virtual bool UsesValueKeys() const override
-	{
-		return PCG::Private::MetadataTraits<T>::CompressData;
+		FPCGMetadataAttributeGeneric::GetValuesFromItemKeys<T>(EntryKeys, OutValues);
 	}
-
-	template<typename IT = T, typename TEnableIf<PCG::Private::MetadataTraits<IT>::CompressData>::Type* = nullptr>
-	PCGMetadataValueKey FindValue(const T& InValue) const
+	
+	PCGMetadataValueKey FindValue(const T& InValue) const 
 	{
-		if (InValue == DefaultValue)
-		{
-			return PCGDefaultValueKey;
-		}
-
-		PCGMetadataValueKey ParentValueKey = (GetParent() ? GetParent()->FindValue(InValue) : PCGNotFoundValueKey);
-		if (ParentValueKey != PCGNotFoundValueKey)
+		if constexpr (PCG::Private::MetadataTraits<T>::CompressData)
 		{
-			return ParentValueKey;
+			TArray<PCGMetadataValueKey> OutValueKeys;
+			const bool bSuccess = FindValues(MakeConstArrayView(&InValue, 1), OutValueKeys);
+			return bSuccess && !OutValueKeys.IsEmpty() ? OutValueKeys[0] : PCGDefaultValueKey;
 		}
 		else
 		{
-			int32 ValueIndex = INDEX_NONE;
-			{
-				PCG::TSharedScopeLock ReadLock(ValueLock);
-				ValueIndex = Values.FindLast(InValue);
-			}
-
-			if (ValueIndex != INDEX_NONE)
-			{
-				return ValueIndex + ValueKeyOffset;
-			}
-			else
-			{
-				return PCGNotFoundValueKey;
-			}
+			return PCGNotFoundValueKey;
 		}
 	}
-
-	template<typename IT = T, typename TEnableIf<PCG::Private::MetadataTraits<IT>::CompressData>::Type* = nullptr>
+	
 	bool FindValues(const TArrayView<const T>& InValues, TArray<PCGMetadataValueKey>& OutValueKeys) const
 	{
-		OutValueKeys.Init(PCGNotFoundValueKey, InValues.Num());
-
-		int ValueKeysSet = 0;
-		FindValuesInternal(InValues, OutValueKeys, ValueKeysSet, /*bIsRoot=*/true);
-
-		return (ValueKeysSet == InValues.Num());
-	}
-
-	template<typename IT = T, typename TEnableIf<PCG::Private::MetadataTraits<IT>::CompressData>::Type* = nullptr>
-	void FindValuesInternal(const TArrayView<const T>& InValues, TArray<PCGMetadataValueKey>& ValueKeys, int& ValueKeysSet, bool bIsRoot) const
-	{
-		check(InValues.Num() == ValueKeys.Num());
-
-		if (bIsRoot)
-		{
-			for (int ValueIndex = 0; ValueIndex < InValues.Num(); ++ValueIndex)
-			{
-				if (InValues[ValueIndex] == DefaultValue)
-				{
-					ValueKeys[ValueIndex] = PCGDefaultValueKey;
-					++ValueKeysSet;
-				}
-			}
-		}
-
-		if (ValueKeysSet != InValues.Num() && GetParent())
+		if constexpr (PCG::Private::MetadataTraits<T>::CompressData)
 		{
-			GetParent()->FindValuesInternal(InValues, ValueKeys, ValueKeysSet, /*bRoot=*/false);
+			return FPCGMetadataAttributeGeneric::FindValues(OutValueKeys, PCG::Private::FInValues{TInPlaceType<PCG::Private::FInValuesByValue>{}, InValues.GetData(), InValues.Num()}, InValues.Num());
 		}
-
-		if (ValueKeysSet != InValues.Num())
+		else
 		{
-			PCG::TSharedScopeLock ReadLock(ValueLock);
-			for (int ValueIndex = 0; ValueIndex < InValues.Num(); ++ValueIndex)
-			{
-				if (ValueKeys[ValueIndex] != PCGNotFoundValueKey)
-				{
-					continue;
-				}
-
-				const int32 FoundValueIndex = Values.FindLast(InValues[ValueIndex]);
-				if (FoundValueIndex != INDEX_NONE)
-				{
-					ValueKeys[ValueIndex] = FoundValueIndex + ValueKeyOffset;
-					++ValueKeysSet;
-				}
-			}
+			OutValueKeys.Init(PCGNotFoundValueKey, InValues.Num());
+			return false;
 		}
 	}
 
-	template<typename IT = T, typename TEnableIf<!PCG::Private::MetadataTraits<IT>::CompressData>::Type* = nullptr>
-	PCGMetadataValueKey FindValue(const T& InValue) const
-	{
-		return PCGNotFoundValueKey;
-	}
-
-	template<typename IT = T, typename TEnableIf<!PCG::Private::MetadataTraits<IT>::CompressData>::Type* = nullptr>
-	bool FindValues(const TArrayView<const T>& InValues, TArray<PCGMetadataValueKey>& OutValueKeys) const
+	void SetValues_TryLockless(TArrayView<PCGMetadataEntryKey*> EntryKeys, TArrayView<const T> InValues, int32 StartIndex)
 	{
-		OutValueKeys.Init(PCGNotFoundValueKey, InValues.Num());
-		return false;
+		FPCGMetadataAttributeGeneric::SetValues_TryLockless<T>(EntryKeys, InValues, StartIndex);
 	}
 
-	void SetDefaultValue(const T& Value)
+protected:
+	/** Code related to computing compared values (min, max, sub, add) */
+	bool SetMin(PCGMetadataEntryKey ItemKey, const FPCGMetadataAttributeGeneric* InAttributeA, PCGMetadataEntryKey InEntryKeyA, const FPCGMetadataAttributeGeneric* InAttributeB, PCGMetadataEntryKey InEntryKeyB)
 	{
-		DefaultValue = Value;
-	}
+		if constexpr (PCG::Private::MetadataTraits<T>::CanMinMax)
+		{
+			SetValue(ItemKey, 
+				PCG::Private::MetadataTraits<T>::Min(
+					InAttributeA->GetValueFromItemKey<T>(InEntryKeyA),
+					InAttributeB->GetValueFromItemKey<T>(InEntryKeyB)));
 
-	void Prepare(int32 Count)
-	{
-		EntryToValueKeyMap.Reserve(EntryToValueKeyMap.Num() + Count);
-		if constexpr (!PCG::Private::MetadataTraits<T>::CompressData)
+			return true;
+		}
+		else
 		{
-			Values.Reserve(Values.Num() + Count);
+			return false;
 		}
 	}
-
-	int32 PreallocateValues(TArrayView<PCGMetadataEntryKey*> EntryKeys, bool bLockless)
+	
+	bool SetMax(PCGMetadataEntryKey ItemKey, const FPCGMetadataAttributeGeneric* InAttributeA, PCGMetadataEntryKey InEntryKeyA, const FPCGMetadataAttributeGeneric* InAttributeB, PCGMetadataEntryKey InEntryKeyB)
 	{
-		int32 StartIndex = INDEX_NONE;
-
-		if constexpr (!PCG::Private::MetadataTraits<T>::CompressData)
+		if constexpr (PCG::Private::MetadataTraits<T>::CanMinMax)
 		{
-			if (!bLockless)
-			{
-				ValueLock.Lock();
-			}
-
-			StartIndex = Values.Num();
-
-			if constexpr (std::is_trivially_copyable_v<T>)
-			{
-				Values.SetNumUninitialized(Values.Num() + EntryKeys.Num(), /*AllowShrinking=*/EAllowShrinking::No);
-			}
-			else
-			{
-				Values.SetNum(Values.Num() + EntryKeys.Num(), /*AllowShrinking=*/EAllowShrinking::No);
-			}
-
-			if (!bLockless)
-			{
-				ValueLock.Unlock();
-			}
+			SetValue(ItemKey,
+				PCG::Private::MetadataTraits<T>::Max(
+					InAttributeA->GetValueFromItemKey<T>(InEntryKeyA),
+					InAttributeB->GetValueFromItemKey<T>(InEntryKeyB)));
+			return true;
 		}
-
-		if (!bLockless)
+		else
 		{
-			EntryMapLock.Lock();
+			return false;
 		}
-
-		EntryToValueKeyMap.Reserve(EntryToValueKeyMap.Num() + EntryKeys.Num());
-
-		if constexpr (!PCG::Private::MetadataTraits<T>::CompressData)
+	}
+	
+	bool SetAdd(PCGMetadataEntryKey ItemKey, const FPCGMetadataAttributeGeneric* InAttributeA, PCGMetadataEntryKey InEntryKeyA, const FPCGMetadataAttributeGeneric* InAttributeB, PCGMetadataEntryKey InEntryKeyB)
+	{
+		if constexpr (PCG::Private::MetadataTraits<T>::CanSubAdd)
 		{
-			for (int32 i = 0; i < EntryKeys.Num(); ++i)
-			{
-				const PCGMetadataValueKey ValueKey = StartIndex + i + ValueKeyOffset;
-				EntryToValueKeyMap.Emplace(*EntryKeys[i], ValueKey);
-			}
+			SetValue(ItemKey,
+				PCG::Private::MetadataTraits<T>::Add(
+					InAttributeA->GetValueFromItemKey<T>(InEntryKeyA),
+					InAttributeB->GetValueFromItemKey<T>(InEntryKeyB)));
+			return true;
 		}
-
-		if (!bLockless)
+		else
 		{
-			EntryMapLock.Unlock();
+			return false;
 		}
-
-		return StartIndex;
 	}
-
-	void SetValues_TryLockless(TArrayView<PCGMetadataEntryKey*> EntryKeys, TArrayView<const T> InValues, int32 StartIndex)
+	
+	bool SetSub(PCGMetadataEntryKey ItemKey, const FPCGMetadataAttributeGeneric* InAttributeA, PCGMetadataEntryKey InEntryKeyA, const FPCGMetadataAttributeGeneric* InAttributeB, PCGMetadataEntryKey InEntryKeyB)
 	{
-		if constexpr (PCG::Private::MetadataTraits<T>::CompressData)
+		if constexpr (PCG::Private::MetadataTraits<T>::CanSubAdd)
 		{
-			SetValues(EntryKeys, InValues);
+			SetValue(ItemKey,
+				PCG::Private::MetadataTraits<T>::Sub(
+					InAttributeA->GetValueFromItemKey<T>(InEntryKeyA),
+					InAttributeB->GetValueFromItemKey<T>(InEntryKeyB)));
+			return true;
 		}
 		else
 		{
-			check(StartIndex != INDEX_NONE && Values.IsValidIndex(StartIndex + InValues.Num() - 1));
-
-			if constexpr (std::is_trivially_copyable_v<T>)
-			{
-				FMemory::Memcpy(Values.GetData() + StartIndex, InValues.GetData(), InValues.Num() * sizeof(T));
-			}
-			else
-			{
-				for (int32 i = 0; i < InValues.Num(); ++i)
-				{
-					Values[StartIndex + i] = InValues[i];
-				}
-			}
+			return false;
 		}
 	}
-
-protected:
-	/** Code related to computing compared values (min, max, sub, add) */
-	template<typename IT = T, typename TEnableIf<PCG::Private::MetadataTraits<IT>::CanMinMax>::Type* = nullptr>
-	bool SetMin(PCGMetadataEntryKey ItemKey, const FPCGMetadataAttributeBase* InAttributeA, PCGMetadataEntryKey InEntryKeyA, const FPCGMetadataAttributeBase* InAttributeB, PCGMetadataEntryKey InEntryKeyB)
-	{
-		SetValue(ItemKey, 
-			PCG::Private::MetadataTraits<IT>::Min(
-				static_cast<const FPCGMetadataAttribute<T>*>(InAttributeA)->GetValueFromItemKey(InEntryKeyA),
-				static_cast<const FPCGMetadataAttribute<T>*>(InAttributeB)->GetValueFromItemKey(InEntryKeyB)));
-
-		return true;
-	}
-
-	template<typename IT = T, typename TEnableIf<!PCG::Private::MetadataTraits<IT>::CanMinMax>::Type* = nullptr>
-	bool SetMin(PCGMetadataEntryKey ItemKey, const FPCGMetadataAttributeBase* InAttributeA, PCGMetadataEntryKey InEntryKeyA, const FPCGMetadataAttributeBase* InAttributeB, PCGMetadataEntryKey InEntryKeyB)
-	{
-		return false;
-	}
-
-	template<typename IT = T, typename TEnableIf<PCG::Private::MetadataTraits<IT>::CanMinMax>::Type* = nullptr>
-	bool SetMax(PCGMetadataEntryKey ItemKey, const FPCGMetadataAttributeBase* InAttributeA, PCGMetadataEntryKey InEntryKeyA, const FPCGMetadataAttributeBase* InAttributeB, PCGMetadataEntryKey InEntryKeyB)
-	{
-		SetValue(ItemKey,
-			PCG::Private::MetadataTraits<IT>::Max(
-				static_cast<const FPCGMetadataAttribute<T>*>(InAttributeA)->GetValueFromItemKey(InEntryKeyA),
-				static_cast<const FPCGMetadataAttribute<T>*>(InAttributeB)->GetValueFromItemKey(InEntryKeyB)));
-
-		return true;
-	}
-
-	template<typename IT = T, typename TEnableIf<!PCG::Private::MetadataTraits<IT>::CanMinMax>::Type* = nullptr>
-	bool SetMax(PCGMetadataEntryKey ItemKey, const FPCGMetadataAttributeBase* InAttributeA, PCGMetadataEntryKey InEntryKeyA, const FPCGMetadataAttributeBase* InAttributeB, PCGMetadataEntryKey InEntryKeyB)
-	{
-		return false;
-	}
-
-	template<typename IT = T, typename TEnableIf<PCG::Private::MetadataTraits<IT>::CanSubAdd>::Type* = nullptr>
-	bool SetAdd(PCGMetadataEntryKey ItemKey, const FPCGMetadataAttributeBase* InAttributeA, PCGMetadataEntryKey InEntryKeyA, const FPCGMetadataAttributeBase* InAttributeB, PCGMetadataEntryKey InEntryKeyB)
-	{
-		SetValue(ItemKey,
-			PCG::Private::MetadataTraits<IT>::Add(
-				static_cast<const FPCGMetadataAttribute<T>*>(InAttributeA)->GetValueFromItemKey(InEntryKeyA),
-				static_cast<const FPCGMetadataAttribute<T>*>(InAttributeB)->GetValueFromItemKey(InEntryKeyB)));
-
-		return true;
-	}
-
-	template<typename IT = T, typename TEnableIf<!PCG::Private::MetadataTraits<IT>::CanSubAdd>::Type* = nullptr>
-	bool SetAdd(PCGMetadataEntryKey ItemKey, const FPCGMetadataAttributeBase* InAttributeA, PCGMetadataEntryKey InEntryKeyA, const FPCGMetadataAttributeBase* InAttributeB, PCGMetadataEntryKey InEntryKeyB)
-	{
-		return false;
-	}
-
-	template<typename IT = T, typename TEnableIf<PCG::Private::MetadataTraits<IT>::CanSubAdd>::Type* = nullptr>
-	bool SetSub(PCGMetadataEntryKey ItemKey, const FPCGMetadataAttributeBase* InAttributeA, PCGMetadataEntryKey InEntryKeyA, const FPCGMetadataAttributeBase* InAttributeB, PCGMetadataEntryKey InEntryKeyB)
-	{
-		SetValue(ItemKey,
-			PCG::Private::MetadataTraits<IT>::Sub(
-				static_cast<const FPCGMetadataAttribute<T>*>(InAttributeA)->GetValueFromItemKey(InEntryKeyA),
-				static_cast<const FPCGMetadataAttribute<T>*>(InAttributeB)->GetValueFromItemKey(InEntryKeyB)));
-
-		return true;
-	}
-
-	template<typename IT = T, typename TEnableIf<!PCG::Private::MetadataTraits<IT>::CanSubAdd>::Type* = nullptr>
-	bool SetSub(PCGMetadataEntryKey ItemKey, const FPCGMetadataAttributeBase* InAttributeA, PCGMetadataEntryKey InEntryKeyA, const FPCGMetadataAttributeBase* InAttributeB, PCGMetadataEntryKey InEntryKeyB)
-	{
-		return false;
-	}
-
-	template<typename IT = T, typename TEnableIf<PCG::Private::MetadataTraits<IT>::CanMulDiv>::Type* = nullptr>
-	bool SetMul(PCGMetadataEntryKey ItemKey, const FPCGMetadataAttributeBase* InAttributeA, PCGMetadataEntryKey InEntryKeyA, const FPCGMetadataAttributeBase* InAttributeB, PCGMetadataEntryKey InEntryKeyB)
-	{
-		SetValue(ItemKey,
-			PCG::Private::MetadataTraits<IT>::Mul(
-				static_cast<const FPCGMetadataAttribute<T>*>(InAttributeA)->GetValueFromItemKey(InEntryKeyA),
-				static_cast<const FPCGMetadataAttribute<T>*>(InAttributeB)->GetValueFromItemKey(InEntryKeyB)));
-
-		return true;
-	}
-
-	template<typename IT = T, typename TEnableIf<!PCG::Private::MetadataTraits<IT>::CanMulDiv>::Type* = nullptr>
-	bool SetMul(PCGMetadataEntryKey ItemKey, const FPCGMetadataAttributeBase* InAttributeA, PCGMetadataEntryKey InEntryKeyA, const FPCGMetadataAttributeBase* InAttributeB, PCGMetadataEntryKey InEntryKeyB)
-	{
-		return false;
-	}
-
-	template<typename IT = T, typename TEnableIf<PCG::Private::MetadataTraits<IT>::CanMulDiv>::Type* = nullptr>
-	bool SetDiv(PCGMetadataEntryKey ItemKey, const FPCGMetadataAttributeBase* InAttributeA, PCGMetadataEntryKey InEntryKeyA, const FPCGMetadataAttributeBase* InAttributeB, PCGMetadataEntryKey InEntryKeyB)
+	
+	bool SetMul(PCGMetadataEntryKey ItemKey, const FPCGMetadataAttributeGeneric* InAttributeA, PCGMetadataEntryKey InEntryKeyA, const FPCGMetadataAttributeGeneric* InAttributeB, PCGMetadataEntryKey InEntryKeyB)
 	{
-		SetValue(ItemKey,
-			PCG::Private::MetadataTraits<IT>::Div(
-				static_cast<const FPCGMetadataAttribute<T>*>(InAttributeA)->GetValueFromItemKey(InEntryKeyA),
-				static_cast<const FPCGMetadataAttribute<T>*>(InAttributeB)->GetValueFromItemKey(InEntryKeyB)));
-
-		return true;
+		if constexpr (PCG::Private::MetadataTraits<T>::CanMulDiv)
+		{
+			SetValue(ItemKey,
+				PCG::Private::MetadataTraits<T>::Mul(
+					InAttributeA->GetValueFromItemKey<T>(InEntryKeyA),
+					InAttributeB->GetValueFromItemKey<T>(InEntryKeyB)));
+			return true;
+		}
+		else
+		{
+			return false;
+		}
 	}
-
-	template<typename IT = T, typename TEnableIf<!PCG::Private::MetadataTraits<IT>::CanMulDiv>::Type* = nullptr>
-	bool SetDiv(PCGMetadataEntryKey ItemKey, const FPCGMetadataAttributeBase* InAttributeA, PCGMetadataEntryKey InEntryKeyA, const FPCGMetadataAttributeBase* InAttributeB, PCGMetadataEntryKey InEntryKeyB)
+	
+	bool SetDiv(PCGMetadataEntryKey ItemKey, const FPCGMetadataAttributeGeneric* InAttributeA, PCGMetadataEntryKey InEntryKeyA, const FPCGMetadataAttributeGeneric* InAttributeB, PCGMetadataEntryKey InEntryKeyB)
 	{
-		return false;
+		if constexpr (PCG::Private::MetadataTraits<T>::CanMulDiv)
+		{
+			SetValue(ItemKey,
+				PCG::Private::MetadataTraits<T>::Div(
+					InAttributeA->GetValueFromItemKey<T>(InEntryKeyA),
+					InAttributeB->GetValueFromItemKey<T>(InEntryKeyB)));
+			return true;
+		}
+		else
+		{
+			return false;
+		}
 	}
 
 	/** Weighted/interpolated values related code */
-	template<typename IT = T, typename TEnableIf<PCG::Private::MetadataTraits<IT>::CanInterpolate>::Type* = nullptr>
-	void ZeroValue(PCGMetadataEntryKey ItemKey)
-	{
-		SetValue(ItemKey, PCG::Private::MetadataTraits<IT>::ZeroValue());
-	}
-
-	template<typename IT = T, typename TEnableIf<!PCG::Private::MetadataTraits<IT>::CanInterpolate>::Type* = nullptr>
 	void ZeroValue(PCGMetadataEntryKey ItemKey)
 	{
-		// Intentionally empty
+		if constexpr (PCG::Private::MetadataTraits<T>::CanInterpolate)
+		{
+			SetValue(ItemKey, PCG::Private::MetadataTraits<T>::ZeroValue());
+		}
 	}
 
-	template<typename IT = T, typename TEnableIf<PCG::Private::MetadataTraits<IT>::CanInterpolate>::Type* = nullptr>
-	void Accumulate(PCGMetadataEntryKey ItemKey, const FPCGMetadataAttributeBase* InAttribute, PCGMetadataEntryKey InEntryKey, float Weight)
+	void Accumulate(PCGMetadataEntryKey ItemKey, const FPCGMetadataAttributeGeneric* InAttribute, PCGMetadataEntryKey InEntryKey, float Weight)
 	{
-		IT Value = PCG::Private::MetadataTraits<IT>::WeightedSum(
-				GetValueFromItemKey(ItemKey),
-				static_cast<const FPCGMetadataAttribute<T>*>(InAttribute)->GetValueFromItemKey(InEntryKey),
-				Weight);
-
-		if constexpr (PCG::Private::MetadataTraits<IT>::InterpolationNeedsNormalization)
+		if constexpr (PCG::Private::MetadataTraits<T>::CanInterpolate)
 		{
-			static_assert(PCG::Private::MetadataTraits<IT>::CanNormalize);
+			T Value = PCG::Private::MetadataTraits<T>::WeightedSum(
+					GetValueFromItemKey(ItemKey),
+					InAttribute->GetValueFromItemKey<T>(InEntryKey),
+					Weight);
 
-			// We need to normalize the resulting value
-			PCG::Private::MetadataTraits<IT>::Normalize(Value);
-		}
-		
-		SetValue(ItemKey, Value);
-	}
+			if constexpr (PCG::Private::MetadataTraits<T>::InterpolationNeedsNormalization)
+			{
+				static_assert(PCG::Private::MetadataTraits<T>::CanNormalize);
 
-	template<typename IT = T, typename TEnableIf<!PCG::Private::MetadataTraits<IT>::CanInterpolate>::Type* = nullptr>
-	void Accumulate(PCGMetadataEntryKey ItemKey, const FPCGMetadataAttributeBase* InAttribute, PCGMetadataEntryKey InEntryKey, float Weight)
-	{
-		// Empty on purpose
+				// We need to normalize the resulting value
+				PCG::Private::MetadataTraits<T>::Normalize(Value);
+			}
+		
+			SetValue(ItemKey, Value);
+		}
 	}
-
-	template<typename IT = T, typename TEnableIf<PCG::Private::MetadataTraits<IT>::CanInterpolate>::Type* = nullptr>
-	void Accumulate(PCGMetadataEntryKey ItemKey, const FPCGMetadataAttributeBase* InAttribute, const TArrayView<const TPair<PCGMetadataEntryKey, float>>& InWeightedKeys)
+	
+	void Accumulate(PCGMetadataEntryKey ItemKey, const FPCGMetadataAttributeGeneric* InAttribute, const TArrayView<const TPair<PCGMetadataEntryKey, float>>& InWeightedKeys)
 	{
-		IT Value = PCG::Private::MetadataTraits<IT>::ZeroValueForWeightedSum();
-		for (const TPair<PCGMetadataEntryKey, float>& WeightedEntry : InWeightedKeys)
+		if constexpr (PCG::Private::MetadataTraits<T>::CanInterpolate)
 		{
-			Value = PCG::Private::MetadataTraits<IT>::WeightedSum(
-				Value,
-				static_cast<const FPCGMetadataAttribute<T>*>(InAttribute)->GetValueFromItemKey(WeightedEntry.Key),
-				WeightedEntry.Value);
-		}
+			T Value = PCG::Private::MetadataTraits<T>::ZeroValueForWeightedSum();
+			for (const TPair<PCGMetadataEntryKey, float>& WeightedEntry : InWeightedKeys)
+			{
+				Value = PCG::Private::MetadataTraits<T>::WeightedSum(
+					Value,
+					InAttribute->GetValueFromItemKey<T>(WeightedEntry.Key),
+					WeightedEntry.Value);
+			}
 
-		if constexpr (PCG::Private::MetadataTraits<IT>::InterpolationNeedsNormalization)
-		{
-			static_assert(PCG::Private::MetadataTraits<IT>::CanNormalize);
+			if constexpr (PCG::Private::MetadataTraits<T>::InterpolationNeedsNormalization)
+			{
+				static_assert(PCG::Private::MetadataTraits<T>::CanNormalize);
 								
-			// We need to normalize the resulting value
-			PCG::Private::MetadataTraits<IT>::Normalize(Value);
-		}
-
-		SetValue(ItemKey, Value);
-	}
+				// We need to normalize the resulting value
+				PCG::Private::MetadataTraits<T>::Normalize(Value);
+			}
 
-	template<typename IT = T, typename TEnableIf<!PCG::Private::MetadataTraits<IT>::CanInterpolate>::Type* = nullptr>
-	void Accumulate(PCGMetadataEntryKey ItemKey, const FPCGMetadataAttributeBase* InAttribute, const TArrayView<const TPair<PCGMetadataEntryKey, float>>& InWeightedKeys)
-	{
-		// Empty on purpose
+			SetValue(ItemKey, Value);
+		}
 	}
-
-protected:
-	mutable PCG::FSharedLock ValueLock;
-	TArray<T> Values;
-	T DefaultValue = T{};
-	PCGMetadataValueKey ValueKeyOffset = 0;
 };
 
 namespace PCGMetadataAttribute
diff --git a/Engine/Plugins/PCG/Tests/Private/Metadata/PCGMetadataAttributeTests.cpp b/Engine/Plugins/PCG/Tests/Private/Metadata/PCGMetadataAttributeTests.cpp
--- a/Engine/Plugins/PCG/Tests/Private/Metadata/PCGMetadataAttributeTests.cpp
+++ b/Engine/Plugins/PCG/Tests/Private/Metadata/PCGMetadataAttributeTests.cpp
@@ -143,11 +143,10 @@ namespace PCGAttributeTypedTests
 	template<> EPCGMetadataTypes GetExpectedMetadataType<FSoftObjectPath>() { return EPCGMetadataTypes::SoftObjectPath; }
 	template<> EPCGMetadataTypes GetExpectedMetadataType<FSoftClassPath>() { return EPCGMetadataTypes::SoftClassPath; }
 	
-	// @todo_pcg To be replaced by `IsOfType<T>` when generic attributes are the base of the typed one.
 	template <typename T>
 	bool IsOfType(const FPCGMetadataAttributeBase* Attribute)
 	{
-		return Attribute && PCG::Private::MetadataTypes<T>::Id == Attribute->GetTypeId();
+		return Attribute && Attribute->IsGeneric() && static_cast<const FPCGMetadataAttributeGeneric*>(Attribute)->IsOfType<T>();
 	}
 	
 }
@@ -191,13 +190,12 @@ TEST_CASE_METHOD(PCGTests::FPCGBaseTestWithBasicContext, "PCG::AttributeTyped::B
 			REQUIRE(IsOfType<T>(TypedAttribute));
 		}
 
-		// @todo_pcg to be enabled when typed attribute inherit from generic
-		// SECTION("Attribute descriptor matches expected type")
-		// {
-		// 	const FPCGMetadataAttributeDesc& Desc = TypedAttribute->GetAttributeDesc();
-		// 	REQUIRE_EQUAL(Desc.Name, AttributeName);
-		// 	REQUIRE_EQUAL(Desc.ValueType, GetExpectedMetadataType<T>());
-		// }
+		SECTION("Attribute descriptor matches expected type")
+		{
+			const FPCGMetadataAttributeDesc& Desc = TypedAttribute->GetAttributeDesc();
+			REQUIRE_EQUAL(Desc.Name, AttributeName);
+			REQUIRE_EQUAL(Desc.ValueType, GetExpectedMetadataType<T>());
+		}
 
 		SECTION("Default value is set correctly")
 		{
@@ -687,15 +685,14 @@ TEST_CASE_METHOD(PCGTests::FPCGBaseTestWithBasicContext, "PCG::AttributeTyped::G
 	FPCGMetadataAttribute<double>* TypedAttr = Domain->CreateAttribute<double>(AttributeName, 3.14, /*bAllowsInterpolation=*/true, /*bOverrideParent=*/false);
 	REQUIRE_NOT_EQUAL(TypedAttr, nullptr);
 
-	// @todo_pcg to be enabled when typed attribute inherit from generic
-	// SECTION("GetMutableGenericAttribute returns the same attribute")
-	// {
-	// 	FPCGMetadataAttributeGeneric* GenericAttr = Domain->GetMutableGenericAttribute(AttributeName);
-	// 	REQUIRE_NOT_EQUAL(GenericAttr, nullptr);
-	// 	REQUIRE(IsOfType<double>(GenericAttr));
-	//
-	// 	REQUIRE_EQUAL(static_cast<FPCGMetadataAttributeGeneric*>(TypedAttr), GenericAttr);
-	// }
+	SECTION("GetMutableGenericAttribute returns the same attribute")
+	{
+		FPCGMetadataAttributeGeneric* GenericAttr = Domain->GetMutableGenericAttribute(AttributeName);
+		REQUIRE_NOT_EQUAL(GenericAttr, nullptr);
+		REQUIRE(IsOfType<double>(GenericAttr));
+	
+		REQUIRE_EQUAL(static_cast<FPCGMetadataAttributeGeneric*>(TypedAttr), GenericAttr);
+	}
 
 	SECTION("GetMutableTypedAttribute returns typed attribute")
 	{
@@ -708,28 +705,26 @@ TEST_CASE_METHOD(PCGTests::FPCGBaseTestWithBasicContext, "PCG::AttributeTyped::G
 		FPCGMetadataAttribute<float>* WrongAttr = Domain->GetMutableTypedAttribute<float>(AttributeName);
 		REQUIRE_EQUAL(WrongAttr, nullptr);
 	}
+	
+	SECTION("Setting via generic, reading via typed")
+	{
+		PCGMetadataEntryKey EntryKey = Domain->AddEntry();
+		FPCGMetadataAttributeGeneric* GenericAttr = Domain->GetMutableGenericAttribute(AttributeName);
+		GenericAttr->SetValue<double>(EntryKey, 2.718);
 
+		double ReadValue = TypedAttr->GetValueFromItemKey(EntryKey);
+		REQUIRE_EQUAL(ReadValue, 2.718);
+	}
 
-	// @todo_pcg to be enabled when typed attribute inherit from generic
-	//SECTION("Setting via generic, reading via typed")
-	//{
-	//	PCGMetadataEntryKey EntryKey = Domain->AddEntry();
-	//	FPCGMetadataAttributeGeneric* GenericAttr = Domain->GetMutableGenericAttribute(AttributeName);
-	//	GenericAttr->SetValue<double>(EntryKey, 2.718);
-
-	//	double ReadValue = TypedAttr->GetValueFromItemKey(EntryKey);
-	//	REQUIRE_EQUAL(ReadValue, 2.718);
-	//}
-
-	// SECTION("Setting via typed, reading via generic")
-	// {
-	// 	PCGMetadataEntryKey EntryKey = Domain->AddEntry();
-	// 	TypedAttr->SetValue(EntryKey, 1.618);
-	//
-	// 	FPCGMetadataAttributeGeneric* GenericAttr = Domain->GetMutableGenericAttribute(AttributeName);
-	// 	double ReadValue = GenericAttr->GetValueFromItemKey<double>(EntryKey);
-	// 	REQUIRE_EQUAL(ReadValue, 1.618);
-	// }
+	 SECTION("Setting via typed, reading via generic")
+	 {
+	 	PCGMetadataEntryKey EntryKey = Domain->AddEntry();
+	 	TypedAttr->SetValue(EntryKey, 1.618);
+	
+	 	FPCGMetadataAttributeGeneric* GenericAttr = Domain->GetMutableGenericAttribute(AttributeName);
+	 	double ReadValue = GenericAttr->GetValueFromItemKey<double>(EntryKey);
+	 	REQUIRE_EQUAL(ReadValue, 1.618);
+	 }
 }
 
 /**
@@ -1207,23 +1202,22 @@ TEST_CASE_METHOD(PCGTests::FPCGBaseTestWithBasicContext, "PCG::AttributeTyped::B
 	}
 }
 
-// @todo_pcg to be enabled when typed attribute inherit from generic
-///**
-// * Test that the typed attribute inherits IsGeneric() == true from FPCGMetadataAttributeGeneric.
-// */
-//TEST_CASE_METHOD(PCGTests::FPCGBaseTestWithBasicContext, "PCG::AttributeTyped::IsGenericFlag", "[PCG][AttributeTyped]")
-//{
-//	using namespace PCGAttributeTypedTests;
-//
-//	UPCGParamData* ParamData = CreateData<UPCGParamData>();
-//	FPCGMetadataDomain* Domain = ParamData->MutableMetadata()->GetDefaultMetadataDomain();
-//
-//	FPCGMetadataAttribute<float>* Attr = Domain->CreateAttribute<float>(AttributeName, 0.0f, true, false);
-//	REQUIRE_NOT_EQUAL(Attr, nullptr);
-//
-//	// FPCGMetadataAttribute<T> inherits from FPCGMetadataAttributeGeneric, so IsGeneric should be true
-//	REQUIRE(Attr->IsGeneric());
-//}
+/**
+* Test that the typed attribute inherits IsGeneric() == true from FPCGMetadataAttributeGeneric.
+*/
+TEST_CASE_METHOD(PCGTests::FPCGBaseTestWithBasicContext, "PCG::AttributeTyped::IsGenericFlag", "[PCG][AttributeTyped]")
+{
+	using namespace PCGAttributeTypedTests;
+
+	UPCGParamData* ParamData = CreateData<UPCGParamData>();
+	FPCGMetadataDomain* Domain = ParamData->MutableMetadata()->GetDefaultMetadataDomain();
+
+	FPCGMetadataAttribute<float>* Attr = Domain->CreateAttribute<float>(AttributeName, 0.0f, true, false);
+	REQUIRE_NOT_EQUAL(Attr, nullptr);
+
+	// FPCGMetadataAttribute<T> inherits from FPCGMetadataAttributeGeneric, so IsGeneric should be true
+	REQUIRE(Attr->IsGeneric());
+}
 
 /**
  * Test CopyExistingAttribute through the domain for typed attributes.
diff --git a/Engine/Source/Runtime/Core/Public/UObject/FortniteMainBranchObjectVersions.inl b/Engine/Source/Runtime/Core/Public/UObject/FortniteMainBranchObjectVersions.inl
--- a/Engine/Source/Runtime/Core/Public/UObject/FortniteMainBranchObjectVersions.inl
+++ b/Engine/Source/Runtime/Core/Public/UObject/FortniteMainBranchObjectVersions.inl
@@ -808,5 +808,8 @@ FORTNITE_MAIN_VERSION(InstancedSkinnedMeshBoneAttachments, 264)
 // Added toggle to show/hide per-pin attribute creation settings on PCG Custom HLSL nodes
 FORTNITE_MAIN_VERSION(PCGCustomHLSLAttributeCreationSettingsToggle, 265)
 
+// Converting FPCGMetadataAttribute<T> to Generic attributes. 
+FORTNITE_MAIN_VERSION(ConvertFPCGMetadataAttributeToGenericAttributes, 266)
+
 // -----<new versions can be added above this line>-------------------------------------------------
 #undef FORTNITE_MAIN_VERSION
