[PCG] Made sure we release execution sources when level changes.

#rnx
#Synced-CL 52317749

diff --git a/Engine/Plugins/PCG/Source/PCGEditor/Private/PCGEditor.cpp b/Engine/Plugins/PCG/Source/PCGEditor/Private/PCGEditor.cpp
--- a/Engine/Plugins/PCG/Source/PCGEditor/Private/PCGEditor.cpp
+++ b/Engine/Plugins/PCG/Source/PCGEditor/Private/PCGEditor.cpp
@@ -3323,6 +3323,9 @@ void FPCGEditor::OnClose()
 		UnregisterDelegatesForWorld(GEditor->PlayWorld.Get());
 	}
 
+	// Lose the default execution source
+	ReleaseDefaultExecutionSource();
+
 	if (IPCGEditorModule* PCGEditorModule = IPCGEditorModule::Get())
 	{
 		PCGEditorModule->OnPermissionModeChanged().RemoveAll(this);
@@ -4215,6 +4218,15 @@ void FPCGEditor::OnMapChanged(UWorld* InWorld, EMapChangeType InMapChangedType)
 {
 	if (InMapChangedType != EMapChangeType::SaveMap)
 	{
+		if (InMapChangedType == EMapChangeType::TearDownWorld)
+		{
+			ReleaseDefaultExecutionSource();
+		}
+		else
+		{
+			UpdateDefaultExecutionSource();
+		}
+
 		RefreshViews();
 
 		// Subsystem has been torn down and rebuilt.
@@ -4260,6 +4272,20 @@ void FPCGEditor::RefreshViews()
 	}
 }
 
+void FPCGEditor::ReleaseDefaultExecutionSource(bool bCollectGarbage)
+{
+	if (IPCGEditorModule* PCGEditorModule = IPCGEditorModule::Get())
+	{
+		PCGEditorModule->ClearExecutionMetadata(PCGDefaultExecutionSource);
+	}
+	PCGDefaultExecutionSource = nullptr;
+
+	if (bCollectGarbage)
+	{
+		CollectGarbage(GARBAGE_COLLECTION_KEEPFLAGS);
+	}	
+}
+
 void FPCGEditor::UpdateDefaultExecutionSource()
 {
 	const EPCGGraphUsage PreviousUsage = (PCGDefaultExecutionSource ? (Cast<UPCGDefaultWorldObjectExecutionSource>(PCGDefaultExecutionSource) != nullptr ? EPCGGraphUsage::Level : EPCGGraphUsage::Asset) : EPCGGraphUsage::Standard);
@@ -4268,12 +4294,7 @@ void FPCGEditor::UpdateDefaultExecutionSource()
 
 	if (PreviousUsage != CurrentUsage && PCGDefaultExecutionSource)
 	{
-		if (IPCGEditorModule* PCGEditorModule = IPCGEditorModule::Get())
-		{
-			PCGEditorModule->ClearExecutionMetadata(PCGDefaultExecutionSource);
-		}
-		PCGDefaultExecutionSource = nullptr;
-		CollectGarbage(GARBAGE_COLLECTION_KEEPFLAGS);
+		ReleaseDefaultExecutionSource(/*bCollectGarbage=*/true);
 		bNeedsRefreshView = true;
 	}
 
diff --git a/Engine/Plugins/PCG/Source/PCGEditor/Public/PCGEditor.h b/Engine/Plugins/PCG/Source/PCGEditor/Public/PCGEditor.h
--- a/Engine/Plugins/PCG/Source/PCGEditor/Public/PCGEditor.h
+++ b/Engine/Plugins/PCG/Source/PCGEditor/Public/PCGEditor.h
@@ -615,6 +615,7 @@ private:
 	void RefreshViews();
 
 	void UpdateDefaultExecutionSource();
+	void ReleaseDefaultExecutionSource(bool bCollectGarbage = false);
 
 	void CloseInvalidDocuments();
 
