[PCG] GPU node memory stats in profiling window

Also add some missing memory accounting in general.

[RN] [PCG] Added GPU memory created for GPU nodes to Profiling window
#Synced-CL 52339640

diff --git a/Engine/Plugins/PCG/Source/PCG/Private/Compute/Data/PCGRawBufferData.cpp b/Engine/Plugins/PCG/Source/PCG/Private/Compute/Data/PCGRawBufferData.cpp
--- a/Engine/Plugins/PCG/Source/PCG/Private/Compute/Data/PCGRawBufferData.cpp
+++ b/Engine/Plugins/PCG/Source/PCG/Private/Compute/Data/PCGRawBufferData.cpp
@@ -16,6 +16,12 @@ void UPCGRawBufferData::AddToCrc(FArchiveCrc32& Ar, bool bFullDataCrc) const
 	AddUIDToCrc(Ar);
 }
 
+void UPCGRawBufferData::GetResourceSizeEx(FResourceSizeEx& CumulativeResourceSize)
+{
+	Super::GetResourceSizeEx(CumulativeResourceSize);
+	CumulativeResourceSize.AddDedicatedSystemMemoryBytes(Data.GetAllocatedSize());
+}
+
 void UPCGRawBufferData::ReleaseTransientResources(const TCHAR* InReason)
 {
 	Super::ReleaseTransientResources(InReason);
diff --git a/Engine/Plugins/PCG/Source/PCG/Private/Compute/Data/PCGRawBufferData.h b/Engine/Plugins/PCG/Source/PCG/Private/Compute/Data/PCGRawBufferData.h
--- a/Engine/Plugins/PCG/Source/PCG/Private/Compute/Data/PCGRawBufferData.h
+++ b/Engine/Plugins/PCG/Source/PCG/Private/Compute/Data/PCGRawBufferData.h
@@ -29,6 +29,7 @@ public:
 	PCG_ASSIGN_TYPE_INFO(FPCGDataTypeInfoRawBuffer);
 
 	PCG_API virtual void AddToCrc(FArchiveCrc32& Ar, bool bFullDataCrc) const override;
+	PCG_API virtual void GetResourceSizeEx(FResourceSizeEx& CumulativeResourceSize) override;
 	PCG_API virtual void ReleaseTransientResources(const TCHAR* InRease = nullptr) override;
 	//~End UPCGData interface
 
diff --git a/Engine/Plugins/PCG/Source/PCG/Private/Compute/Elements/PCGComputeGraphElement.cpp b/Engine/Plugins/PCG/Source/PCG/Private/Compute/Elements/PCGComputeGraphElement.cpp
--- a/Engine/Plugins/PCG/Source/PCG/Private/Compute/Elements/PCGComputeGraphElement.cpp
+++ b/Engine/Plugins/PCG/Source/PCG/Private/Compute/Elements/PCGComputeGraphElement.cpp
@@ -10,6 +10,7 @@
 #include "Compute/PCGComputeCommon.h"
 #include "Data/PCGTexture2DArrayData.h"
 #include "Data/PCGTextureData.h"
+#include "Compute/PCGPinPropertiesGPU.h"
 #include "Compute/Data/PCGProxyForGPUData.h"
 #include "Compute/DataInterfaces/PCGDataCollectionDataInterface.h"
 #include "Elements/PCGStaticMeshSpawnerKernel.h"
@@ -504,6 +505,11 @@ bool FPCGComputeGraphElement::ExecuteInternal(FPCGContext* InContext) const
 			return false;
 		}
 
+#if WITH_EDITOR
+		// Compute per-node output memory sizes from the data descriptions before they are flushed.
+		ComputeNodeOutputMemorySize(Context);
+#endif
+
 		// After prepare for execute flush the data description cache because the cache uses non-negligible memory.
 		Context->DataBinding->FlushDataDescriptionCache();
 
@@ -721,6 +727,7 @@ bool FPCGComputeGraphElement::ExecuteInternal(FPCGContext* InContext) const
 
 void FPCGComputeGraphElement::PostExecuteInternal(FPCGContext* InContext) const
 {
+#if WITH_EDITOR
 	check(InContext);
 	FPCGComputeGraphContext* Context = static_cast<FPCGComputeGraphContext*>(InContext);
 	const FPCGStack* Stack = Context->GetStack();
@@ -729,7 +736,6 @@ void FPCGComputeGraphElement::PostExecuteInternal(FPCGContext* InContext) const
 		return;
 	}
 
-#if WITH_EDITOR
 	if (Context->bExecutionSuccess)
 	{
 		for (TSoftObjectPtr<const UPCGNode> NodePtr : Context->DataBinding->GetComputeGraph()->KernelToNode)
@@ -738,7 +744,14 @@ void FPCGComputeGraphElement::PostExecuteInternal(FPCGContext* InContext) const
 			IPCGGraphExecutionSource* ExecutionSource = Context->ExecutionSource.Get();
 			if (ExecutionSource && Stack && Node)
 			{
-				ExecutionSource->GetExecutionState().GetInspection().NotifyNodeExecuted(Node, Stack, /*InTimer=*/nullptr, /*bUsedCache*/false);
+				// Update profiling stats.
+				PCGUtils::FCallTime Timer;
+				if (const uint64* MemSize = Context->NodeOutputMemorySizes.Find(TObjectKey<const UPCGNode>(Node)))
+				{
+					Timer.OutputGPUMemorySize = *MemSize;
+				}
+
+				ExecutionSource->GetExecutionState().GetInspection().NotifyNodeExecuted(Node, Stack, &Timer, /*bUsedCache*/false);
 			}
 		}
 	}
@@ -992,7 +1005,42 @@ void FPCGComputeGraphElement::LogCompilationMessages(FPCGComputeGraphContext* In
 		}
 	}
 }
-#endif
+
+void FPCGComputeGraphElement::ComputeNodeOutputMemorySize(FPCGComputeGraphContext* InContext) const
+{
+	check(IsInGameThread());
+
+	const UPCGComputeGraph* ComputeGraph = InContext ? InContext->ComputeGraph : nullptr;
+	if (!ComputeGraph || !InContext->DataBinding)
+	{
+		return;
+	}
+
+	for (int32 KernelIndex = 0; KernelIndex < ComputeGraph->KernelToNode.Num(); ++KernelIndex)
+	{
+		const UPCGComputeKernel* Kernel = Cast<UPCGComputeKernel>(ComputeGraph->GetKernel(KernelIndex));
+		const UPCGNode* Node = ComputeGraph->KernelToNode[KernelIndex].Get();
+		if (!Kernel || !Node)
+		{
+			continue;
+		}
+
+		TArray<FPCGPinPropertiesGPU> OutputPins;
+		Kernel->GetOutputPins(OutputPins);
+
+		uint64 KernelMemorySize = 0;
+		for (const FPCGPinPropertiesGPU& OutputPin : OutputPins)
+		{
+			if (TSharedPtr<const FPCGDataCollectionDesc> DataDesc = InContext->DataBinding->GetCachedKernelPinDataDesc(Kernel, OutputPin.Label, /*bIsInput=*/false))
+			{
+				KernelMemorySize += PCGComputeHelpers::ComputeSizeBytes(DataDesc);
+			}
+		}
+
+		InContext->NodeOutputMemorySizes.FindOrAdd(Node) += KernelMemorySize;
+	}
+}
+#endif // WITH_EDITOR
 
 UPCGComputeGraphSettings::UPCGComputeGraphSettings()
 {
diff --git a/Engine/Plugins/PCG/Source/PCG/Private/Compute/Elements/PCGComputeGraphElement.h b/Engine/Plugins/PCG/Source/PCG/Private/Compute/Elements/PCGComputeGraphElement.h
--- a/Engine/Plugins/PCG/Source/PCG/Private/Compute/Elements/PCGComputeGraphElement.h
+++ b/Engine/Plugins/PCG/Source/PCG/Private/Compute/Elements/PCGComputeGraphElement.h
@@ -86,6 +86,11 @@ public:
 
 	std::atomic<bool> bGraphValid = false;
 
+#if WITH_EDITOR
+	/** Accumulated output memory size per GPU node, computed from data descriptions before the cache is flushed. */
+	TMap<TObjectKey<const UPCGNode>, uint64> NodeOutputMemorySizes;
+#endif
+
 	using FDebugDataPrepareAction = TFunction<bool(FPCGComputeGraphContext*)>;
 	TArray<FDebugDataPrepareAction> DebugDataPrepareActions;
 };
@@ -138,6 +143,7 @@ protected:
 	void StoreDataForInspection(FPCGComputeGraphContext* InContext) const;
 
 	void LogCompilationMessages(FPCGComputeGraphContext* InContext) const;
+	void ComputeNodeOutputMemorySize(FPCGComputeGraphContext* InContext) const;
 #endif
 };
 
diff --git a/Engine/Plugins/PCG/Source/PCG/Private/Compute/PCGComputeCommon.cpp b/Engine/Plugins/PCG/Source/PCG/Private/Compute/PCGComputeCommon.cpp
--- a/Engine/Plugins/PCG/Source/PCG/Private/Compute/PCGComputeCommon.cpp
+++ b/Engine/Plugins/PCG/Source/PCG/Private/Compute/PCGComputeCommon.cpp
@@ -256,6 +256,75 @@ namespace PCGComputeHelpers
 		return bBufferTooLarge;
 	}
 
+	bool CanPackIntoDataCollection(const FPCGDataCollectionDesc& InDataCollectionDesc)
+	{
+		TConstArrayView<FPCGDataDesc> DataDescs = InDataCollectionDesc.GetDataDescriptions();
+		if (DataDescs.IsEmpty())
+		{
+			return true;
+		}
+
+		const FPCGDataTypeIdentifier& FirstType = DataDescs[0].GetType();
+		if (!IsTypeAllowedInDataCollection(FirstType))
+		{
+			return false;
+		}
+
+		for (int32 i = 1; i < DataDescs.Num(); ++i)
+		{
+			if (!DataDescs[i].GetType().IsSameType(FirstType))
+			{
+				return false;
+			}
+		}
+
+		return true;
+	}
+
+	uint64 ComputeSizeBytes(TSharedPtr<const FPCGDataCollectionDesc> InDataCollectionDesc)
+	{
+		if (!InDataCollectionDesc)
+		{
+			return 0;
+		}
+
+		if (CanPackIntoDataCollection(*InDataCollectionDesc))
+		{
+			return PCGDataCollectionPackingHelpers::ComputePackedSizeBytes(InDataCollectionDesc);
+		}
+
+		// Non-data-collection types that can be written from GPU nodes: currently texture render targets or raw uint32 buffers.
+		uint64 TotalSizeBytes = 0;
+		for (const FPCGDataDesc& DataDesc : InDataCollectionDesc->GetDataDescriptions())
+		{
+			const FPCGDataTypeIdentifier& Type = DataDesc.GetType();
+			const FIntVector4 ElementCount = DataDesc.GetElementCount();
+
+			if (Type.IsChildOf(FPCGDataTypeInfoTexture2DBase::AsId()))
+			{
+				// Render target: sum all mip levels. For 2D textures ElementDimension is Two; for 2D array it is Three with Z = layer count.
+				const int64 Layers = (DataDesc.GetElementDimension() == EPCGElementDimension::Three) ? FMath::Max(ElementCount.Z, 1) : 1;
+				const EPixelFormat PixelFormat = GetPixelFormatFromPCGRenderTargetFormat(DataDesc.GetRenderTargetFormat());
+				const int64 BytesPerTexel = GPixelFormats[PixelFormat].BlockBytes;
+				int64 MipWidth = FMath::Max(ElementCount.X, 1);
+				int64 MipHeight = FMath::Max(ElementCount.Y, 1);
+				for (int32 Mip = 0; Mip < DataDesc.GetTextureNumMips(); ++Mip)
+				{
+					TotalSizeBytes += static_cast<uint64>(MipWidth * MipHeight * Layers * BytesPerTexel);
+					MipWidth  = FMath::Max(MipWidth  / 2, 1LL);
+					MipHeight = FMath::Max(MipHeight / 2, 1LL);
+				}
+			}
+			else if (Type.IsChildOf(FPCGDataTypeInfoRawBuffer::AsId()))
+			{
+				// Raw buffer: element count is the number of uint32 values.
+				TotalSizeBytes += static_cast<uint64>(FMath::Max(ElementCount.X, 0)) * sizeof(uint32);
+			}
+		}
+
+		return TotalSizeBytes;
+	}
+
 	int32 GetAttributeIdFromMetadataAttributeIndex(int32 InAttributeIndex)
 	{
 		return InAttributeIndex >= 0 ? (InAttributeIndex + PCGDataCollectionPackingConstants::NUM_RESERVED_ATTRS) : INDEX_NONE;
diff --git a/Engine/Plugins/PCG/Source/PCG/Private/Data/PCGTexture2DBaseData.cpp b/Engine/Plugins/PCG/Source/PCG/Private/Data/PCGTexture2DBaseData.cpp
--- a/Engine/Plugins/PCG/Source/PCG/Private/Data/PCGTexture2DBaseData.cpp
+++ b/Engine/Plugins/PCG/Source/PCG/Private/Data/PCGTexture2DBaseData.cpp
@@ -14,6 +14,16 @@
 
 PCG_DEFINE_TYPE_INFO(FPCGDataTypeInfoTexture2DBase, UPCGTexture2DBaseData)
 
+void UPCGTexture2DBaseData::GetResourceSizeEx(FResourceSizeEx& CumulativeResourceSize)
+{
+	Super::GetResourceSizeEx(CumulativeResourceSize);
+
+	if (TextureHandle.IsValid())
+	{
+		CumulativeResourceSize.AddDedicatedVideoMemoryBytes(TextureHandle->ComputeMemorySize());
+	}
+}
+
 void UPCGTexture2DBaseData::CopyTexture2DBaseData(UPCGTexture2DBaseData* NewTextureData) const
 {
 	CopyBaseSurfaceData(NewTextureData);
diff --git a/Engine/Plugins/PCG/Source/PCG/Private/Data/PCGTextureData.cpp b/Engine/Plugins/PCG/Source/PCG/Private/Data/PCGTextureData.cpp
--- a/Engine/Plugins/PCG/Source/PCG/Private/Data/PCGTextureData.cpp
+++ b/Engine/Plugins/PCG/Source/PCG/Private/Data/PCGTextureData.cpp
@@ -654,6 +654,12 @@ bool UPCGTextureData::Initialize(TRefCountPtr<IPooledRenderTarget> InTextureHand
 	return bSuccessfullyInitialized;
 }
 
+void UPCGTexture2DSingleBaseData::GetResourceSizeEx(FResourceSizeEx& CumulativeResourceSize)
+{
+	Super::GetResourceSizeEx(CumulativeResourceSize);
+	CumulativeResourceSize.AddDedicatedSystemMemoryBytes(ColorData.GetAllocatedSize());
+}
+
 void UPCGTextureData::GetResourceSizeEx(FResourceSizeEx& CumulativeResourceSize)
 {
 	Super::GetResourceSizeEx(CumulativeResourceSize);
diff --git a/Engine/Plugins/PCG/Source/PCG/Private/Utils/PCGExtraCapture.cpp b/Engine/Plugins/PCG/Source/PCG/Private/Utils/PCGExtraCapture.cpp
--- a/Engine/Plugins/PCG/Source/PCG/Private/Utils/PCGExtraCapture.cpp
+++ b/Engine/Plugins/PCG/Source/PCG/Private/Utils/PCGExtraCapture.cpp
@@ -66,7 +66,7 @@ void PCGUtils::FExtraCapture::Update(const PCGUtils::FScopedCall& InScopedCall)
 		}
 		else
 		{
-			ensure(Timer.OutputDataMemorySize == 0);
+			ensure(!Timer.OutputCPUMemorySize.IsSet() && !Timer.OutputGPUMemorySize.IsSet());
 			// Count memory here
 			// FIXME / TODO: This is sensitive to representation of data. If/when we change data representation to separate data from metadata, this might not report memory accurately
 			auto VisitAllData = [](const FPCGDataCollection& InCollection, TSet<const UPCGData*>& OutData)
@@ -87,14 +87,17 @@ void PCGUtils::FExtraCapture::Update(const PCGUtils::FScopedCall& InScopedCall)
 
 			TSet<const UPCGData*> NewData = OutputData.Difference(InputData);
 
-			Timer.OutputDataMemorySize = 0;
+			uint64 CPUMemorySize = 0;
+			uint64 GPUMemorySize = 0;
 			for (const UPCGData* Data : NewData)
 			{
 				FResourceSizeEx ResSize = FResourceSizeEx(EResourceSizeMode::Exclusive);
 				const_cast<UPCGData*>(Data)->GetResourceSizeEx(ResSize);
-				const SIZE_T DataSize = ResSize.GetDedicatedSystemMemoryBytes();
-				Timer.OutputDataMemorySize += DataSize;
+				CPUMemorySize += ResSize.GetDedicatedSystemMemoryBytes();
+				GPUMemorySize += ResSize.GetDedicatedVideoMemoryBytes();
 			}
+			Timer.OutputCPUMemorySize = CPUMemorySize;
+			Timer.OutputGPUMemorySize = GPUMemorySize;
 		}
 
 		break;
@@ -121,7 +124,6 @@ namespace PCGUtils
 
 		if (!Info.Children.IsEmpty())
 		{
-			Info.CallTime.OutputDataMemorySize = 0;
 			Info.NumberOfTasks = 0;
 		}
 
@@ -141,7 +143,15 @@ namespace PCGUtils
 			Info.CallTime.MinExecutionFrameTime = FMath::Min(Info.CallTime.MinExecutionFrameTime, Child.CallTime.MinExecutionFrameTime);
 			Info.CallTime.MaxExecutionFrameTime = FMath::Max(Info.CallTime.MaxExecutionFrameTime, Child.CallTime.MaxExecutionFrameTime);
 
-			Info.CallTime.OutputDataMemorySize += Child.CallTime.OutputDataMemorySize;
+			if (Child.CallTime.OutputCPUMemorySize.IsSet())
+			{
+				Info.CallTime.OutputCPUMemorySize = Info.CallTime.OutputCPUMemorySize.Get(0) + Child.CallTime.OutputCPUMemorySize.GetValue();
+			}
+
+			if (Child.CallTime.OutputGPUMemorySize.IsSet())
+			{
+				Info.CallTime.OutputGPUMemorySize = Info.CallTime.OutputGPUMemorySize.Get(0) + Child.CallTime.OutputGPUMemorySize.GetValue();
+			}
 
 			Info.NumberOfTasks += Child.NumberOfTasks;
 		}
diff --git a/Engine/Plugins/PCG/Source/PCG/Public/Compute/PCGComputeCommon.h b/Engine/Plugins/PCG/Source/PCG/Public/Compute/PCGComputeCommon.h
--- a/Engine/Plugins/PCG/Source/PCG/Public/Compute/PCGComputeCommon.h
+++ b/Engine/Plugins/PCG/Source/PCG/Public/Compute/PCGComputeCommon.h
@@ -19,6 +19,7 @@ class UPCGNode;
 class UPCGSettings;
 enum EPixelFormat : uint8;
 struct FPCGContext;
+struct FPCGDataCollectionDesc;
 struct FPCGGPUCompilationContext;
 struct FPCGPinProperties;
 
@@ -136,6 +137,13 @@ namespace PCGComputeHelpers
 	/** Returns true if the given buffer size is dangerously large. Optionally emits error log. */
 	bool IsBufferSizeTooLarge(uint64 InBufferSizeBytes, bool bInLogError = true);
 
+	/** True if all data descriptions share the same type and that type is allowed in a GPU data collection. */
+	PCG_API bool CanPackIntoDataCollection(const FPCGDataCollectionDesc& InDataCollectionDesc);
+
+	/** Computes total GPU memory size in bytes for a data collection descriptor, covering all output types:
+	* data-collection types (Point, Param) via packing helpers, and non-collection types (textures, raw buffers) via element count and format. */
+	PCG_API uint64 ComputeSizeBytes(TSharedPtr<const FPCGDataCollectionDesc> InDataCollectionDesc);
+
 	int32 GetAttributeIdFromMetadataAttributeIndex(int32 InAttributeIndex);
 	int32 GetMetadataAttributeIndexFromAttributeId(int32 InAttributeId);
 
diff --git a/Engine/Plugins/PCG/Source/PCG/Public/Data/PCGTexture2DBaseData.h b/Engine/Plugins/PCG/Source/PCG/Public/Data/PCGTexture2DBaseData.h
--- a/Engine/Plugins/PCG/Source/PCG/Public/Data/PCGTexture2DBaseData.h
+++ b/Engine/Plugins/PCG/Source/PCG/Public/Data/PCGTexture2DBaseData.h
@@ -28,6 +28,10 @@ class UPCGTexture2DBaseData : public UPCGSurfaceData
 	GENERATED_BODY()
 
 public:
+	//~ Begin UObject interface
+	virtual void GetResourceSizeEx(FResourceSizeEx& CumulativeResourceSize) override;
+	//~ End UObject interface
+
 	// ~Begin UPCGData interface
 	PCG_ASSIGN_TYPE_INFO(FPCGDataTypeInfoTexture2DBase)
 	// ~End UPCGData interface
diff --git a/Engine/Plugins/PCG/Source/PCG/Public/Data/PCGTextureData.h b/Engine/Plugins/PCG/Source/PCG/Public/Data/PCGTextureData.h
--- a/Engine/Plugins/PCG/Source/PCG/Public/Data/PCGTextureData.h
+++ b/Engine/Plugins/PCG/Source/PCG/Public/Data/PCGTextureData.h
@@ -68,6 +68,7 @@ class UPCGTexture2DSingleBaseData : public UPCGTexture2DBaseData
 public:
 	//~Being UObject interface
 	UE_API virtual void PostLoad() override;
+	UE_API virtual void GetResourceSizeEx(FResourceSizeEx& CumulativeResourceSize) override;
 	//~End UObject interface
 
 	//~Begin UPCGData interface
diff --git a/Engine/Plugins/PCG/Source/PCG/Public/Utils/PCGExtraCapture.h b/Engine/Plugins/PCG/Source/PCG/Public/Utils/PCGExtraCapture.h
--- a/Engine/Plugins/PCG/Source/PCG/Public/Utils/PCGExtraCapture.h
+++ b/Engine/Plugins/PCG/Source/PCG/Public/Utils/PCGExtraCapture.h
@@ -32,7 +32,8 @@ namespace PCGUtils
 
 		int32 ExecutionFrameCount = 0;
 		int32 PrepareDataFrameCount = 0;
-		uint64 OutputDataMemorySize = 0;
+		TOptional<uint64> OutputCPUMemorySize;
+		TOptional<uint64> OutputGPUMemorySize;
 
 		double PrepareDataWallTime() const { return PrepareDataEndTime - PrepareDataStartTime; }
 		double ExecutionWallTime() const { return ExecutionEndTime - ExecutionStartTime; }
diff --git a/Engine/Plugins/PCG/Source/PCGEditor/Private/Widgets/SPCGEditorGraphProfilingView.cpp b/Engine/Plugins/PCG/Source/PCGEditor/Private/Widgets/SPCGEditorGraphProfilingView.cpp
--- a/Engine/Plugins/PCG/Source/PCGEditor/Private/Widgets/SPCGEditorGraphProfilingView.cpp
+++ b/Engine/Plugins/PCG/Source/PCGEditor/Private/Widgets/SPCGEditorGraphProfilingView.cpp
@@ -37,7 +37,8 @@ namespace PCGEditorGraphProfilingView
 	const FName NAME_NbExecutionFrames = FName(TEXT("NbExecutionFrames"));
 	const FName NAME_TotalTime = FName(TEXT("TotalTime"));
 	const FName NAME_TotalWallTime = FName(TEXT("Total_WallTime"));
-	const FName NAME_MemoryInMB = FName(TEXT("Data size (MB)"));
+	const FName NAME_CPUMemoryInMB = FName(TEXT("CPU data size (MB)"));
+	const FName NAME_GPUMemoryInMB = FName(TEXT("GPU data size (MB)"));
 	const FName NAME_NumberOfTasks = FName(TEXT("NumberOfTasks"));
 
 	/** Labels of the columns */
@@ -52,7 +53,8 @@ namespace PCGEditorGraphProfilingView
 	const FText TEXT_NbExecutionFramesLabel = LOCTEXT("NbExecutionFramesLabel", "Exec Frames");
 	const FText TEXT_TotalTimeLabel = LOCTEXT("TotalTimeLabel", "Total");
 	const FText TEXT_TotalWallTimeLabel = LOCTEXT("TotalWallTimeLabel", "Total WallTime");
-	const FText TEXT_MemoryInMBLabel = LOCTEXT("MemoryInMBLabel", "Data size (MB)");
+	const FText TEXT_CPUMemoryInMBLabel = LOCTEXT("CPUMemoryInMBLabel", "CPU data (MB)");
+	const FText TEXT_GPUMemoryInMBLabel = LOCTEXT("GPUMemoryInMBLabel", "GPU data (MB)");
 	const FText TEXT_NumberOfTasks = LOCTEXT("NumberOfTasksLabel", "Num Subtasks");
 
 	/** Tooltips */
@@ -66,7 +68,8 @@ namespace PCGEditorGraphProfilingView
 	const FText TEXT_NbExecutionFramesTooltip = LOCTEXT("NbExecutionFramesTooltip", "The number of frames in which one or more execution phases were executed.");
 	const FText TEXT_TotalTimeTooltip = LOCTEXT("TotalTimeTooltip", "The total time spent in this node, summed over all execution and prepare data frames, in ms.");
 	const FText TEXT_TotalWallTimeTooltip = LOCTEXT("TotalWallTimeTooltip", "Total real time elapsed between the first call until completion, including any wait/sleep time, in ms.");
-	const FText TEXT_MemoryInMBTooltip = LOCTEXT("MemoryInMBTooltip", "Size in MB of newly created data.");
+	const FText TEXT_CPUMemoryInMBTooltip = LOCTEXT("CPUMemoryInMBTooltip", "Size in MB of newly created CPU system memory data.");
+	const FText TEXT_GPUMemoryInMBTooltip = LOCTEXT("GPUMemoryInMBTooltip", "Size in MB of output data allocated in GPU memory while executing kernels emitted by the node. This memory is freed once any downstream consumers have executed.");
 	const FText TEXT_NumberOfTasksTooltip = LOCTEXT("NumberOfTasksTooltip", "Number of subtasks executed for this subgraph. It's an approximation to give an order of magnitude. Only show a value for subgraph nodes.");
 }
 
@@ -80,9 +83,14 @@ FText FPCGProfilingTotalData::GetTotalWallTimeLabel() const
 	return FText::Format(LOCTEXT("GraphTotalWallTimeLabel", "Total Wall Time: {0} s"), TotalWallTime);
 }
 
-FText FPCGProfilingTotalData::GetTotalDataSizeInMBLabel() const
+FText FPCGProfilingTotalData::GetTotalCPUDataSizeInMBLabel() const
 {
-	return FText::Format(LOCTEXT("GraphTotalDataSizeInMBLabel", "Total data size (MB): {0} MB"), TotalDataSizeInMB);
+	return FText::Format(LOCTEXT("GraphTotalCPUDataSizeInMBLabel", "Total CPU data (MB): {0} MB"), TotalCPUDataSizeInMB);
+}
+
+FText FPCGProfilingTotalData::GetTotalGPUDataSizeInMBLabel() const
+{
+	return FText::Format(LOCTEXT("GraphTotalGPUDataSizeInMBLabel", "Total GPU data (MB): {0} MB"), TotalGPUDataSizeInMB);
 }
 
 FText FPCGProfilingTotalData::GetTotalNumTasksLabel() const
@@ -126,7 +134,8 @@ TSharedRef<SWidget> FPCGProfilingTotalData::OnOptionsMenu()
 
 	AddEntry(LOCTEXT("TotalTimeOption", "Show Total Time"), &bShowTotalTime);
 	AddEntry(LOCTEXT("TotalWallTimeOption", "Show Total Wall Time"), &bShowTotalWallTime);
-	AddEntry(LOCTEXT("TotalDataSizeOption", "Show Total Data Size"), &bShowTotalDataSize);
+	AddEntry(LOCTEXT("TotalCPUDataSizeOption", "Show Total CPU Data Size"), &bShowTotalCPUDataSize);
+	AddEntry(LOCTEXT("TotalGPUDataSizeOption", "Show Total GPU Data Size"), &bShowTotalGPUDataSize);
 	AddEntry(LOCTEXT("TotalNumTasksOption", "Show Total Num Tasks"), &bShowTotalTasks);
 
 	return MenuBuilder.MakeWidget();
@@ -135,11 +144,25 @@ TSharedRef<SWidget> FPCGProfilingTotalData::OnOptionsMenu()
 void FPCGProfilingTotalData::Reset()
 {
 	TotalNumTasks = 0;
-	TotalDataSizeInMB = 0.0;
+	TotalCPUDataSizeInMB = 0.0;
+	TotalGPUDataSizeInMB = 0.0;
 	TotalTime = 0;
 	TotalWallTime = 0;
 }
 
+static FText FormatMemoryMB(uint64 Bytes)
+{
+	const double ValueMB = Bytes / (1024.0 * 1024.0);
+
+	// For zero or values >= 1 MB, use 2 decimal places. For smaller values, use enough decimal places to show 2 significant figures.
+	const int32 FractionalDigits = (ValueMB == 0.0 || ValueMB >= 1.0) ? 2 : 1 - FMath::FloorToInt(FMath::LogX(10.0, ValueMB));
+	FNumberFormattingOptions Opts;
+	Opts.MinimumFractionalDigits = FractionalDigits;
+	Opts.MaximumFractionalDigits = FractionalDigits;
+
+	return FText::AsNumber(ValueMB, &Opts);
+}
+
 void SPCGProfilingListViewItemRow::Construct(const FArguments& InArgs, const TSharedRef<STableViewBase>& InOwnerTableView, const PCGProfilingListViewItemPtr& Item)
 {
 	InternalItem = Item;
@@ -162,7 +185,7 @@ FText FPCGProfilingListViewItem::GetTextForColumn(FName ColumnId, bool bNoGroupi
 		{
 			if (PCGSettings->ShouldExecuteOnGPU())
 			{
-				return FText::FromString(Name + TEXT(" (profiling GPU nodes not supported)"));
+				return FText::FromString(Name + TEXT(" (GPU)"));
 			}
 			else if (!PCGSettings->bEnabled)
 			{
@@ -170,7 +193,7 @@ FText FPCGProfilingListViewItem::GetTextForColumn(FName ColumnId, bool bNoGroupi
 			}
 		}
 
-		if (bHasData)
+		if (bHasCPUTimingData)
 		{
 			return FText::FromString(Name);
 		}
@@ -187,7 +210,15 @@ FText FPCGProfilingListViewItem::GetTextForColumn(FName ColumnId, bool bNoGroupi
 	{
 		return ((CallTime.PrepareDataFrameCount >= 0) ? FText::AsNumber(CallTime.PrepareDataFrameCount, NumberFormattingOptions) : FText());
 	}
-	else if (!bHasData)
+	else if (ColumnId == PCGEditorGraphProfilingView::NAME_CPUMemoryInMB)
+	{
+		return CallTime.OutputCPUMemorySize.IsSet() ? FormatMemoryMB(CallTime.OutputCPUMemorySize.GetValue()) : PCGEditorGraphProfilingView::NoDataAvailableText;
+	}
+	else if (ColumnId == PCGEditorGraphProfilingView::NAME_GPUMemoryInMB)
+	{
+		return CallTime.OutputGPUMemorySize.IsSet() ? FormatMemoryMB(CallTime.OutputGPUMemorySize.GetValue()) : PCGEditorGraphProfilingView::NoDataAvailableText;
+	}
+	else if (!bHasCPUTimingData)
 	{
 		return PCGEditorGraphProfilingView::NoDataAvailableText;
 	}
@@ -231,11 +262,6 @@ FText FPCGProfilingListViewItem::GetTextForColumn(FName ColumnId, bool bNoGroupi
 		// In ms
 		return FText::AsNumber((CallTime.TotalWallTime()) * 1000.0, NumberFormattingOptions);
 	}
-	else if (ColumnId == PCGEditorGraphProfilingView::NAME_MemoryInMB)
-	{
-		// In MB
-		return FText::AsNumber(CallTime.OutputDataMemorySize / (1024.0 * 1024.0), NumberFormattingOptions);
-	}
 	else if (ColumnId == PCGEditorGraphProfilingView::NAME_NumberOfTasks)
 	{
 		// Do not show the number of tasks if it is not higher than 1
@@ -256,7 +282,7 @@ TSharedRef<SWidget> SPCGProfilingListViewItemRow::GenerateWidgetForColumn(const
 	{
 		ColumnData = InternalItem->GetTextForColumn(ColumnId, /*bNoGrouping=*/false);
 
-		if (!InternalItem->bHasData)
+		if (!InternalItem->bHasCPUTimingData && !InternalItem->CallTime.OutputCPUMemorySize.IsSet() && !InternalItem->CallTime.OutputGPUMemorySize.IsSet())
 		{
 			ColumnDataColor = FColor(75, 75, 75);
 		}
@@ -404,8 +430,17 @@ void SPCGEditorGraphProfilingView::Construct(const FArguments& InArgs, TSharedPt
 			.AutoWidth()
 			[
 				SNew(STextBlock)
-				.Text_Lambda([this]() { return ProfilingTotalData.GetTotalDataSizeInMBLabel();})
-				.Visibility_Lambda([this]() { return ProfilingTotalData.bShowTotalDataSize ? EVisibility::Visible : EVisibility::Collapsed;})
+				.Text_Lambda([this]() { return ProfilingTotalData.GetTotalCPUDataSizeInMBLabel();})
+				.Visibility_Lambda([this]() { return ProfilingTotalData.bShowTotalCPUDataSize ? EVisibility::Visible : EVisibility::Collapsed;})
+			]
+			+SHorizontalBox::Slot()
+			.Padding(FMargin(30.0f, 0.0f, 2.0f, 0.0f))
+			.VAlign(EVerticalAlignment::VAlign_Center)
+			.AutoWidth()
+			[
+				SNew(STextBlock)
+				.Text_Lambda([this]() { return ProfilingTotalData.GetTotalGPUDataSizeInMBLabel();})
+				.Visibility_Lambda([this]() { return ProfilingTotalData.bShowTotalGPUDataSize ? EVisibility::Visible : EVisibility::Collapsed;})
 			]
 			+SHorizontalBox::Slot()
 			.Padding(FMargin(30.0f, 0.0f, 2.0f, 0.0f))
@@ -565,15 +600,24 @@ TSharedRef<SHeaderRow> SPCGEditorGraphProfilingView::CreateHeaderRowWidget()
 		.OnSort(this, &SPCGEditorGraphProfilingView::OnSortColumnHeader)
 		.InitialSortMode(EColumnSortMode::Descending)
 		.DefaultTooltip(PCGEditorGraphProfilingView::TEXT_TotalWallTimeTooltip)
-		+ SHeaderRow::Column(PCGEditorGraphProfilingView::NAME_MemoryInMB)
+		+ SHeaderRow::Column(PCGEditorGraphProfilingView::NAME_CPUMemoryInMB)
 		.ManualWidth(100.0f)
-		.DefaultLabel(PCGEditorGraphProfilingView::TEXT_MemoryInMBLabel)
+		.DefaultLabel(PCGEditorGraphProfilingView::TEXT_CPUMemoryInMBLabel)
 		.HAlignHeader(HAlign_Center)
 		.HAlignCell(HAlign_Right)
-		.SortMode(this, &SPCGEditorGraphProfilingView::GetColumnSortMode, PCGEditorGraphProfilingView::NAME_MemoryInMB)
+		.SortMode(this, &SPCGEditorGraphProfilingView::GetColumnSortMode, PCGEditorGraphProfilingView::NAME_CPUMemoryInMB)
 		.OnSort(this, &SPCGEditorGraphProfilingView::OnSortColumnHeader)
 		.InitialSortMode(EColumnSortMode::Descending)
-		.DefaultTooltip(PCGEditorGraphProfilingView::TEXT_MemoryInMBTooltip)
+		.DefaultTooltip(PCGEditorGraphProfilingView::TEXT_CPUMemoryInMBTooltip)
+		+ SHeaderRow::Column(PCGEditorGraphProfilingView::NAME_GPUMemoryInMB)
+		.ManualWidth(100.0f)
+		.DefaultLabel(PCGEditorGraphProfilingView::TEXT_GPUMemoryInMBLabel)
+		.HAlignHeader(HAlign_Center)
+		.HAlignCell(HAlign_Right)
+		.SortMode(this, &SPCGEditorGraphProfilingView::GetColumnSortMode, PCGEditorGraphProfilingView::NAME_GPUMemoryInMB)
+		.OnSort(this, &SPCGEditorGraphProfilingView::OnSortColumnHeader)
+		.InitialSortMode(EColumnSortMode::Descending)
+		.DefaultTooltip(PCGEditorGraphProfilingView::TEXT_GPUMemoryInMBTooltip)
 		+ SHeaderRow::Column(PCGEditorGraphProfilingView::NAME_NumberOfTasks)
 		.ManualWidth(110.0f)
 		.DefaultLabel(PCGEditorGraphProfilingView::TEXT_NumberOfTasks)
@@ -730,7 +774,7 @@ namespace PCGEditorGraphProfilingView
 				ListViewItem->PCGNode = Info.Node;
 				ListViewItem->EditorNode = EditorNode;
 				ListViewItem->Name = Fullname;
-				ListViewItem->bHasData = (Info.CallTime.MaxExecutionFrameTime > 0);
+				ListViewItem->bHasCPUTimingData = (Info.CallTime.MaxExecutionFrameTime > 0);
 
 				ListViewItem->CallTime = Info.CallTime;
 				ListViewItem->NumberOfTasks = Info.NumberOfTasks;
@@ -792,7 +836,8 @@ FReply SPCGEditorGraphProfilingView::Refresh()
 	{
 		ProfilingTotalData.TotalTime = TreeInfo.CallTime.TotalTime();
 		ProfilingTotalData.TotalWallTime = TreeInfo.CallTime.TotalWallTime();
-		ProfilingTotalData.TotalDataSizeInMB = TreeInfo.CallTime.OutputDataMemorySize / (1024.0 * 1024.0);
+		ProfilingTotalData.TotalCPUDataSizeInMB = TreeInfo.CallTime.OutputCPUMemorySize.Get(0) / (1024.0 * 1024.0);
+		ProfilingTotalData.TotalGPUDataSizeInMB = TreeInfo.CallTime.OutputGPUMemorySize.Get(0) / (1024.0 * 1024.0);
 		ProfilingTotalData.TotalNumTasks = TreeInfo.NumberOfTasks;
 	}
 
@@ -848,9 +893,13 @@ FReply SPCGEditorGraphProfilingView::Refresh()
 				{
 					isLess = A->CallTime.TotalWallTime() < B->CallTime.TotalWallTime();
 				}
-				else if (SortingColumn == PCGEditorGraphProfilingView::NAME_MemoryInMB)
+				else if (SortingColumn == PCGEditorGraphProfilingView::NAME_CPUMemoryInMB)
+				{
+					isLess = A->CallTime.OutputCPUMemorySize.Get(0) < B->CallTime.OutputCPUMemorySize.Get(0);
+				}
+				else if (SortingColumn == PCGEditorGraphProfilingView::NAME_GPUMemoryInMB)
 				{
-					isLess = A->CallTime.OutputDataMemorySize < B->CallTime.OutputDataMemorySize;
+					isLess = A->CallTime.OutputGPUMemorySize.Get(0) < B->CallTime.OutputGPUMemorySize.Get(0);
 				}
 				else if (SortingColumn == PCGEditorGraphProfilingView::NAME_NumberOfTasks)
 				{
diff --git a/Engine/Plugins/PCG/Source/PCGEditor/Private/Widgets/SPCGEditorGraphProfilingView.h b/Engine/Plugins/PCG/Source/PCGEditor/Private/Widgets/SPCGEditorGraphProfilingView.h
--- a/Engine/Plugins/PCG/Source/PCGEditor/Private/Widgets/SPCGEditorGraphProfilingView.h
+++ b/Engine/Plugins/PCG/Source/PCGEditor/Private/Widgets/SPCGEditorGraphProfilingView.h
@@ -28,24 +28,27 @@ struct FPCGProfilingListViewItem
 	PCGUtils::FCallTime CallTime;
 	int32 NumberOfTasks = 0;
 
-	bool bHasData = false;
+	bool bHasCPUTimingData = false;
 };
 
 struct FPCGProfilingTotalData
 {
 	bool bShowTotalTime = true;
 	bool bShowTotalWallTime = true;
-	bool bShowTotalDataSize = true;
+	bool bShowTotalCPUDataSize = true;
+	bool bShowTotalGPUDataSize = true;
 	bool bShowTotalTasks = true;
 	
 	int32 TotalNumTasks = 0;
 	double TotalTime = 0.0;
 	double TotalWallTime = 0.0;
-	double TotalDataSizeInMB = 0.0;
+	double TotalCPUDataSizeInMB = 0.0;
+	double TotalGPUDataSizeInMB = 0.0;
 
 	FText GetTotalTimeLabel() const;
 	FText GetTotalWallTimeLabel() const;
-	FText GetTotalDataSizeInMBLabel() const;
+	FText GetTotalCPUDataSizeInMBLabel() const;
+	FText GetTotalGPUDataSizeInMBLabel() const;
 	FText GetTotalNumTasksLabel() const;
 
 	TSharedRef<SWidget> MakeOptionsWidget();
