[PCG] Fix static analysis error

#rnx
#Synced-CL 52339640

diff --git a/Engine/Plugins/PCG/Source/PCG/Private/Compute/Elements/PCGComputeGraphElement.cpp b/Engine/Plugins/PCG/Source/PCG/Private/Compute/Elements/PCGComputeGraphElement.cpp
--- a/Engine/Plugins/PCG/Source/PCG/Private/Compute/Elements/PCGComputeGraphElement.cpp
+++ b/Engine/Plugins/PCG/Source/PCG/Private/Compute/Elements/PCGComputeGraphElement.cpp
@@ -1010,7 +1010,12 @@ void FPCGComputeGraphElement::ComputeNodeOutputMemorySize(FPCGComputeGraphContex
 {
 	check(IsInGameThread());
 
-	const UPCGComputeGraph* ComputeGraph = InContext ? InContext->ComputeGraph : nullptr;
+	if (!InContext)
+	{
+		return;
+	}
+
+	const UPCGComputeGraph* ComputeGraph = InContext->ComputeGraph;
 	if (!ComputeGraph || !InContext->DataBinding)
 	{
 		return;
