Fix crash in APCGPartitionActor::Teleport - stale OriginalToLocal in pooled PA

#rnx
#Synced-CL 52356503

diff --git a/Engine/Plugins/PCG/Source/PCG/Private/Grid/PCGPartitionActor.cpp b/Engine/Plugins/PCG/Source/PCG/Private/Grid/PCGPartitionActor.cpp
--- a/Engine/Plugins/PCG/Source/PCG/Private/Grid/PCGPartitionActor.cpp
+++ b/Engine/Plugins/PCG/Source/PCG/Private/Grid/PCGPartitionActor.cpp
@@ -368,7 +368,9 @@ FPCGGridDescriptor APCGPartitionActor::GetGridDescriptor() const
 bool APCGPartitionActor::Teleport(const FVector& NewLocation)
 {
 	// We should not be teleporting a PA that is in use. We only teleport empty RuntimeGen PAs to their grid cell before initialization.
-	check(bIsRuntimeGenerated && OriginalToLocal.IsEmpty() && LocalToOriginal.IsEmpty());
+	ensureMsgf(bIsRuntimeGenerated, TEXT("Teleporting a non-RuntimeGenerated partition actor '%s'."), *GetActorNameOrLabel());
+	ensureMsgf(LocalToOriginal.IsEmpty(), TEXT("Teleporting partition actor '%s' with %d entries in LocalToOriginal."), *GetActorNameOrLabel(), LocalToOriginal.Num());
+	ensureMsgf(OriginalToLocal.IsEmpty(), TEXT("Teleporting partition actor '%s' with %d entries in OriginalToLocal (stale derived map)."), *GetActorNameOrLabel(), OriginalToLocal.Num());
 
 	RootComponent->SetMobility(EComponentMobility::Movable);
 	bool bResult = Super::SetActorLocation(NewLocation, /*bSweep=*/false, nullptr, ETeleportType::TeleportPhysics);
@@ -527,10 +529,11 @@ bool APCGPartitionActor::CleanupDeadGraphInstancesInternal(bool bCleanupTest, TA
 			}
 		}
 
-		// Remove all dead entries
+		// Remove all dead entries.
 		for (auto It = OriginalToLocal.CreateIterator(); It; ++It)
 		{
-			if (!It.Key().ResolveObjectPtr())
+			UPCGComponent* ResolvedKey = It.Key().ResolveObjectPtr();
+			if (!ResolvedKey || !IsValid(ResolvedKey))
 			{
 				It.RemoveCurrent();
 			}
