[PCG] Fix deadlock when disabling actor-componentless mode

#rnx
#Synced-CL 52356503

diff --git a/Engine/Plugins/PCG/Source/PCG/Private/RuntimeGen/PCGRuntimeGenScheduler.cpp b/Engine/Plugins/PCG/Source/PCG/Private/RuntimeGen/PCGRuntimeGenScheduler.cpp
--- a/Engine/Plugins/PCG/Source/PCG/Private/RuntimeGen/PCGRuntimeGenScheduler.cpp
+++ b/Engine/Plugins/PCG/Source/PCG/Private/RuntimeGen/PCGRuntimeGenScheduler.cpp
@@ -1997,8 +1997,10 @@ void FPCGRuntimeGenScheduler::RefreshExecutionSource(IPCGGraphExecutionSource* I
 		return;
 	}
 
-	// If we are mid way through setting up an original source, early out from this refresh.
-	if (OriginalSource == OriginalSourceBeingRegistered)
+	// If we are mid way through setting up an original source, early out from this refresh. This guards against refreshes triggered during execution source registration.
+	// Compare against OriginalComponent (UPCGComponent*) rather than OriginalSource, since OriginalSourceBeingRegistered is always a UPCGComponent* and OriginalSource
+	// may not be (e.g. when InExecutionSource is not a UPCGComponent, i.e. a runtime gen execution source, which resolves to its owning component via OriginalComponent).
+	if (OriginalComponent == OriginalSourceBeingRegistered)
 	{
 		return;
 	}
