[PCG] Broaden GPU override type support
* Add bool, float, and struct

#Synced-CL 52356503

diff --git a/Engine/Plugins/PCG/Source/PCG/Private/Compute/DataInterfaces/PCGComputeDataInterface.cpp b/Engine/Plugins/PCG/Source/PCG/Private/Compute/DataInterfaces/PCGComputeDataInterface.cpp
--- a/Engine/Plugins/PCG/Source/PCG/Private/Compute/DataInterfaces/PCGComputeDataInterface.cpp
+++ b/Engine/Plugins/PCG/Source/PCG/Private/Compute/DataInterfaces/PCGComputeDataInterface.cpp
@@ -9,6 +9,7 @@
 #include "PCGSettings.h"
 #include "Compute/PCGComputeKernel.h"
 #include "Compute/PCGDataBinding.h"
+#include "Compute/PCGKernelOverridableParamTypeInfo.h"
 
 #include "ComputeFramework/ShaderParamTypeDefinition.h"
 
@@ -206,30 +207,40 @@ bool UPCGExportableDataProvider::PrepareForExecute_GameThread(UPCGDataBinding* I
 
 void UPCGKernelParamsDataInterface::GetSupportedInputs(TArray<FShaderFunctionDefinition>& OutFunctions) const
 {
+#if WITH_EDITOR
 	if (GetProducerKernel())
 	{
 		for (const FPCGKernelOverridableParam& OverridableParam : GetProducerKernel()->GetCachedOverridableParams())
 		{
-			// @todo_pcg: Broaden type support
+			const FPCGKernelOverridableParamTypeInfo& TypeInfo = FPCGKernelOverridableParamTypeInfo::Get(OverridableParam.UnderlyingType);
+
 			OutFunctions.AddDefaulted_GetRef()
 				.SetName(FString::Format(TEXT("Get{0}Internal"), { OverridableParam.Label.ToString(), }))
-				.AddReturnType(FShaderValueType::Get(EShaderFundamentalType::Uint));
+				.AddReturnType(TypeInfo.GetShaderValueType());
 		}
 	}
+#else
+	checkf(false, TEXT("GetSupportedInputs should only be called during compute graph compilation, which is editor-only."));
+#endif
 }
 
 void UPCGKernelParamsDataInterface::GetHLSL(FString& OutHLSL, FString const& InDataInterfaceName) const
 {
+#if WITH_EDITOR
 	if (GetProducerKernel())
 	{
 		for (const FPCGKernelOverridableParam& OverridableParam : GetProducerKernel()->GetCachedOverridableParams())
 		{
-			// @todo_pcg: Broaden type support.
+			const FPCGKernelOverridableParamTypeInfo& TypeInfo = FPCGKernelOverridableParamTypeInfo::Get(OverridableParam.UnderlyingType);
+
 			OutHLSL += FString::Format(TEXT(
-				"uint {0}_{1};\n"
-				"uint Get{1}Internal_{0}() { return {0}_{1}; }\n"
-			), { InDataInterfaceName, OverridableParam.Label.ToString(), });
+				"{2} {0}_{1};\n"
+				"{2} Get{1}Internal_{0}() {{ return {0}_{1}; }}\n"
+			), { InDataInterfaceName, OverridableParam.Label.ToString(), TypeInfo.HLSLType });
 		}
 	}
+#else
+	checkf(false, TEXT("GetHLSL should only be called during compute graph compilation, which is editor-only."));
+#endif
 }
 
diff --git a/Engine/Plugins/PCG/Source/PCG/Private/Compute/PCGComputeKernel.cpp b/Engine/Plugins/PCG/Source/PCG/Private/Compute/PCGComputeKernel.cpp
--- a/Engine/Plugins/PCG/Source/PCG/Private/Compute/PCGComputeKernel.cpp
+++ b/Engine/Plugins/PCG/Source/PCG/Private/Compute/PCGComputeKernel.cpp
@@ -70,6 +70,7 @@ void UPCGComputeKernel::Initialize(const FPCGComputeKernelParams& InParams)
 				KernelOverridableParam.PropertiesNames = OverridableParam.PropertiesNames;
 				KernelOverridableParam.bRequiresGPUReadback = OverridableParam.bRequiresGPUReadback;
 				KernelOverridableParam.bIsPropertyOverriddenByPin = ResolvedSettings->IsPropertyOverriddenByPin(OverridableParam.PropertiesNames);
+				KernelOverridableParam.UnderlyingType = OverridableParam.UnderlyingType;
 			}
 		}
 	}
diff --git a/Engine/Plugins/PCG/Source/PCG/Private/Compute/PCGKernelOverridableParamTypeInfo.cpp b/Engine/Plugins/PCG/Source/PCG/Private/Compute/PCGKernelOverridableParamTypeInfo.cpp
new file mode 100644
--- /dev/null
+++ b/Engine/Plugins/PCG/Source/PCG/Private/Compute/PCGKernelOverridableParamTypeInfo.cpp
@@ -0,0 +1,67 @@
+// Copyright Epic Games, Inc. All Rights Reserved.
+
+#include "Compute/PCGKernelOverridableParamTypeInfo.h"
+
+#if WITH_EDITOR
+FShaderValueTypeHandle FPCGKernelOverridableParamTypeInfo::GetShaderValueType() const
+{
+	check(NumRows >= 1 && NumRows <= 4 && NumCols >= 1 && NumCols <= 4);
+
+	if (NumCols > 1)
+	{
+		return FShaderValueType::Get(FundamentalType, NumRows, NumCols);
+	}
+	else if (NumRows > 1)
+	{
+		return FShaderValueType::Get(FundamentalType, NumRows);
+	}
+	else
+	{
+		return FShaderValueType::Get(FundamentalType);
+	}
+}
+
+const FPCGKernelOverridableParamTypeInfo& FPCGKernelOverridableParamTypeInfo::Get(EPCGMetadataTypes InType)
+{
+	static const FPCGKernelOverridableParamTypeInfo BoolInfo      = { TEXT("bool"),     TEXT("Bool"),      EShaderFundamentalType::Bool,  1, 1 };
+	static const FPCGKernelOverridableParamTypeInfo UintInfo      = { TEXT("uint"),     TEXT("Uint"),      EShaderFundamentalType::Uint,  1, 1 };
+	static const FPCGKernelOverridableParamTypeInfo IntInfo       = { TEXT("int"),      TEXT("Int"),       EShaderFundamentalType::Int,   1, 1 };
+	static const FPCGKernelOverridableParamTypeInfo FloatInfo     = { TEXT("float"),    TEXT("Float"),     EShaderFundamentalType::Float, 1, 1 };
+	static const FPCGKernelOverridableParamTypeInfo Float2Info    = { TEXT("float2"),   TEXT("Float2"),    EShaderFundamentalType::Float, 2, 1 };
+	static const FPCGKernelOverridableParamTypeInfo Float3Info    = { TEXT("float3"),   TEXT("Float3"),    EShaderFundamentalType::Float, 3, 1 };
+	static const FPCGKernelOverridableParamTypeInfo Float4Info    = { TEXT("float4"),   TEXT("Float4"),    EShaderFundamentalType::Float, 4, 1 };
+	static const FPCGKernelOverridableParamTypeInfo RotatorInfo   = { TEXT("float3"),   TEXT("Rotator"),   EShaderFundamentalType::Float, 3, 1 };
+	static const FPCGKernelOverridableParamTypeInfo QuatInfo      = { TEXT("float4"),   TEXT("Quat"),      EShaderFundamentalType::Float, 4, 1 };
+	static const FPCGKernelOverridableParamTypeInfo TransformInfo = { TEXT("float4x4"), TEXT("Transform"), EShaderFundamentalType::Float, 4, 4 };
+
+	switch (InType)
+	{
+	case EPCGMetadataTypes::Boolean:
+		return BoolInfo;
+	case EPCGMetadataTypes::Name:
+	case EPCGMetadataTypes::Enum:
+		return UintInfo;
+	case EPCGMetadataTypes::Integer32:
+	case EPCGMetadataTypes::Integer64: // Note: truncated to 32-bit in HLSL.
+		return IntInfo;
+	case EPCGMetadataTypes::Float:
+	case EPCGMetadataTypes::Double: // Note: reduced to 32-bit float in HLSL.
+		return FloatInfo;
+	case EPCGMetadataTypes::Vector2:
+		return Float2Info;
+	case EPCGMetadataTypes::Vector:
+		return Float3Info;
+	case EPCGMetadataTypes::Vector4:
+		return Float4Info;
+	case EPCGMetadataTypes::Rotator:
+		return RotatorInfo;
+	case EPCGMetadataTypes::Quaternion:
+		return QuatInfo;
+	case EPCGMetadataTypes::Transform:
+		return TransformInfo;
+	default:
+		ensureMsgf(false, TEXT("Unsupported EPCGMetadataTypes '%s' in GPU override type mapping."), *StaticEnum<EPCGMetadataTypes>()->GetNameStringByValue(static_cast<int64>(InType)));
+		return UintInfo;
+	}
+}
+#endif
diff --git a/Engine/Plugins/PCG/Source/PCG/Private/Compute/PCGKernelOverridableParamTypeInfo.h b/Engine/Plugins/PCG/Source/PCG/Private/Compute/PCGKernelOverridableParamTypeInfo.h
new file mode 100644
--- /dev/null
+++ b/Engine/Plugins/PCG/Source/PCG/Private/Compute/PCGKernelOverridableParamTypeInfo.h
@@ -0,0 +1,33 @@
+// Copyright Epic Games, Inc. All Rights Reserved.
+
+#pragma once
+
+#include "Metadata/PCGMetadataCommon.h"
+
+#include "ComputeFramework/ShaderParamTypeDefinition.h"
+
+#if WITH_EDITOR
+/** Maps EPCGMetadataTypes to HLSL type info for GPU override codegen. */
+struct FPCGKernelOverridableParamTypeInfo
+{
+	/** HLSL type name (e.g. "uint", "float3", "float4x4"). */
+	const TCHAR* HLSLType = nullptr;
+
+	/** Suffix for the GetFirstAttributeAs*() HLSL accessor (e.g. "Uint", "Float3", "Transform"). */
+	const TCHAR* AccessorSuffix = nullptr;
+
+	/** Shader fundamental type (Float, Uint, Int, etc.). */
+	EShaderFundamentalType FundamentalType = EShaderFundamentalType::None;
+
+	/** Number of rows. */
+	int32 NumRows = 0;
+
+	/** Number of columns. */
+	int32 NumCols = 0;
+
+	/** Returns the appropriate FShaderValueType, selecting the correct scalar/vector/matrix overload. */
+	FShaderValueTypeHandle GetShaderValueType() const;
+
+	static const FPCGKernelOverridableParamTypeInfo& Get(EPCGMetadataTypes InType);
+};
+#endif
diff --git a/Engine/Plugins/PCG/Source/PCG/Private/Compute/PCGKernelParamsCache.cpp b/Engine/Plugins/PCG/Source/PCG/Private/Compute/PCGKernelParamsCache.cpp
--- a/Engine/Plugins/PCG/Source/PCG/Private/Compute/PCGKernelParamsCache.cpp
+++ b/Engine/Plugins/PCG/Source/PCG/Private/Compute/PCGKernelParamsCache.cpp
@@ -59,10 +59,10 @@ namespace PCGKernelParamsHelpers
 				switch (AttributeBase->GetTypeId())
 				{
 				case PCG::Private::MetadataTypes<float>::Id:
-					Value = static_cast<const FPCGMetadataAttribute<int32>*>(AttributeBase)->GetValueFromItemKey(PCGFirstEntryKey);
+					Value = static_cast<const FPCGMetadataAttribute<float>*>(AttributeBase)->GetValueFromItemKey(PCGFirstEntryKey);
 					break;
 				case PCG::Private::MetadataTypes<double>::Id:
-					Value = static_cast<const FPCGMetadataAttribute<int64>*>(AttributeBase)->GetValueFromItemKey(PCGFirstEntryKey);
+					Value = static_cast<const FPCGMetadataAttribute<double>*>(AttributeBase)->GetValueFromItemKey(PCGFirstEntryKey);
 					break;
 				default:
 					LogUnsupportedAttributeTypeError();
@@ -145,9 +145,53 @@ namespace PCGKernelParamsHelpers
 
 			BoolProperty->SetPropertyValue(OverrideValuePtr, Value);
 		}
+		else if (const FStructProperty* StructProperty = CastField<FStructProperty>(OverrideProperty))
+		{
+			const UScriptStruct* Struct = StructProperty->Struct;
+
+			if (!ensure(Struct))
+			{
+				return;
+			}
+
+			bool bHandled = false;
+
+#define PCG_OVERRIDE_STRUCT_TYPE(Type) \
+			if (Struct == TBaseStructure<Type>::Get()) \
+			{ \
+				if (AttributeBase->GetTypeId() == PCG::Private::MetadataTypes<Type>::Id) \
+				{ \
+					*static_cast<Type*>(OverrideValuePtr) = static_cast<const FPCGMetadataAttribute<Type>*>(AttributeBase)->GetValueFromItemKey(PCGFirstEntryKey); \
+					bHandled = true; \
+				} \
+				else \
+				{ \
+					LogUnsupportedAttributeTypeError(); \
+					return; \
+				} \
+			}
+
+			PCG_OVERRIDE_STRUCT_TYPE(FVector2D)
+			else PCG_OVERRIDE_STRUCT_TYPE(FVector)
+			else PCG_OVERRIDE_STRUCT_TYPE(FVector4)
+			else PCG_OVERRIDE_STRUCT_TYPE(FQuat)
+			else PCG_OVERRIDE_STRUCT_TYPE(FRotator)
+			else PCG_OVERRIDE_STRUCT_TYPE(FTransform)
+
+#undef PCG_OVERRIDE_STRUCT_TYPE
+
+			if (!bHandled)
+			{
+				PCG_KERNEL_VALIDATION_ERR(
+					InContext,
+					InSettings,
+					FText::Format(LOCTEXT("UnsupportedStructType", "Tried to override pin '{0}', but struct type '{1}' is not supported."),
+						FText::FromName(InParamName),
+						FText::FromString(Struct->GetName())));
+			}
+		}
 		else
 		{
-			// @todo_pcg: Broaden type support.
 			PCG_KERNEL_VALIDATION_ERR(
 				InContext,
 				InSettings,
@@ -215,6 +259,74 @@ FName FPCGKernelParams::GetValueName(FName InParamName) const
 	return ensure(Property && ValuePtr) ? Property->GetPropertyValue(ValuePtr) : NAME_None;
 }
 
+float FPCGKernelParams::GetValueFloat(FName InParamName) const
+{
+	const FFloatProperty* Property = nullptr;
+	const void* ValuePtr = nullptr;
+
+	if (const FPropertyAndAddress* PropertyAndAddress = OverriddenProperties.Find(InParamName))
+	{
+		Property = CastField<FFloatProperty>(PropertyAndAddress->Key);
+		ValuePtr = PropertyAndAddress->Value;
+	}
+	else if (const UPCGSettings* Settings = WeakSettings.Get())
+	{
+		Property = CastField<FFloatProperty>(Settings->GetClass()->FindPropertyByName(InParamName));
+		ValuePtr = Property ? Property->ContainerPtrToValuePtr<void>(Settings) : nullptr;
+	}
+
+	return ensure(Property && ValuePtr) ? Property->GetPropertyValue(ValuePtr) : 0.0f;
+}
+
+double FPCGKernelParams::GetValueDouble(FName InParamName) const
+{
+	const FDoubleProperty* Property = nullptr;
+	const void* ValuePtr = nullptr;
+
+	if (const FPropertyAndAddress* PropertyAndAddress = OverriddenProperties.Find(InParamName))
+	{
+		Property = CastField<FDoubleProperty>(PropertyAndAddress->Key);
+		ValuePtr = PropertyAndAddress->Value;
+	}
+	else if (const UPCGSettings* Settings = WeakSettings.Get())
+	{
+		Property = CastField<FDoubleProperty>(Settings->GetClass()->FindPropertyByName(InParamName));
+		ValuePtr = Property ? Property->ContainerPtrToValuePtr<void>(Settings) : nullptr;
+	}
+
+	return ensure(Property && ValuePtr) ? Property->GetPropertyValue(ValuePtr) : 0.0;
+}
+
+FVector2D FPCGKernelParams::GetValueVector2D(FName InParamName) const
+{
+	return GetStructValue<FVector2D>(InParamName);
+}
+
+FVector FPCGKernelParams::GetValueVector(FName InParamName) const
+{
+	return GetStructValue<FVector>(InParamName);
+}
+
+FVector4 FPCGKernelParams::GetValueVector4(FName InParamName) const
+{
+	return GetStructValue<FVector4>(InParamName);
+}
+
+FQuat FPCGKernelParams::GetValueQuat(FName InParamName) const
+{
+	return GetStructValue<FQuat>(InParamName);
+}
+
+FRotator FPCGKernelParams::GetValueRotator(FName InParamName) const
+{
+	return GetStructValue<FRotator>(InParamName);
+}
+
+FTransform FPCGKernelParams::GetValueTransform(FName InParamName) const
+{
+	return GetStructValue<FTransform>(InParamName);
+}
+
 bool FPCGKernelParamsCache::Initialize(FPCGContext* InContext, UPCGDataBinding* InDataBinding)
 {
 	TRACE_CPUPROFILER_EVENT_SCOPE(FPCGKernelParamsCache::Initialize);
diff --git a/Engine/Plugins/PCG/Source/PCG/Private/Graph/PCGGraphCompilerGPU.cpp b/Engine/Plugins/PCG/Source/PCG/Private/Graph/PCGGraphCompilerGPU.cpp
--- a/Engine/Plugins/PCG/Source/PCG/Private/Graph/PCGGraphCompilerGPU.cpp
+++ b/Engine/Plugins/PCG/Source/PCG/Private/Graph/PCGGraphCompilerGPU.cpp
@@ -8,6 +8,7 @@
 #include "Compute/PCGComputeCommon.h"
 #include "Compute/PCGComputeKernel.h"
 #include "Compute/PCGComputeKernelSource.h"
+#include "Compute/PCGKernelOverridableParamTypeInfo.h"
 #include "Compute/PCGPinPropertiesGPU.h"
 #include "Compute/DataInterfaces/PCGComputeDataInterface.h"
 #include "Compute/DataInterfaces/PCGCustomKernelDataInterface.h"
@@ -1629,24 +1630,26 @@ void FPCGGraphCompilerGPU::CompileComputeGraph(
 			{
 				if (OverridableParam.bSupportsGPU)
 				{
-					// @todo_pcg: Broaden type support.
+					const FPCGKernelOverridableParamTypeInfo& TypeInfo = FPCGKernelOverridableParamTypeInfo::Get(OverridableParam.UnderlyingType);
+					const FString ParamLabel = OverridableParam.Label.ToString();
+
 					if (!OverridableParam.bRequiresGPUReadback && Settings->IsPropertyOverriddenByPin(OverridableParam.PropertiesNames))
 					{
 						Source += FString::Format(TEXT(
-							"uint {0}_GetOverridableValue()\n"
+							"{1} {0}_GetOverridableValue()\n"
 							"{\n"
-							"	return {0}_GetFirstAttributeAsUint();\n"
+							"	return {0}_GetFirstAttributeAs{2}();\n"
 							"}\n"
-							"\n"), { OverridableParam.Label.ToString() });
+							"\n"), { ParamLabel, TypeInfo.HLSLType, TypeInfo.AccessorSuffix });
 					}
 					else
 					{
 						Source += FString::Format(TEXT(
-							"uint {0}_GetOverridableValue()\n"
+							"{1} {0}_GetOverridableValue()\n"
 							"{\n"
 							"	return Get{0}Internal();\n"
 							"}\n"
-							"\n"), { OverridableParam.Label.ToString() });
+							"\n"), { ParamLabel, TypeInfo.HLSLType });
 					}
 				}
 			}
diff --git a/Engine/Plugins/PCG/Source/PCG/Public/Compute/PCGComputeKernel.h b/Engine/Plugins/PCG/Source/PCG/Public/Compute/PCGComputeKernel.h
--- a/Engine/Plugins/PCG/Source/PCG/Public/Compute/PCGComputeKernel.h
+++ b/Engine/Plugins/PCG/Source/PCG/Public/Compute/PCGComputeKernel.h
@@ -88,6 +88,11 @@ struct FPCGKernelOverridableParam
 	/** True if this param is actually overidden by a pin. */
 	UPROPERTY()
 	bool bIsPropertyOverriddenByPin = false;
+
+#if WITH_EDITOR
+	/** The underlying metadata type, used for type-aware codegen. */
+	EPCGMetadataTypes UnderlyingType = EPCGMetadataTypes::Integer32;
+#endif
 };
 
 UCLASS(MinimalAPI, Abstract)
diff --git a/Engine/Plugins/PCG/Source/PCG/Public/Compute/PCGKernelParamsCache.h b/Engine/Plugins/PCG/Source/PCG/Public/Compute/PCGKernelParamsCache.h
--- a/Engine/Plugins/PCG/Source/PCG/Public/Compute/PCGKernelParamsCache.h
+++ b/Engine/Plugins/PCG/Source/PCG/Public/Compute/PCGKernelParamsCache.h
@@ -16,7 +16,15 @@ struct FPCGKernelParams
 public:
 	PCG_API bool GetValueBool(FName InParamName) const;
 	PCG_API int GetValueInt(FName InParamName) const;
+	PCG_API float GetValueFloat(FName InParamName) const;
+	PCG_API double GetValueDouble(FName InParamName) const;
 	PCG_API FName GetValueName(FName InParamName) const;
+	PCG_API FVector2D GetValueVector2D(FName InParamName) const;
+	PCG_API FVector GetValueVector(FName InParamName) const;
+	PCG_API FVector4 GetValueVector4(FName InParamName) const;
+	PCG_API FQuat GetValueQuat(FName InParamName) const;
+	PCG_API FRotator GetValueRotator(FName InParamName) const;
+	PCG_API FTransform GetValueTransform(FName InParamName) const;
 
 	template<typename EnumValueType>
 	EnumValueType GetValueEnum(FName InParamName) const
@@ -40,7 +48,33 @@ public:
 		}
 
 		const FNumericProperty* NumericProperty = Property ? Property->GetUnderlyingProperty() : nullptr;
-		return static_cast<EnumValueType>(ensure(NumericProperty && ValuePtr) ? NumericProperty->GetSignedIntPropertyValue(ValuePtr) : 0);
+		return static_cast<EnumValueType>(ensure(NumericProperty && ValuePtr) ? NumericProperty->GetUnsignedIntPropertyValue(ValuePtr) : 0);
+	}
+
+private:
+	template<typename T>
+	T GetStructValue(FName InParamName) const
+	{
+		const FStructProperty* Property = nullptr;
+		const void* ValuePtr = nullptr;
+
+		if (const FPropertyAndAddress* PropertyAndAddress = OverriddenProperties.Find(InParamName))
+		{
+			Property = CastField<FStructProperty>(PropertyAndAddress->Key);
+			ValuePtr = PropertyAndAddress->Value;
+		}
+		else if (const UPCGSettings* Settings = WeakSettings.Get())
+		{
+			Property = CastField<FStructProperty>(Settings->GetClass()->FindPropertyByName(InParamName));
+			ValuePtr = Property ? Property->ContainerPtrToValuePtr<void>(Settings) : nullptr;
+		}
+
+		if (ensure(Property && ValuePtr && Property->Struct == TBaseStructure<T>::Get()))
+		{
+			return *static_cast<const T*>(ValuePtr);
+		}
+
+		return T{};
 	}
 
 private:
