[PCG] Add persistent UV cache for inline ray tracing

Avoids rebuilding per-primitive UV data every dispatch by caching it in FPCGRayTracingUVCacheExtension. GetRaytracingVFAttributeData is called once per-primitive on scene add, and the GPU resources are rebuilt from cached data at dispatch time.

Gated by pcg.GPU.EnableRayTracingUVCache (default on). When disabled, falls back to the original per-dispatch primitive scan.

#Synced-CL 52356503

diff --git a/Engine/Plugins/PCG/Source/PCG/Private/RuntimeGen/PCGRuntimeGenScheduler.cpp b/Engine/Plugins/PCG/Source/PCG/Private/RuntimeGen/PCGRuntimeGenScheduler.cpp
--- a/Engine/Plugins/PCG/Source/PCG/Private/RuntimeGen/PCGRuntimeGenScheduler.cpp
+++ b/Engine/Plugins/PCG/Source/PCG/Private/RuntimeGen/PCGRuntimeGenScheduler.cpp
@@ -7,6 +7,7 @@
 #include "PCGComponent.h"
 #include "PCGGraph.h"
 #include "PCGModule.h"
+#include "PCGRayTracingUVCacheUtils.h"
 #include "PCGSubsystem.h"
 #include "PCGTrackingManager.h"
 #include "PCGWorldActor.h"
@@ -1543,6 +1544,9 @@ void FPCGRuntimeGenScheduler::OnOriginalComponentRegistered(UPCGComponent* InOri
 		return;
 	}
 
+	// Enable the RT UV cache so it's warm before the first GPU ray trace dispatch.
+	PCGRayTracingUVCache::RequestEnable_GameThread(Context.World ? Context.World->Scene : nullptr);
+
 	// When an original/non-partitioned component is registered, we need to dirty the state.
 	bContextDirty = true;
 
diff --git a/Engine/Plugins/PCG/Source/PCGCompute/Internal/PCGRayTracingUVCacheUtils.h b/Engine/Plugins/PCG/Source/PCGCompute/Internal/PCGRayTracingUVCacheUtils.h
new file mode 100644
--- /dev/null
+++ b/Engine/Plugins/PCG/Source/PCGCompute/Internal/PCGRayTracingUVCacheUtils.h
@@ -0,0 +1,13 @@
+// Copyright Epic Games, Inc. All Rights Reserved.
+
+#pragma once
+
+#include "CoreMinimal.h"
+
+class FSceneInterface;
+
+namespace PCGRayTracingUVCache
+{
+	/** Enable the PCG ray tracing UV cache for a given scene. Enqueues a render command that builds the cache from all existing scene primitives. No-op if already enabled or if ray tracing is disabled. */
+	PCGCOMPUTE_API void RequestEnable_GameThread(FSceneInterface* SceneInterface);
+}
diff --git a/Engine/Plugins/PCG/Source/PCGCompute/Private/PCGRayTrace.cpp b/Engine/Plugins/PCG/Source/PCGCompute/Private/PCGRayTrace.cpp
--- a/Engine/Plugins/PCG/Source/PCGCompute/Private/PCGRayTrace.cpp
+++ b/Engine/Plugins/PCG/Source/PCGCompute/Private/PCGRayTrace.cpp
@@ -2,9 +2,8 @@
 
 #include "PCGRayTrace.h"
 
-#include "RHI.h"
-
 #include "PCGComputeModule.h"
+#include "PCGRayTracingUVCache.h"
 
 #include "BuiltInRayTracingShaders.h"
 #include "DataDrivenShaderPlatformInfo.h"
@@ -12,6 +11,7 @@
 #include "FXRenderingUtils.h"
 #include "GlobalShader.h"
 #include "PipelineStateCache.h"
+#include "RHI.h"
 #include "RHIGPUReadback.h"
 #include "RayTracingPayloadType.h"
 #include "RayTracingShaderBindingLayout.h"
@@ -108,77 +108,77 @@ void PCGRayTrace::RenderPCGRayTraceInline(FRDGBuilder& GraphBuilder, const FPCGR
 	FGlobalShaderMap* ShaderMap = GetGlobalShaderMap(GMaxRHIFeatureLevel);
 	TShaderMapRef<FPCGRayTraceCS> ComputeShader(ShaderMap);
 
-	FRHIResourceCollectionRef PrimitiveUVCollection;
-	FRDGBufferRef PrimitiveIndirectionsRDG;
+	// Populate UV collection members and indirection data.
+	TArray<FRHIResourceCollectionMember, SceneRenderingAllocator> PrimitiveUVCollectionMembers;
+	FRDGUploadData<FUintPoint> PrimitiveIndirections(GraphBuilder, 1);
+	PrimitiveIndirections[0] = { ~0u, 0u };
 
 	if (InParams.bNeedsUVData && InParams.TexCoordsChannelIndex >= 0)
 	{
 		TRACE_CPUPROFILER_EVENT_SCOPE(PCGRayTrace::BuildPrimitiveUVData);
 
-		TArray<TPair<FRaytracingVFUVData, uint32>, SceneRenderingAllocator> UVBuffers;
-		uint32 MaxPrimitiveId = 0;
+		const FPCGRayTracingUVCacheExtension* UVCache = RenderScene->GetExtensionPtr<FPCGRayTracingUVCacheExtension>();
 
-		// Build UV data for all primitives in the render scene.
-		// @todo_pcg: Instead of building this every time we dispatch the RT, we should cache the UVs when primitives are added/removed to the scene.
-		for (int32 PrimitiveIndex = 0; PrimitiveIndex < RenderScene->Primitives.Num(); PrimitiveIndex++)
+		if (UVCache && UVCache->Cache.Num() > 0)
 		{
-			const FPrimitiveSceneInfo* Primitive = RenderScene->Primitives[PrimitiveIndex];
+			// Build from the persistent cache — no per-primitive proxy calls needed.
+			PrimitiveIndirections = FRDGUploadData<FUintPoint>(GraphBuilder, UVCache->MaxPersistentIndex + 1);
+			FMemory::Memset(PrimitiveIndirections.GetData(), 0xFF, PrimitiveIndirections.GetTotalSize());
 
-			if (!Primitive || !Primitive->Proxy || !Primitive->Proxy->HasRayTracingRepresentation())
+			for (const auto& [PrimitiveIdx, Entry] : UVCache->Cache)
 			{
-				continue;
+				PrimitiveIndirections[PrimitiveIdx] = { static_cast<uint32>(PrimitiveUVCollectionMembers.Num()), Entry.NumCoordinates };
+				PrimitiveUVCollectionMembers.Add(FRHIResourceCollectionMember(Entry.BufferSRV));
 			}
+		}
+		else
+		{
+			// Fallback: cache not available — build UV data on demand by scanning all primitives.
+			// Note: this can be extremely expensive.
+			TArray<TPair<FRaytracingVFUVData, uint32>, SceneRenderingAllocator> UVBuffers;
+			uint32 MaxPrimitiveId = 0;
 
-			const uint32 PrimitiveId = Primitive->GetPersistentIndex().Index;
-			MaxPrimitiveId = FMath::Max(MaxPrimitiveId, PrimitiveId);
+			for (int32 PrimitiveIndex = 0; PrimitiveIndex < RenderScene->Primitives.Num(); PrimitiveIndex++)
+			{
+				const FPrimitiveSceneInfo* Primitive = RenderScene->Primitives[PrimitiveIndex];
 
-			const FRaytracingVFAttributeData VFAttributeData = Primitive->Proxy->GetRaytracingVFAttributeData(InParams.View);
+				if (!Primitive || !Primitive->Proxy || !Primitive->Proxy->HasRayTracingRepresentation())
+				{
+					continue;
+				}
 
-			if (VFAttributeData.Sections.IsEmpty() || !VFAttributeData.Sections[0].UV.BufferSRV)
-			{
-				continue;
-			}
+				const uint32 PrimitiveId = Primitive->GetPersistentIndex().Index;
+				MaxPrimitiveId = FMath::Max(MaxPrimitiveId, PrimitiveId);
 
-			// Currently only FStaticMeshSceneProxy overrides GetRaytracingVFAttributeData, and it returns
-			// a single section (static meshes share one UV buffer across all sections per LOD).
-			// If other proxy types override this with multiple sections, this will need to select
-			// the correct section based on the geometry index from the ray hit.
-			UVBuffers.Emplace(VFAttributeData.Sections[0].UV, PrimitiveId);
-		}
+				const FRaytracingVFAttributeData VFAttributeData = Primitive->Proxy->GetRaytracingVFAttributeData(InParams.View);
 
-		// Build a resource collection for the primitives.
-		TArray<FRHIResourceCollectionMember, SceneRenderingAllocator> PrimitiveUVCollectionMembers;
-		FRDGUploadData<FUintPoint> PrimitiveIndirections(GraphBuilder, MaxPrimitiveId + 1);
-		FMemory::Memset(PrimitiveIndirections.GetData(), 0xFF, PrimitiveIndirections.GetTotalSize());
+				if (VFAttributeData.Sections.IsEmpty() || !VFAttributeData.Sections[0].UV.BufferSRV)
+				{
+					continue;
+				}
 
-		for (auto&& [UVBuffer, Key] : UVBuffers)
-		{
-			PrimitiveIndirections[Key] = { static_cast<uint32>(PrimitiveUVCollectionMembers.Num()), UVBuffer.NumCoordinates };
-			PrimitiveUVCollectionMembers.Add(FRHIResourceCollectionMember(UVBuffer.BufferSRV));
-		}
+				UVBuffers.Emplace(VFAttributeData.Sections[0].UV, PrimitiveId);
+			}
 
-		PrimitiveUVCollection = GraphBuilder.RHICmdList.CreateResourceCollection(PrimitiveUVCollectionMembers);
-		PrimitiveIndirectionsRDG = CreateUploadBuffer(
-			GraphBuilder,
-			TEXT("PCG::UVIndirectionBuffer"),
-			/*BytesPerElement=*/sizeof(FUintPoint),
-			/*NumElements=*/FMath::RoundUpToPowerOfTwo(PrimitiveIndirections.Num()),
-			PrimitiveIndirections);
-	}
-	else
-	{
-		// Create empty dummy resources to satisfy shader parameter binding.
-		PrimitiveUVCollection = GraphBuilder.RHICmdList.CreateResourceCollection({});
-		FRDGUploadData<FUintPoint> DummyIndirections(GraphBuilder, 1);
-		DummyIndirections[0] = { ~0u, 0u };
-		PrimitiveIndirectionsRDG = CreateUploadBuffer(
-			GraphBuilder,
-			TEXT("PCG::UVIndirectionBuffer"),
-			/*BytesPerElement=*/sizeof(FUintPoint),
-			/*NumElements=*/1,
-			DummyIndirections);
+			PrimitiveIndirections = FRDGUploadData<FUintPoint>(GraphBuilder, MaxPrimitiveId + 1);
+			FMemory::Memset(PrimitiveIndirections.GetData(), 0xFF, PrimitiveIndirections.GetTotalSize());
+
+			for (auto&& [UVBuffer, Key] : UVBuffers)
+			{
+				PrimitiveIndirections[Key] = { static_cast<uint32>(PrimitiveUVCollectionMembers.Num()), UVBuffer.NumCoordinates };
+				PrimitiveUVCollectionMembers.Add(FRHIResourceCollectionMember(UVBuffer.BufferSRV));
+			}
+		}
 	}
 
+	FRHIResourceCollectionRef PrimitiveUVCollection = GraphBuilder.RHICmdList.CreateResourceCollection(PrimitiveUVCollectionMembers);
+	FRDGBufferRef PrimitiveIndirectionsRDG = CreateUploadBuffer(
+		GraphBuilder,
+		TEXT("PCG::UVIndirectionBuffer"),
+		/*BytesPerElement=*/sizeof(FUintPoint),
+		/*NumElements=*/FMath::RoundUpToPowerOfTwo(PrimitiveIndirections.Num()),
+		PrimitiveIndirections);
+
 	FSceneUniformBuffer SceneUniformBuffer;
 	TRDGUniformBufferRef<FSceneUniformParameters> SceneUniformBufferRDG = UE::FXRenderingUtils::GetSceneUniformBuffer(GraphBuilder, SceneUniformBuffer);
 
diff --git a/Engine/Plugins/PCG/Source/PCGCompute/Private/PCGRayTracingUVCache.cpp b/Engine/Plugins/PCG/Source/PCGCompute/Private/PCGRayTracingUVCache.cpp
new file mode 100644
--- /dev/null
+++ b/Engine/Plugins/PCG/Source/PCGCompute/Private/PCGRayTracingUVCache.cpp
@@ -0,0 +1,153 @@
+// Copyright Epic Games, Inc. All Rights Reserved.
+
+#include "PCGRayTracingUVCache.h"
+#include "PCGRayTracingUVCacheUtils.h"
+
+#include "RenderingThread.h"
+#include "SceneInterface.h"
+
+#if RHI_RAYTRACING
+
+#include "PCGComputeModule.h"
+#include "DataDrivenShaderPlatformInfo.h"
+#include "HAL/IConsoleManager.h"
+#include "PrimitiveSceneInfo.h"
+#include "PrimitiveSceneProxy.h"
+#include "ScenePrivate.h"
+#include "ScenePrimitiveUpdates.h"
+
+static TAutoConsoleVariable<bool> CVarPCGRayTracingUVCache(
+	TEXT("pcg.GPU.EnableRayTracingUVCache"),
+	true,
+	TEXT("When enabled, PCG caches per-primitive UV data in a scene extension to avoid rebuilding it every ray trace dispatch."),
+	ECVF_Default);
+
+IMPLEMENT_SCENE_EXTENSION(FPCGRayTracingUVCacheExtension);
+
+namespace PCGRayTracingUVCache
+{
+	void AddPrimitiveToCache(FPCGRayTracingUVCacheExtension& Extension, const FPrimitiveSceneInfo* PrimitiveSceneInfo)
+	{
+		FPrimitiveSceneProxy* Proxy = PrimitiveSceneInfo ? PrimitiveSceneInfo->Proxy : nullptr;
+		if (!Proxy || !Proxy->HasRayTracingRepresentation())
+		{
+			return;
+		}
+
+		// Pass nullptr for View to force LOD 0 — PCG always wants highest-detail geometry.
+		const FRaytracingVFAttributeData VFData = Proxy->GetRaytracingVFAttributeData(nullptr);
+		if (VFData.Sections.IsEmpty() || !VFData.Sections[0].UV.BufferSRV)
+		{
+			return;
+		}
+
+		const int32 Index = PrimitiveSceneInfo->GetPersistentIndex().Index;
+
+		FPCGRayTracingUVCacheExtension::FCachedEntry Entry;
+		Entry.BufferSRV = VFData.Sections[0].UV.BufferSRV;
+		Entry.NumCoordinates = VFData.Sections[0].UV.NumCoordinates;
+
+		Extension.Cache.Add(Index, MoveTemp(Entry));
+
+		// MaxPersistentIndex grows monotonically.
+		Extension.MaxPersistentIndex = FMath::Max(Extension.MaxPersistentIndex, static_cast<uint32>(Index));
+	}
+}
+
+class FPCGRayTracingUVCacheUpdater : public ISceneExtensionUpdater
+{
+	DECLARE_SCENE_EXTENSION_UPDATER(FPCGRayTracingUVCacheUpdater, FPCGRayTracingUVCacheExtension);
+
+public:
+	FPCGRayTracingUVCacheUpdater(FPCGRayTracingUVCacheExtension& InExtension)
+		: Extension(InExtension)
+	{
+	}
+
+	virtual void PreSceneUpdate(FRDGBuilder& GraphBuilder, const FScenePreUpdateChangeSet& ChangeSet) override
+	{
+		if (!Extension.bEnabled)
+		{
+			return;
+		}
+
+		LLM_SCOPE_BYTAG(PCGCompute);
+		for (const FPersistentPrimitiveIndex& Id : ChangeSet.RemovedPrimitiveIds)
+		{
+			Extension.Cache.Remove(Id.Index);
+		}
+	}
+
+	virtual void PostSceneUpdate(FRDGBuilder& GraphBuilder, const FScenePostUpdateChangeSet& ChangeSet) override
+	{
+		if (!Extension.bEnabled)
+		{
+			return;
+		}
+
+		LLM_SCOPE_BYTAG(PCGCompute);
+
+		for (FPrimitiveSceneInfo* PrimitiveSceneInfo : ChangeSet.AddedPrimitiveSceneInfos)
+		{
+			PCGRayTracingUVCache::AddPrimitiveToCache(Extension, PrimitiveSceneInfo);
+		}
+	}
+
+private:
+	FPCGRayTracingUVCacheExtension& Extension;
+};
+
+FPCGRayTracingUVCacheExtension::FPCGRayTracingUVCacheExtension(FScene& InScene)
+	: ISceneExtension(InScene)
+{
+}
+
+bool FPCGRayTracingUVCacheExtension::ShouldCreateExtension(FScene& Scene)
+{
+	return IsRayTracingEnabled() && CVarPCGRayTracingUVCache.GetValueOnAnyThread();
+}
+
+ISceneExtensionUpdater* FPCGRayTracingUVCacheExtension::CreateUpdater()
+{
+	return new FPCGRayTracingUVCacheUpdater(*this);
+}
+
+#endif // RHI_RAYTRACING
+
+void PCGRayTracingUVCache::RequestEnable_GameThread(FSceneInterface* SceneInterface)
+{
+#if RHI_RAYTRACING
+	if (!SceneInterface)
+	{
+		return;
+	}
+
+	FScene* RenderScene = SceneInterface->GetRenderScene();
+	if (!RenderScene)
+	{
+		return;
+	}
+
+	ENQUEUE_RENDER_COMMAND(EnablePCGRTUVCache)([RenderScene](FRHICommandListImmediate&)
+	{
+		FPCGRayTracingUVCacheExtension* Ext = RenderScene->GetExtensionPtr<FPCGRayTracingUVCacheExtension>();
+		if (!Ext || Ext->IsEnabled())
+		{
+			return;
+		}
+
+		TRACE_CPUPROFILER_EVENT_SCOPE(PCGRayTracingUVCache::Initialize);
+		LLM_SCOPE_BYTAG(PCGCompute);
+
+		for (FPrimitiveSceneInfo* PrimitiveSceneInfo : RenderScene->Primitives)
+		{
+			PCGRayTracingUVCache::AddPrimitiveToCache(*Ext, PrimitiveSceneInfo);
+		}
+
+		Ext->SetEnabled(true);
+
+		UE_LOG(LogPCGCompute, Verbose, TEXT("PCGRayTracingUVCache: Initialized with %d entries from %d scene primitives"),
+			Ext->Cache.Num(), RenderScene->Primitives.Num());
+	});
+#endif
+}
diff --git a/Engine/Plugins/PCG/Source/PCGCompute/Private/PCGRayTracingUVCache.h b/Engine/Plugins/PCG/Source/PCGCompute/Private/PCGRayTracingUVCache.h
new file mode 100644
--- /dev/null
+++ b/Engine/Plugins/PCG/Source/PCGCompute/Private/PCGRayTracingUVCache.h
@@ -0,0 +1,44 @@
+// Copyright Epic Games, Inc. All Rights Reserved.
+
+#pragma once
+
+#include "CoreMinimal.h"
+#include "RHIResources.h"
+#include "SceneExtensions.h"
+
+#if RHI_RAYTRACING
+
+/** Scene extension that caches per-primitive UV SRV pointers for PCG GPU ray tracing.
+ *  Populated on-demand when the first runtime gen original component registers, then maintained incrementally by the scene updater.
+ */
+class FPCGRayTracingUVCacheExtension : public ISceneExtension
+{
+	DECLARE_SCENE_EXTENSION(PCGCOMPUTE_API, FPCGRayTracingUVCacheExtension);
+
+public:
+	struct FCachedEntry
+	{
+		FShaderResourceViewRHIRef BufferSRV = nullptr;
+		uint32 NumCoordinates = 0;
+	};
+
+	explicit FPCGRayTracingUVCacheExtension(FScene& InScene);
+
+	static bool ShouldCreateExtension(FScene& Scene);
+	virtual ISceneExtensionUpdater* CreateUpdater() override;
+
+	bool IsEnabled() const { return bEnabled; }
+	void SetEnabled(bool bInEnabled) { bEnabled = bInEnabled; }
+
+	/** Keyed by FPersistentPrimitiveIndex::Index. */
+	TMap<int32, FCachedEntry> Cache;
+	uint32 MaxPersistentIndex = 0;
+
+private:
+	friend class FPCGRayTracingUVCacheUpdater;
+
+	/** Set after initial build completes. */
+	bool bEnabled = false;
+};
+
+#endif // RHI_RAYTRACING
diff --git a/Engine/Source/Runtime/Engine/Private/StaticMeshSceneProxy.cpp b/Engine/Source/Runtime/Engine/Private/StaticMeshSceneProxy.cpp
--- a/Engine/Source/Runtime/Engine/Private/StaticMeshSceneProxy.cpp
+++ b/Engine/Source/Runtime/Engine/Private/StaticMeshSceneProxy.cpp
@@ -2147,11 +2147,6 @@ FRaytracingVFAttributeData FStaticMeshSceneProxy::GetRaytracingVFAttributeData(c
 {
 	FRaytracingVFAttributeData Out;
 
-	if (!View)
-	{
-		return Out;
-	}
-
 	if (!bSupportRayTracing || !bHasRayTracingRepresentation || !IsRayTracingRelevant())
 	{
 		return Out;
@@ -2162,7 +2157,7 @@ FRaytracingVFAttributeData FStaticMeshSceneProxy::GetRaytracingVFAttributeData(c
 		return Out;
 	}
 
-	const int32 LODIndex = RenderData->RayTracingProxy->bUsingRenderingLODs ? FMath::Max(GetLOD(View), (int32)GetCurrentFirstLODIdx_RenderThread()) : 0;
+	const int32 LODIndex = (View && RenderData->RayTracingProxy->bUsingRenderingLODs) ? FMath::Max(GetLOD(View), (int32)GetCurrentFirstLODIdx_RenderThread()) : 0;
 	const FStaticMeshVertexFactories& VertexFactories = (*RenderData->RayTracingProxy->LODVertexFactories)[LODIndex];
 	FRaytracingVFSectionData& Section = Out.Sections.Emplace_GetRef();
 	Section.UV.BufferSRV = VertexFactories.VertexFactory.GetTextureCoordinatesSRV();
