[PCG] Added 'get point data' mode that behaves somewhat like the level-to-asset data on a per actor basis.

#Synced-CL 52317749

diff --git a/Engine/Plugins/PCG/Source/PCG/Private/Data/Registry/PCGDataFunctionRegistryHelpers.cpp b/Engine/Plugins/PCG/Source/PCG/Private/Data/Registry/PCGDataFunctionRegistryHelpers.cpp
new file mode 100644
--- /dev/null
+++ b/Engine/Plugins/PCG/Source/PCG/Private/Data/Registry/PCGDataFunctionRegistryHelpers.cpp
@@ -0,0 +1,300 @@
+// Copyright Epic Games, Inc. All Rights Reserved.
+
+#include "Data/Registry/PCGDataFunctionRegistryHelpers.h"
+
+#include "PCGModule.h"
+#include "Data/PCGBasePointData.h"
+#include "Helpers/PCGHelpers.h"
+#include "Helpers/PCGTagHelpers.h"
+#include "Metadata/PCGMetadata.h"
+#include "Metadata/PCGMetadataAttribute.h"
+
+#include "Components/BillboardComponent.h"
+#include "Components/InstancedStaticMeshComponent.h"
+#include "Components/PrimitiveComponent.h"
+#include "Components/StaticMeshComponent.h"
+#include "Engine/StaticMesh.h"
+#include "GameFramework/Actor.h"
+#include "Materials/MaterialInstance.h"
+
+void FPCGActorMeshParsingContext::CreateTagAttributes(AActor* InActor)
+{
+	check(InActor);
+
+	if (!TagToAttributeMap || !SanitizedAttributeNames)
+	{
+		return;
+	}
+
+	UPCGMetadata* Metadata = PointMetadata ? PointMetadata : (PointData ? PointData->MutableMetadata() : nullptr);
+	if (!Metadata)
+	{
+		return;
+	}
+
+	// Derive reserved attribute names from the actual names of the pre-created attribute pointers,
+	// so we don't create tag attributes that would collide with them.
+	TSet<FName> ReservedAttributeNames;	
+	{
+		TArray<FName> DefaultDomainAttributeNames;
+		TArray<EPCGMetadataTypes> DefaultDomainAttributeTypes;
+		Metadata->GetAttributes(DefaultDomainAttributeNames, DefaultDomainAttributeTypes);
+
+		ReservedAttributeNames.Append(DefaultDomainAttributeNames);
+	}
+
+	const TArray<FName>& Tags = ActorTags ? *ActorTags : InActor->Tags;
+	for (const FName& Tag : Tags)
+	{
+		PCG::Private::FParseTagResult TagData(Tag);
+
+		if (!TagData.IsValid())
+		{
+			continue;
+		}
+
+		const FName OriginalAttributeName(TagData.GetOriginalAttribute());
+		const FName SanitizedAttributeName(TagData.Attribute);
+
+		if (ReservedAttributeNames.Contains(SanitizedAttributeName) || TagToAttributeMap->Contains(OriginalAttributeName) || SanitizedAttributeNames->Contains(SanitizedAttributeName))
+		{
+			continue;
+		}
+
+		if (!PCG::Private::CreateAttributeFromTag(TagData, Metadata))
+		{
+			continue;
+		}
+
+		if (TagData.HasBeenSanitized())
+		{
+			UE_LOG(LogPCG, Warning, TEXT("Sanitized tag string on actor '%s' to remove invalid characters: '%s' -> '%s'"), *InActor->GetName(), *TagData.GetOriginalAttribute(), *TagData.Attribute);
+		}
+
+		TagToAttributeMap->Add(OriginalAttributeName, Metadata->GetMutableAttribute(SanitizedAttributeName));
+		SanitizedAttributeNames->Add(SanitizedAttributeName);
+	}
+}
+
+int FPCGActorMeshParsingContext::ParseActorComponents(AActor* InActor)
+{
+	check(InActor && PointData);
+
+	UPCGMetadata* Metadata = PointMetadata ? PointMetadata : PointData->MutableMetadata();
+	const TArray<FName>& Tags = ActorTags ? *ActorTags : InActor->Tags;
+
+	CreateTagAttributes(InActor);
+
+	// Get relevant visual actor components
+	TInlineComponentArray<UPrimitiveComponent*> PrimitiveComponents;
+	InActor->GetComponents(PrimitiveComponents);
+
+	// Count number of points we'll add so we do a single reallocation
+	const int NumPointsBefore = PointData->GetNumPoints();
+	int NumPointsToAdd = 0;
+
+	{
+		// Implementation note - code keeps the same layout as parsing code below for readability
+		for (UPrimitiveComponent* PrimComp : PrimitiveComponents)
+		{
+			if (bIgnorePCGCreatedComponents && PrimComp->ComponentTags.Contains(PCGHelpers::DefaultPCGTag))
+			{
+				continue;
+			}
+
+			// Skip billboard components - they have no meaningful geometry
+			if (Cast<UBillboardComponent>(PrimComp))
+			{
+				continue;
+			}
+
+			if (UStaticMeshComponent* SMC = Cast<UStaticMeshComponent>(PrimComp))
+			{
+				UStaticMesh* StaticMesh = SMC->GetStaticMesh();
+				if (!StaticMesh)
+				{
+					continue;
+				}
+
+				if (UInstancedStaticMeshComponent* ISMC = Cast<UInstancedStaticMeshComponent>(SMC))
+				{
+					NumPointsToAdd += ISMC->GetNumInstances();
+				}
+				else
+				{
+					++NumPointsToAdd;
+				}
+			}
+			else if (!bOnlyStaticMeshComponents)
+			{
+				++NumPointsToAdd;
+			}
+		}
+	}
+
+	// If we're not going to add any points we can early out.
+	if (NumPointsToAdd == 0)
+	{
+		return NumPointsToAdd;
+	}
+
+	// Allocate new points and gather ranges
+	PointData->SetNumPoints(NumPointsBefore + NumPointsToAdd);
+
+	TPCGValueRange<FTransform> TransformRange = PointData->GetTransformValueRange();
+	TPCGValueRange<int32> SeedRange = PointData->GetSeedValueRange();
+	TPCGValueRange<FVector> BoundsMinRange = PointData->GetBoundsMinValueRange();
+	TPCGValueRange<FVector> BoundsMaxRange = PointData->GetBoundsMaxValueRange();
+	TPCGValueRange<int64> MetadataRange = PointData->GetMetadataEntryValueRange();
+
+	// Write actor reference value
+	const FSoftObjectPath ActorPath(InActor);
+	PCGMetadataValueKey ActorReferenceValueKey = ActorReferenceAttribute ? ActorReferenceAttribute->AddValue(ActorPath) : PCGNotFoundValueKey;
+
+	// Lambda to write all attributes on a newly-allocated point slot.
+	auto WritePointAttributes = [this, Metadata, &Tags, &TransformRange, &SeedRange, &BoundsMinRange, &BoundsMaxRange, &MetadataRange, ActorReferenceValueKey](int32 PointIndex, const FTransform& Transform, const PCGMetadataValueKey ComponentReferenceValueKey, const PCGMetadataValueKey MeshPathValueKey, const FBox& Bounds, const PCGMetadataValueKey MaterialPathValueKey)
+	{
+		TransformRange[PointIndex] = Transform;
+		SeedRange[PointIndex] = PCGHelpers::ComputeSeedFromPosition(Transform.GetLocation());
+		BoundsMinRange[PointIndex] = Bounds.Min;
+		BoundsMaxRange[PointIndex] = Bounds.Max;
+
+		int64& MetadataEntry = MetadataRange[PointIndex];
+		Metadata->InitializeOnSet(MetadataEntry);
+
+		if (MeshAttribute && MeshPathValueKey != PCGNotFoundValueKey)
+		{
+			MeshAttribute->SetValueFromValueKey(MetadataEntry, MeshPathValueKey);
+		}
+
+		if (MaterialAttribute && MaterialPathValueKey != PCGNotFoundValueKey)
+		{
+			MaterialAttribute->SetValueFromValueKey(MetadataEntry, MaterialPathValueKey);
+		}
+
+		if (ActorIndexAttribute)
+		{
+			ActorIndexAttribute->SetValue(MetadataEntry, ActorIndex);
+		}
+
+		if (ParentIndexAttribute)
+		{
+			ParentIndexAttribute->SetValue(MetadataEntry, ParentIndex);
+		}
+
+		if (RelativeTransformAttribute)
+		{
+			RelativeTransformAttribute->SetValue(MetadataEntry, RelativeTransform);
+		}
+
+		if (HierarchyDepthAttribute)
+		{
+			HierarchyDepthAttribute->SetValue(MetadataEntry, HierarchyDepth);
+		}
+
+		if (ActorReferenceAttribute)
+		{
+			ActorReferenceAttribute->SetValueFromValueKey(MetadataEntry, ActorReferenceValueKey);
+		}
+
+		if (ComponentReferenceAttribute)
+		{
+			ComponentReferenceAttribute->SetValueFromValueKey(MetadataEntry, ComponentReferenceValueKey);
+		}
+
+		// Set tag attribute values from the pre-built map
+		if (TagToAttributeMap)
+		{
+			for (const FName& Tag : Tags)
+			{
+				PCG::Private::FParseTagResult TagData(Tag);
+				if (TagData.IsValid() && TagToAttributeMap->Contains(FName(TagData.GetOriginalAttribute())))
+				{
+					PCG::Private::SetAttributeFromTag(TagData, Metadata, MetadataEntry);
+				}
+			}
+		}
+	};
+
+	int PointWriteIndex = NumPointsBefore;
+
+	for (UPrimitiveComponent* PrimComp : PrimitiveComponents)
+	{
+		if (bIgnorePCGCreatedComponents && PrimComp->ComponentTags.Contains(PCGHelpers::DefaultPCGTag))
+		{
+			continue;
+		}
+
+		// Skip billboard components - they have no meaningful geometry
+		if (Cast<UBillboardComponent>(PrimComp))
+		{
+			continue;
+		}
+
+		const FSoftObjectPath ComponentPath(PrimComp);
+
+		if (UStaticMeshComponent* SMC = Cast<UStaticMeshComponent>(PrimComp))
+		{
+			UStaticMesh* StaticMesh = SMC->GetStaticMesh();
+			if (!StaticMesh)
+			{
+				continue;
+			}
+
+			// Prepare value keys.
+			PCGMetadataValueKey ComponentReferenceValueKey = ComponentReferenceAttribute ? ComponentReferenceAttribute->AddValue(ComponentPath) : PCGNotFoundValueKey;
+			PCGMetadataValueKey MeshPathValueKey = MeshAttribute ? MeshAttribute->AddValue(FSoftObjectPath(StaticMesh)) : PCGNotFoundValueKey;
+			PCGMetadataValueKey MaterialPathValueKey = PCGNotFoundValueKey;
+
+			if(MaterialAttribute)
+			{
+				TArray<UMaterialInterface*> Materials = SMC->GetMaterials();
+				if (!Materials.IsEmpty())
+				{
+					// Walk up the material instance chain to find the nearest public asset
+					UMaterialInterface* Material = Materials[0];
+					UMaterialInstance* MaterialInstance = Cast<UMaterialInstance>(Material);
+					while (MaterialInstance && !MaterialInstance->IsAsset())
+					{
+						Material = MaterialInstance->Parent;
+						MaterialInstance = Cast<UMaterialInstance>(Material);
+					}
+
+					MaterialPathValueKey = MaterialAttribute->AddValue(FSoftObjectPath(Material));
+				}
+			}
+
+			const FBox MeshBounds = StaticMesh->GetBoundingBox();
+
+			if (UInstancedStaticMeshComponent* ISMC = Cast<UInstancedStaticMeshComponent>(SMC))
+			{
+				const int32 InstanceCount = ISMC->GetNumInstances();
+				for (int32 I = 0; I < InstanceCount; ++I)
+				{
+					FTransform InstanceTransform;
+					ISMC->GetInstanceTransform(I, InstanceTransform, /*bWorldSpace=*/true);
+					WritePointAttributes(PointWriteIndex + I, InstanceTransform, ComponentReferenceValueKey, MeshPathValueKey, MeshBounds, MaterialPathValueKey);
+				}
+
+				PointWriteIndex += InstanceCount;
+			}
+			else
+			{
+				WritePointAttributes(PointWriteIndex, SMC->GetComponentTransform(), ComponentReferenceValueKey, MeshPathValueKey, MeshBounds, MaterialPathValueKey);
+				++PointWriteIndex;
+			}
+		}
+		else if (!bOnlyStaticMeshComponents)
+		{
+			// Non-SM primitive: single point at the component transform with component-local bounds
+			const FBoxSphereBounds LocalBounds = PrimComp->CalcLocalBounds();
+			PCGMetadataValueKey ComponentReferenceValueKey = ComponentReferenceAttribute ? ComponentReferenceAttribute->AddValue(ComponentPath) : PCGNotFoundValueKey;
+			WritePointAttributes(PointWriteIndex, PrimComp->GetComponentTransform(), ComponentReferenceValueKey, PCGNotFoundValueKey, LocalBounds.GetBox(), PCGNotFoundValueKey);
+			++PointWriteIndex;
+		}
+	}
+
+	check(PointWriteIndex == NumPointsBefore + NumPointsToAdd);
+
+	return NumPointsToAdd;
+}
diff --git a/Engine/Plugins/PCG/Source/PCG/Private/Data/Registry/PCGGetDataFunctionRegistry.cpp b/Engine/Plugins/PCG/Source/PCG/Private/Data/Registry/PCGGetDataFunctionRegistry.cpp
--- a/Engine/Plugins/PCG/Source/PCG/Private/Data/Registry/PCGGetDataFunctionRegistry.cpp
+++ b/Engine/Plugins/PCG/Source/PCG/Private/Data/Registry/PCGGetDataFunctionRegistry.cpp
@@ -21,6 +21,7 @@
 #include "Data/PCGUnionData.h"
 #include "Data/PCGVirtualTextureData.h"
 #include "Data/PCGVolumeData.h"
+#include "Data/Registry/PCGDataFunctionRegistryHelpers.h"
 #include "Grid/PCGPartitionActor.h"
 #include "Helpers/PCGActorHelpers.h"
 #include "Helpers/PCGBlueprintHelpers.h"
@@ -131,6 +132,44 @@ int FPCGGetDataFunctionRegistry::DefaultDataFromActor(FPCGContext* InContext, co
 		TaggedData.Data = Data;
 		TaggedData.Tags = ActorTags;
 	}
+	else if (InDataFilter == EPCGDataType::Point)
+	{
+		// Implementation note: when we're asking explicitly for point data, this behaves like the 'level to asset' workflow
+		// but operates on all primitive components (not just SMCs).
+		UPCGBasePointData* PointData = FPCGContext::NewPointData_AnyThread(InContext);
+		UPCGMetadata* PointMetadata = PointData->MutableMetadata();
+		PointData->AllocateProperties(EPCGPointNativeProperties::Transform | EPCGPointNativeProperties::MetadataEntry | EPCGPointNativeProperties::BoundsMin | EPCGPointNativeProperties::BoundsMax | EPCGPointNativeProperties::Seed);
+
+		FPCGMetadataAttribute<FSoftObjectPath>* MeshAttribute = PointMetadata->CreateAttribute<FSoftObjectPath>(PCGLevelToAssetConstants::MeshAttributeName, FSoftObjectPath(), /*bAllowsInterpolation=*/false, /*bOverrideParent=*/true);
+		FPCGMetadataAttribute<FSoftObjectPath>* MaterialAttribute = PointMetadata->CreateAttribute<FSoftObjectPath>(PCGLevelToAssetConstants::MaterialAttributeName, FSoftObjectPath(), /*bAllowsInterpolation=*/false, /*bOverrideParent=*/true);
+		FPCGMetadataAttribute<FSoftObjectPath>* ActorReferenceAttribute = PointMetadata->CreateAttribute<FSoftObjectPath>(PCGAttributeNameConstants::ActorReferenceAttribute, FSoftObjectPath(), /*bAllowsInterpolation=*/false, /*bOverrideParent=*/true);
+		FPCGMetadataAttribute<FSoftObjectPath>* ComponentReferenceAttribute = PointMetadata->CreateAttribute<FSoftObjectPath>(PCGAttributeNameConstants::ComponentReferenceAttribute, FSoftObjectPath(), /*bAllowsInterpolation=*/false, /*bOverrideParent=*/true);
+
+		TMap<FName, FPCGMetadataAttributeBase*> TagToAttributeMap;
+		TSet<FName> SanitizedAttributeNames;
+
+		FPCGActorMeshParsingContext Context;
+		Context.PointData = PointData;
+		Context.PointMetadata = PointMetadata;
+		Context.MeshAttribute = MeshAttribute;
+		Context.MaterialAttribute = MaterialAttribute;
+		Context.ActorReferenceAttribute = ActorReferenceAttribute;
+		Context.ComponentReferenceAttribute = ComponentReferenceAttribute;
+		Context.TagToAttributeMap = &TagToAttributeMap;
+		Context.SanitizedAttributeNames = &SanitizedAttributeNames;
+		Context.bOnlyStaticMeshComponents = false;
+		Context.bIgnorePCGCreatedComponents = InParams.bIgnorePCGGeneratedComponents;
+
+		if (Context.ParseActorComponents(Actor) == 0)
+		{
+			// Fallback on standard single point per actor path
+			PointData->AddSinglePointFromActor(Actor, &Output.bSanitizedTagAttributeNames);
+		}
+
+		FPCGTaggedData& TaggedData = Output.Collection.TaggedData.Emplace_GetRef();
+		TaggedData.Data = PointData;
+		TaggedData.Tags = ActorTags;
+	}
 	else if (PartitionActor)
 	{
 		// If we have a Partition Actor, the execution source must be a PCG Component.
diff --git a/Engine/Plugins/PCG/Source/PCG/Private/Elements/PCGDataFromActor.cpp b/Engine/Plugins/PCG/Source/PCG/Private/Elements/PCGDataFromActor.cpp
--- a/Engine/Plugins/PCG/Source/PCG/Private/Elements/PCGDataFromActor.cpp
+++ b/Engine/Plugins/PCG/Source/PCG/Private/Elements/PCGDataFromActor.cpp
@@ -1093,7 +1093,9 @@ void FPCGDataFromActorElement::ProcessActor(FPCGContext* InContext, const UPCGDa
 	else
 	{
 		const bool bParseActor = (Settings->Mode != EPCGGetDataFromActorMode::GetSinglePoint);
-		FPCGDataCollection ActorDataCollection = GetActorDataCollection(FoundActor, Settings->GetDataFilter(), bParseActor);
+		EPCGDataType TargetType = (Settings->Mode == EPCGGetDataFromActorMode::GetPointData ? EPCGDataType::Point : Settings->GetDataFilter());
+
+		FPCGDataCollection ActorDataCollection = GetActorDataCollection(FoundActor, TargetType, bParseActor);
 		Outputs += ActorDataCollection.TaggedData;
 	}
 
diff --git a/Engine/Plugins/PCG/Source/PCG/Public/Data/Registry/PCGDataFunctionRegistryHelpers.h b/Engine/Plugins/PCG/Source/PCG/Public/Data/Registry/PCGDataFunctionRegistryHelpers.h
new file mode 100644
--- /dev/null
+++ b/Engine/Plugins/PCG/Source/PCG/Public/Data/Registry/PCGDataFunctionRegistryHelpers.h
@@ -0,0 +1,68 @@
+// Copyright Epic Games, Inc. All Rights Reserved.
+
+#pragma once
+
+#include "PCGCommon.h"
+
+#define UE_API PCG_API
+
+class AActor;
+class UPCGBasePointData;
+class UPCGMetadata;
+class FPCGMetadataAttributeBase;
+template<typename T> class FPCGMetadataAttribute;
+
+/**
+ * Context passed to FPCGActorMeshParsingContext::ParseActorComponents.
+ * When parsing multiple actors into the same PointData, this struct should be shared
+ * across all calls so that tag attribute creation is deduplicated.
+ */
+struct FPCGActorMeshParsingContext
+{
+	// Required output target
+	UPCGBasePointData* PointData     = nullptr;
+	UPCGMetadata*      PointMetadata = nullptr; // if null, derived from PointData->MutableMetadata()
+
+	// Attribute pointers - null means that attribute will not be written
+	FPCGMetadataAttribute<FSoftObjectPath>* MeshAttribute = nullptr;
+	FPCGMetadataAttribute<FSoftObjectPath>* MaterialAttribute = nullptr;
+	FPCGMetadataAttribute<int64>* ActorIndexAttribute = nullptr;
+	FPCGMetadataAttribute<int64>* ParentIndexAttribute = nullptr;
+	FPCGMetadataAttribute<FTransform>* RelativeTransformAttribute = nullptr;
+	FPCGMetadataAttribute<int64>* HierarchyDepthAttribute = nullptr;
+	FPCGMetadataAttribute<FSoftObjectPath>* ActorReferenceAttribute = nullptr;
+	FPCGMetadataAttribute<FSoftObjectPath>* ComponentReferenceAttribute = nullptr;
+
+	// Per-actor values - caller must set these before each call
+	int64 ActorIndex = -1;
+	int64 ParentIndex = 0;
+	FTransform RelativeTransform = FTransform::Identity;
+	int64 HierarchyDepth = 1; // only used if HierarchyDepthAttribute != nullptr
+	const TArray<FName>* ActorTags = nullptr; // if null, Actor->Tags is used
+
+	// Tag attribute tracking - shared across all calls for the same PointData
+	TMap<FName, FPCGMetadataAttributeBase*>* TagToAttributeMap = nullptr;
+	TSet<FName>* SanitizedAttributeNames = nullptr;
+
+	// Processing flags
+	bool bOnlyStaticMeshComponents = false; // restricts to SMC/ISMC only when true (LevelToAsset behavior)
+	bool bIgnorePCGCreatedComponents = false;
+
+	/**
+	 * Creates metadata attributes on PointMetadata for any of the actor's tags that are not yet
+	 * tracked in TagToAttributeMap, skipping any that would collide with pre-existing attribute pointers.
+	 * No-op if TagToAttributeMap or SanitizedAttributeNames are null.
+	 */
+	UE_API void CreateTagAttributes(AActor* InActor);
+
+	/**
+	 * Parses the primitive components of InActor and writes points into PointData.
+	 * For UStaticMeshComponent/UInstancedStaticMeshComponent, creates one point per mesh instance.
+	 * For other UPrimitiveComponent types (when bOnlyStaticMeshComponents is false), creates a single
+	 * point at the component transform using the component's local bounds.
+	 * Returns the number of points added.
+	 */
+	UE_API int ParseActorComponents(AActor* InActor);
+};
+
+#undef UE_API
diff --git a/Engine/Plugins/PCG/Source/PCG/Public/Elements/PCGDataFromActor.h b/Engine/Plugins/PCG/Source/PCG/Public/Elements/PCGDataFromActor.h
--- a/Engine/Plugins/PCG/Source/PCG/Public/Elements/PCGDataFromActor.h
+++ b/Engine/Plugins/PCG/Source/PCG/Public/Elements/PCGDataFromActor.h
@@ -22,7 +22,8 @@ enum class EPCGGetDataFromActorMode : uint8
 	GetDataFromPCGComponent UMETA(Tooltip = "Copy generated output from other PCG components on the found actor(s)."),
 	GetDataFromPCGComponentOrParseComponents UMETA(Tooltip = "Attempts to copy generated output from other PCG components on the found actor(s), otherwise, falls back to parsing actor components."),
 	GetActorReference UMETA(Tooltip = "Produces one entry per actor with only the actor reference."),
-	GetComponentsReference UMETA(Tooltip = "Produces one entry per component within the actor selection.")
+	GetComponentsReference UMETA(Tooltip = "Produces one entry per component within the actor selection."),
+	GetPointData UMETA(Tooltip = "Produces a point data containing a point per component, except for ISMs that will have a point per instance. Similar behavior to the PCG Data Asset format.")
 };
 
 /** Builds a collection of PCG-compatible data from the selected actors. */
@@ -92,7 +93,7 @@ public:
 	EPCGGetDataFromActorMode Mode = EPCGGetDataFromActorMode::ParseActorComponents;
 
 	/** Ignores any component that was spawned by PCG. */
-	UPROPERTY(BlueprintReadWrite, EditAnywhere, Category = "Data Retrieval Settings", meta = (EditCondition = "Mode == EPCGGetDataFromActorMode::ParseActorComponents || Mode == EPCGGetDataFromActorMode::GetComponentsReference || Mode == EPCGGetDataFromActorMode::GetDataFromPCGComponentOrParseComponents", EditConditionHides))
+	UPROPERTY(BlueprintReadWrite, EditAnywhere, Category = "Data Retrieval Settings", meta = (EditCondition = "Mode == EPCGGetDataFromActorMode::ParseActorComponents || Mode == EPCGGetDataFromActorMode::GetPointData || Mode == EPCGGetDataFromActorMode::GetComponentsReference || Mode == EPCGGetDataFromActorMode::GetDataFromPCGComponentOrParseComponents", EditConditionHides))
 	bool bIgnorePCGGeneratedComponents = true;
 
 	/** Also produces a single point data at the actor location. */
diff --git a/Engine/Plugins/PCG/Source/PCG/Public/PCGCommon.h b/Engine/Plugins/PCG/Source/PCG/Public/PCGCommon.h
--- a/Engine/Plugins/PCG/Source/PCG/Public/PCGCommon.h
+++ b/Engine/Plugins/PCG/Source/PCG/Public/PCGCommon.h
@@ -933,4 +933,18 @@ enum class EPCGTextureFilter : uint8
 	Bilinear UMETA(Tooltip="Bilinearly interpolates the values of the four nearest texels to the sample location.")
 };
 
+namespace PCGLevelToAssetConstants
+{
+	/** Attribute names shared between the LevelToAsset workflow and the GetDataFunctionRegistry actor parsing. */
+	const FLazyName MaterialAttributeName = TEXT("Material");
+	const FLazyName MeshAttributeName = TEXT("Mesh");
+	const FLazyName HierarchyDepthAttributeName = TEXT("HierarchyDepth");
+	const FLazyName ActorIndexAttributeName = TEXT("ActorIndex");
+	const FLazyName ParentIndexAttributeName = TEXT("ParentIndex");
+	const FLazyName RelativeTransformAttributeName = TEXT("RelativeTransform");
+	const FLazyName IgnoreParentRotationAttributeName = TEXT("IgnoreParentRotation");
+	const FLazyName IgnoreParentScaleAttributeName = TEXT("IgnoreParentScale");
+	const FLazyName ExcludedActorTag = TEXT("PCGLevelToAssetExclude");
+}
+
 #undef UE_API
diff --git a/Engine/Plugins/PCG/Source/PCG/Public/PCGDataAsset.h b/Engine/Plugins/PCG/Source/PCG/Public/PCGDataAsset.h
--- a/Engine/Plugins/PCG/Source/PCG/Public/PCGDataAsset.h
+++ b/Engine/Plugins/PCG/Source/PCG/Public/PCGDataAsset.h
@@ -9,19 +9,6 @@
 
 class UPCGDataAsset;
 
-namespace PCGLevelToAssetConstants
-{
-	const FName MaterialAttributeName = TEXT("Material");
-	const FName MeshAttributeName = TEXT("Mesh");
-	const FName HierarchyDepthAttributeName = TEXT("HierarchyDepth");
-	const FName ActorIndexAttributeName = TEXT("ActorIndex");
-	const FName ParentIndexAttributeName = TEXT("ParentIndex");
-	const FName RelativeTransformAttributeName = TEXT("RelativeTransform");
-	const FName IgnoreParentRotationAttributeName = TEXT("IgnoreParentRotation");
-	const FName IgnoreParentScaleAttributeName = TEXT("IgnoreParentScale");
-	const FName ExcludedActorTag = TEXT("PCGLevelToAssetExclude");
-}
-
 USTRUCT(BlueprintInternalUseOnly)
 struct FPCGDataAssetCachedPins
 {
diff --git a/Engine/Plugins/PCG/Source/PCGEditor/Private/PCGLevelToAsset.cpp b/Engine/Plugins/PCG/Source/PCGEditor/Private/PCGLevelToAsset.cpp
--- a/Engine/Plugins/PCG/Source/PCGEditor/Private/PCGLevelToAsset.cpp
+++ b/Engine/Plugins/PCG/Source/PCGEditor/Private/PCGLevelToAsset.cpp
@@ -10,9 +10,9 @@
 #include "Data/PCGBasePointData.h"
 #include "Data/PCGPointData.h"
 #include "Data/PCGPointArrayData.h"
+#include "Data/Registry/PCGDataFunctionRegistryHelpers.h"
 #include "Helpers/PCGActorHelpers.h"
 #include "Helpers/PCGHelpers.h"
-#include "Helpers/PCGTagHelpers.h"
 #include "Metadata/PCGMetadata.h"
 
 #include "CoreGlobals.h"
@@ -20,12 +20,11 @@
 #include "FileHelpers.h"
 #include "IContentBrowserSingleton.h"
 #include "AssetRegistry/AssetRegistryModule.h"
-#include "Components/InstancedStaticMeshComponent.h"
+#include "Components/StaticMeshComponent.h"
 #include "Engine/LevelStreaming.h"
 #include "LevelInstance/LevelInstanceEditorInstanceActor.h"
 #include "LevelInstance/LevelInstanceInterface.h"
 #include "LevelInstance/LevelInstanceSubsystem.h"
-#include "Materials/MaterialInstance.h"
 #include "UObject/UObjectGlobals.h"
 #include "UObject/Package.h"
 #include "WorldPartition/WorldPartition.h"
@@ -209,20 +208,12 @@ bool UPCGLevelToAsset::BP_ExportWorld_Implementation(UWorld* World, const FStrin
 	FPCGMetadataAttribute<int64>* ParentIndexAttribute = PointMetadata->CreateAttribute<int64>(ParentIndexAttributeName, -1, /*bAllowsInterpolation=*/false, /*bOverrideParent=*/true);
 	FPCGMetadataAttribute<FTransform>* RelativeTransformAttribute = PointMetadata->CreateAttribute<FTransform>(RelativeTransformAttributeName, FTransform::Identity, /*bAllowsInterpolation=*/false, /*bOverrideParent=*/true);
 
-	// Map from raw/unsanitized tag name to corresponding attribute.
+	// Map from raw/unsanitized tag name to corresponding attribute - shared across all actor iterations.
 	TMap<FName, FPCGMetadataAttributeBase*> TagToAttributeMap;
 
 	// Relationship Tag:SanitizedName is many:1, so keep track of which sanitized names are created so we don't attempt to create the same one multiple times.
 	TSet<FName> SanitizedAttributeNames;
 
-	TSet<FName> ReservedTags;
-	ReservedTags.Add(MaterialAttributeName);
-	ReservedTags.Add(MeshAttributeName);
-	ReservedTags.Add(HierarchyDepthAttributeName);
-	ReservedTags.Add(ActorIndexAttributeName);
-	ReservedTags.Add(ParentIndexAttributeName);
-	ReservedTags.Add(RelativeTransformAttributeName);
-
 	// Hierarchy root point
 	{
 		PointData->SetNumPoints(1);
@@ -419,8 +410,23 @@ bool UPCGLevelToAsset::BP_ExportWorld_Implementation(UWorld* World, const FStrin
 		return true;
 	});
 
+	// Prepare actor parsing context - we'll set only the non-actor relevant things here.
+	FPCGActorMeshParsingContext Context;
+	Context.PointData = PointData;
+	Context.PointMetadata = PointMetadata;
+	Context.MeshAttribute = MeshAttribute;
+	Context.MaterialAttribute = MaterialAttribute;
+	Context.ActorIndexAttribute = ActorIndexAttribute;
+	Context.ParentIndexAttribute = ParentIndexAttribute;
+	Context.RelativeTransformAttribute = RelativeTransformAttribute;
+	Context.HierarchyDepthAttribute = HierarchyDepthAttribute;
+	Context.TagToAttributeMap = &TagToAttributeMap;
+	Context.SanitizedAttributeNames = &SanitizedAttributeNames;
+	Context.bOnlyStaticMeshComponents = true;
+	Context.bIgnorePCGCreatedComponents = false;
+
 	// Create points
-	UPCGActorHelpers::ForEachActorInWorld(World, AActor::StaticClass(), [&](AActor* Actor)
+	UPCGActorHelpers::ForEachActorInWorld(World, AActor::StaticClass(), [&ExcludedActors, &AllActorBounds, &ActorIndexMap, &GetAttachParentActor, &Context](AActor* Actor)
 	{
 		// TODO Actor-level decisions if any; if the actor is "consumed" at this step, make sure to update AllActorBounds as well.
 
@@ -429,57 +435,19 @@ bool UPCGLevelToAsset::BP_ExportWorld_Implementation(UWorld* World, const FStrin
 			return true;
 		}
 
-		// Otherwise, parse "known" actor components
+		// Only process actors with static mesh components.
 		TArray<UStaticMeshComponent*> SMCs;
 		Actor->GetComponents(SMCs);
-
-		if (SMCs.IsEmpty()) // early out
+		if (SMCs.IsEmpty())
 		{
 			return true;
 		}
 
-		const FBox ActorBounds = PCGHelpers::GetActorBounds(Actor, /*bIgnorePCGCreatedComponents=*/true);
-		AllActorBounds += ActorBounds;
+		// implementation note: we're not ignoring pcg generated components here because they'll be parsed regardless, and the PCG component itself would be lost
+		AllActorBounds += PCGHelpers::GetActorBounds(Actor, /*bIgnorePCGCreatedComponents=*/false);
 
 		const FActorEntry& ActorEntry = ActorIndexMap[Actor];
-
-		for (FName ActorTag : ActorEntry.Tags)
-		{
-			PCG::Private::FParseTagResult TagData(ActorTag);
-
-			if (!TagData.IsValid())
-			{
-				continue;
-			}
-
-			FName OriginalAttributeName(TagData.GetOriginalAttribute());
-			FName SanitizedAttributeName(TagData.Attribute);
-
-			// Check if we can safely skip that tag
-			if (ReservedTags.Contains(SanitizedAttributeName) || TagToAttributeMap.Contains(OriginalAttributeName) || SanitizedAttributeNames.Contains(SanitizedAttributeName))
-			{
-				continue;
-			}
-
-			// Try to create the attribute
-			if (!PCG::Private::CreateAttributeFromTag(TagData, PointMetadata))
-			{
-				continue;
-			}
-
-			// Log warning if we sanitized some values
-			if (TagData.HasBeenSanitized())
-			{
-				UE_LOG(LogPCGEditor, Warning, TEXT("Sanitized tag string on actor '%s' to remove invalid characters: '%s' -> '%s'"), *Actor->GetName(), *TagData.GetOriginalAttribute(), *TagData.Attribute);
-			}
-
-			TagToAttributeMap.Add(OriginalAttributeName, PointMetadata->GetMutableAttribute(SanitizedAttributeName));
-			SanitizedAttributeNames.Add(SanitizedAttributeName);
-		}
-
-		// Prepare actor-level data that's propagated to all points
 		const FTransform& ActorTransform = Actor->GetTransform();
-		const int64 ActorIndex = ActorEntry.Index;
 		AActor* ParentActor = ActorEntry.AttachmentParent;
 		const int64 ParentActorIndex = ParentActor ? ActorIndexMap[ParentActor].Index : 0;
 		const FTransform RelativeTransform = ParentActor ? ActorTransform.GetRelativeTransform(ParentActor->GetTransform()) : ActorTransform;
@@ -492,94 +460,14 @@ bool UPCGLevelToAsset::BP_ExportWorld_Implementation(UWorld* World, const FStrin
 			ParentActor = GetAttachParentActor(ParentActor);
 		}
 
-		auto MakePoint = [&](int32 PointIndex, const FTransform& Transform, const FSoftObjectPath& MeshPath, const FBox& MeshBounds, const TArray<UMaterialInterface*>& MeshMaterials)
-		{
-			TPCGValueRange<FTransform> TransformRange = PointData->GetTransformValueRange();
-			TPCGValueRange<int32> SeedRange = PointData->GetSeedValueRange();
-			TPCGValueRange<FVector> BoundsMinRange = PointData->GetBoundsMinValueRange();
-			TPCGValueRange<FVector> BoundsMaxRange = PointData->GetBoundsMaxValueRange();
-			TPCGValueRange<int64> MetadataEntryRange = PointData->GetMetadataEntryValueRange();
-
-			TransformRange[PointIndex] = Transform;
-			SeedRange[PointIndex] = PCGHelpers::ComputeSeedFromPosition(Transform.GetLocation());
-			BoundsMinRange[PointIndex] = MeshBounds.Min;
-			BoundsMaxRange[PointIndex] = MeshBounds.Max;
-			
-			int64& MetadataEntry = MetadataEntryRange[PointIndex];
-			PointMetadata->InitializeOnSet(MetadataEntry);
-			MeshAttribute->SetValue(MetadataEntry, MeshPath);
-
-			if (!MeshMaterials.IsEmpty())
-			{
-				// Avoid references to materials that are not public assets to avoid invalid references, in this case we reference the first parent asset we find.
-				UMaterialInterface* Material = MeshMaterials[0];
-				UMaterialInstance* MaterialInstance = Cast<UMaterialInstance>(Material);
-				while (MaterialInstance && !MaterialInstance->IsAsset())
-				{
-					Material = MaterialInstance->Parent;
-					MaterialInstance = Cast<UMaterialInstance>(Material);
-				}
-
-				MaterialAttribute->SetValue(MetadataEntry, FSoftObjectPath(Material));
-			}
-
-			ActorIndexAttribute->SetValue(MetadataEntry, ActorIndex);
-			ParentIndexAttribute->SetValue(MetadataEntry, ParentActorIndex);
-			RelativeTransformAttribute->SetValue(MetadataEntry, RelativeTransform);
-			HierarchyDepthAttribute->SetValue(MetadataEntry, HierarchyDepth);
-
-			// For all tags, if the tag is of format 'Name:Value' then create attribute Name and assign Value, otherwise
-			// create a boolean attribute with the name given by the sanitized tag string.
-			for (FName ActorTag : ActorEntry.Tags)
-			{
-				// Implementation note: we don't set values directly from the tag in the eventuality that there are some name collisions
-				PCG::Private::FParseTagResult TagData(ActorTag);
-
-				if (TagData.IsValid() && TagToAttributeMap.Contains(FName(TagData.GetOriginalAttribute())))
-				{
-					PCG::Private::SetAttributeFromTag(TagData, PointMetadata, MetadataEntry);
-				}
-			}
-		};
-
-		for (UStaticMeshComponent* SMC : SMCs)
-		{
-			TObjectPtr<UStaticMesh> StaticMesh = SMC->GetStaticMesh();
-
-			if (!StaticMesh)
-			{
-				continue;
-			}
-
-			const FSoftObjectPath MeshPath(StaticMesh);
-			const FBox MeshBounds = StaticMesh->GetBoundingBox();
-			TArray<UMaterialInterface*> Materials = SMC->GetMaterials();
-
-			const int FirstIndex = PointData->GetNumPoints();
-
-			// For all instances (or a single instance if this is a static mesh and not an ISM)
-			// if a static mesh -> use actor transform (which might be wrong?)
-			// if ISM -> get instance transform in world space
-			if (UInstancedStaticMeshComponent* ISMC = Cast<UInstancedStaticMeshComponent>(SMC))
-			{
-				const int InstanceCount = ISMC->GetNumInstances();
-				
-				PointData->SetNumPoints(PointData->GetNumPoints() + InstanceCount);
-				for (int I = 0; I < InstanceCount; ++I)
-				{
-					FTransform InstanceTransform;
-					ISMC->GetInstanceTransform(I, InstanceTransform, /*bWorldSpace=*/true);
-
-					MakePoint(FirstIndex+I, InstanceTransform, MeshPath, MeshBounds, Materials);
-				}
-			}
-			else
-			{
-				PointData->SetNumPoints(PointData->GetNumPoints() + 1);
-				MakePoint(FirstIndex, SMC->GetComponentTransform(), MeshPath, MeshBounds, Materials);
-			}
-		}
+		// Setup the per-actor data in the actor parsing context struct
+		Context.ActorIndex = ActorEntry.Index;
+		Context.ParentIndex = ParentActorIndex;
+		Context.RelativeTransform = RelativeTransform;
+		Context.HierarchyDepth = HierarchyDepth;
+		Context.ActorTags = &ActorEntry.Tags;
 
+		Context.ParseActorComponents(Actor);
 		return true;
 	});
 
