[PCG] Add GPU override support to TransformPoints

#Synced-CL 52356503

diff --git a/Engine/Plugins/PCG/Shaders/Private/Elements/PCGTransformPoints.usf b/Engine/Plugins/PCG/Shaders/Private/Elements/PCGTransformPoints.usf
--- a/Engine/Plugins/PCG/Shaders/Private/Elements/PCGTransformPoints.usf
+++ b/Engine/Plugins/PCG/Shaders/Private/Elements/PCGTransformPoints.usf
@@ -33,20 +33,20 @@ void PCGTransformPointsCS(uint3 GroupId : SV_GroupID, uint GroupIndex : SV_Group
 	}
 
 	// Transform points code
-	const float3 OffsetMin = TransformPoints_GetOffsetMin();
-	const float3 OffsetMax = TransformPoints_GetOffsetMax();
-	const bool bAbsoluteOffset = TransformPoints_GetAbsoluteOffset() != 0;
-	const float3 RotationMin = TransformPoints_GetRotationMin();
-	const float3 RotationMax = TransformPoints_GetRotationMax();
-	const bool bAbsoluteRotation = TransformPoints_GetAbsoluteRotation() != 0;
-	const float3 ScaleMin = TransformPoints_GetScaleMin();
-	const float3 ScaleMax = TransformPoints_GetScaleMax();
-	const bool bAbsoluteScale = TransformPoints_GetAbsoluteScale() != 0;
-	const bool bUniformScale = TransformPoints_GetUniformScale() != 0;
-	const bool bRecomputeSeed = TransformPoints_GetRecomputeSeed() != 0;
+	const float3 OffsetMin = OffsetMin_GetOverridableValue();
+	const float3 OffsetMax = OffsetMax_GetOverridableValue();
+	const bool bAbsoluteOffset = bAbsoluteOffset_GetOverridableValue() != 0;
+	const float3 RotationMin = RotationMin_GetOverridableValue();
+	const float3 RotationMax = RotationMax_GetOverridableValue();
+	const bool bAbsoluteRotation = bAbsoluteRotation_GetOverridableValue() != 0;
+	const float3 ScaleMin = ScaleMin_GetOverridableValue();
+	const float3 ScaleMax = ScaleMax_GetOverridableValue();
+	const bool bAbsoluteScale = bAbsoluteScale_GetOverridableValue() != 0;
+	const bool bUniformScale = bUniformScale_GetOverridableValue() != 0;
+	const bool bRecomputeSeed = bRecomputeSeed_GetOverridableValue() != 0;
 	const int AttributeId = TransformPoints_GetAttributeId();
 	const bool bApplyToAttribute = (TransformPoints_GetApplyToAttribute() != 0) && (AttributeId != -1);
-	
+
 	float3 InPosition;
 	FQuat InRotation;
 	float3 InScale;
@@ -92,14 +92,14 @@ void PCGTransformPointsCS(uint3 GroupId : SV_GroupID, uint GroupIndex : SV_Group
 
 	// Rotation
 	FQuat OutRotation = bAbsoluteRotation ? FQuat(0, 0, 0, 1) : InRotation;
-	float3 RandomRotation = float3(0, 0, 0);
-
-	// Note: Random numbers drawn in same order as CPU.
-	RandomRotation.y = lerp(RotationMin.y, RotationMax.y, FRand(Seed));
-	RandomRotation.z = lerp(RotationMin.z, RotationMax.z, FRand(Seed));
-	RandomRotation.x = lerp(RotationMin.x, RotationMax.x, FRand(Seed));
-
-	OutRotation = Quat_Multiply(OutRotation, Quat_FromRotator(RandomRotation));
+	// RotationMin/Max are in raw FRotator order: (Pitch, Yaw, Roll).
+	// Random numbers drawn in same order as CPU: Pitch, Yaw, Roll.
+	const float RandomPitch = lerp(RotationMin.x, RotationMax.x, FRand(Seed));
+	const float RandomYaw   = lerp(RotationMin.y, RotationMax.y, FRand(Seed));
+	const float RandomRoll  = lerp(RotationMin.z, RotationMax.z, FRand(Seed));
+
+	// Quat_FromRotator expects Euler order: (Roll, Pitch, Yaw).
+	OutRotation = Quat_Multiply(OutRotation, Quat_FromRotator(float3(RandomRoll, RandomPitch, RandomYaw)));
 
 	// Scale
 	float3 OutScale = bAbsoluteScale ? float3(1, 1, 1) : InScale;
diff --git a/Engine/Plugins/PCG/Source/PCG/Private/Compute/DataInterfaces/Elements/PCGTransformPointsDataInterface.cpp b/Engine/Plugins/PCG/Source/PCG/Private/Compute/DataInterfaces/Elements/PCGTransformPointsDataInterface.cpp
--- a/Engine/Plugins/PCG/Source/PCG/Private/Compute/DataInterfaces/Elements/PCGTransformPointsDataInterface.cpp
+++ b/Engine/Plugins/PCG/Source/PCG/Private/Compute/DataInterfaces/Elements/PCGTransformPointsDataInterface.cpp
@@ -5,6 +5,7 @@
 #include "PCGContext.h"
 #include "Compute/PCGComputeGraph.h"
 #include "Compute/PCGDataBinding.h"
+#include "Compute/PCGKernelParamsCache.h"
 #include "Elements/PCGTransformPoints.h"
 #include "Compute/Elements/PCGTransformPointsKernel.h"
 
@@ -24,50 +25,10 @@
 
 void UPCGTransformPointsDataInterface::GetSupportedInputs(TArray<FShaderFunctionDefinition>& OutFunctions) const
 {
-	OutFunctions.AddDefaulted_GetRef()
-		.SetName(TEXT("TransformPoints_GetOffsetMin"))
-		.AddReturnType(FShaderValueType::Get(EShaderFundamentalType::Float, 3));
-
-	OutFunctions.AddDefaulted_GetRef()
-		.SetName(TEXT("TransformPoints_GetOffsetMax"))
-		.AddReturnType(FShaderValueType::Get(EShaderFundamentalType::Float, 3));
-
-	OutFunctions.AddDefaulted_GetRef()
-		.SetName(TEXT("TransformPoints_GetAbsoluteOffset"))
-		.AddReturnType(FShaderValueType::Get(EShaderFundamentalType::Uint, 1));
-
-	OutFunctions.AddDefaulted_GetRef()
-		.SetName(TEXT("TransformPoints_GetRotationMin"))
-		.AddReturnType(FShaderValueType::Get(EShaderFundamentalType::Float, 3));
-
-	OutFunctions.AddDefaulted_GetRef()
-		.SetName(TEXT("TransformPoints_GetRotationMax"))
-		.AddReturnType(FShaderValueType::Get(EShaderFundamentalType::Float, 3));
-
-	OutFunctions.AddDefaulted_GetRef()
-		.SetName(TEXT("TransformPoints_GetAbsoluteRotation"))
-	.AddReturnType(FShaderValueType::Get(EShaderFundamentalType::Uint, 1));
-
-	OutFunctions.AddDefaulted_GetRef()
-		.SetName(TEXT("TransformPoints_GetScaleMin"))
-		.AddReturnType(FShaderValueType::Get(EShaderFundamentalType::Float, 3));
-
-	OutFunctions.AddDefaulted_GetRef()
-		.SetName(TEXT("TransformPoints_GetScaleMax"))
-		.AddReturnType(FShaderValueType::Get(EShaderFundamentalType::Float, 3));
-
-	OutFunctions.AddDefaulted_GetRef()
-		.SetName(TEXT("TransformPoints_GetAbsoluteScale"))
-		.AddReturnType(FShaderValueType::Get(EShaderFundamentalType::Uint, 1));
-
-	OutFunctions.AddDefaulted_GetRef()
-		.SetName(TEXT("TransformPoints_GetUniformScale"))
-		.AddReturnType(FShaderValueType::Get(EShaderFundamentalType::Uint, 1));
-
-	OutFunctions.AddDefaulted_GetRef()
-		.SetName(TEXT("TransformPoints_GetRecomputeSeed"))
-		.AddReturnType(FShaderValueType::Get(EShaderFundamentalType::Uint, 1));
+	// Overridable params are registered by the base class.
+	Super::GetSupportedInputs(OutFunctions);
 
+	// Non-overridable params.
 	OutFunctions.AddDefaulted_GetRef()
 		.SetName(TEXT("TransformPoints_GetApplyToAttribute"))
 		.AddReturnType(FShaderValueType::Get(EShaderFundamentalType::Uint, 1));
@@ -100,37 +61,19 @@ void UPCGTransformPointsDataInterface::GetShaderParameters(TCHAR const* UID, FSh
 
 void UPCGTransformPointsDataInterface::GetHLSL(FString& OutHLSL, FString const& InDataInterfaceName) const
 {
+	// Overridable params are handled by the base class.
+	Super::GetHLSL(OutHLSL, InDataInterfaceName);
+
 	TMap<FString, FStringFormatArg> TemplateArgs =
 	{
 		{ TEXT("DataInterfaceName"), InDataInterfaceName },
 	};
 
+	// Non-overridable params only.
 	OutHLSL += FString::Format(TEXT(
-		"float3 {DataInterfaceName}_OffsetMin;\n"
-		"float3 {DataInterfaceName}_OffsetMax;\n"
-		"uint {DataInterfaceName}_bAbsoluteOffset;\n"
-		"float3 {DataInterfaceName}_RotationMin;\n"
-		"float3 {DataInterfaceName}_RotationMax;\n"
-		"uint {DataInterfaceName}_bAbsoluteRotation;\n"
-		"float3 {DataInterfaceName}_ScaleMin;\n"
-		"float3 {DataInterfaceName}_ScaleMax;\n"
-		"uint {DataInterfaceName}_bAbsoluteScale;\n"
-		"uint {DataInterfaceName}_bUniformScale;\n"
-		"uint {DataInterfaceName}_bRecomputeSeed;\n"
 		"uint {DataInterfaceName}_bApplyToAttribute;\n"
 		"int {DataInterfaceName}_AttributeId;\n"
 		"\n"
-		"float3 TransformPoints_GetOffsetMin_{DataInterfaceName}() { return {DataInterfaceName}_OffsetMin; }\n"
-		"float3 TransformPoints_GetOffsetMax_{DataInterfaceName}() { return {DataInterfaceName}_OffsetMax; }\n"
-		"uint TransformPoints_GetAbsoluteOffset_{DataInterfaceName}() { return {DataInterfaceName}_bAbsoluteOffset; }\n"
-		"float3 TransformPoints_GetRotationMin_{DataInterfaceName}() { return {DataInterfaceName}_RotationMin; }\n"
-		"float3 TransformPoints_GetRotationMax_{DataInterfaceName}() { return {DataInterfaceName}_RotationMax; }\n"
-		"uint TransformPoints_GetAbsoluteRotation_{DataInterfaceName}() { return {DataInterfaceName}_bAbsoluteRotation; }\n"
-		"float3 TransformPoints_GetScaleMin_{DataInterfaceName}() { return {DataInterfaceName}_ScaleMin; }\n"
-		"float3 TransformPoints_GetScaleMax_{DataInterfaceName}() { return {DataInterfaceName}_ScaleMax; }\n"
-		"uint TransformPoints_GetAbsoluteScale_{DataInterfaceName}() { return {DataInterfaceName}_bAbsoluteScale; }\n"
-		"uint TransformPoints_GetUniformScale_{DataInterfaceName}() { return {DataInterfaceName}_bUniformScale; }\n"
-		"uint TransformPoints_GetRecomputeSeed_{DataInterfaceName}() { return {DataInterfaceName}_bRecomputeSeed; }\n"
 		"uint TransformPoints_GetApplyToAttribute_{DataInterfaceName}() { return {DataInterfaceName}_bApplyToAttribute; }\n"
 		"int TransformPoints_GetAttributeId_{DataInterfaceName}() { return {DataInterfaceName}_AttributeId; }\n"
 		), TemplateArgs);
@@ -141,23 +84,47 @@ UComputeDataProvider* UPCGTransformPointsDataInterface::CreateDataProvider() con
 	return NewObject<UPCGTransformPointsDataProvider>();
 }
 
+bool UPCGTransformPointsDataProvider::PerformPreExecuteReadbacks_GameThread(UPCGDataBinding* InBinding)
+{
+	check(InBinding);
+
+	if (!Super::PerformPreExecuteReadbacks_GameThread(InBinding))
+	{
+		return false;
+	}
+
+	KernelParams = InBinding->GetCachedKernelParams(GetProducerKernel());
+
+	return true;
+}
+
 FComputeDataProviderRenderProxy* UPCGTransformPointsDataProvider::GetRenderProxy()
 {
 	const UPCGTransformPointsSettings* Settings = CastChecked<UPCGTransformPointsSettings>(GetProducerKernel()->GetSettings());
 
 	FPCGTransformPointsDataProviderProxy::FData ProxyData;
 
-	ProxyData.OffsetMin = Settings->OffsetMin;
-	ProxyData.OffsetMax = Settings->OffsetMax;
-	ProxyData.bAbsoluteOffset = Settings->bAbsoluteOffset;
-	ProxyData.RotationMin = Settings->RotationMin.Euler();
-	ProxyData.RotationMax = Settings->RotationMax.Euler();
-	ProxyData.bAbsoluteRotation = Settings->bAbsoluteRotation;
-	ProxyData.ScaleMin = Settings->ScaleMin;
-	ProxyData.ScaleMax = Settings->ScaleMax;
-	ProxyData.bAbsoluteScale = Settings->bAbsoluteScale;
-	ProxyData.bUniformScale = Settings->bUniformScale;
-	ProxyData.bRecomputeSeed = Settings->bRecomputeSeed;
+	// GPU-overridable params read from KernelParams (supports overrides via connected pins).
+	if (ensure(KernelParams))
+	{
+		ProxyData.OffsetMin = KernelParams->GetValueVector(GET_MEMBER_NAME_CHECKED(UPCGTransformPointsSettings, OffsetMin));
+		ProxyData.OffsetMax = KernelParams->GetValueVector(GET_MEMBER_NAME_CHECKED(UPCGTransformPointsSettings, OffsetMax));
+		ProxyData.bAbsoluteOffset = KernelParams->GetValueBool(GET_MEMBER_NAME_CHECKED(UPCGTransformPointsSettings, bAbsoluteOffset));
+		// Store as raw FRotator components (Pitch, Yaw, Roll) to match the attribute buffer layout
+		// used by the overridable system. The shader expects this ordering (not Euler).
+		const FRotator RotMin = KernelParams->GetValueRotator(GET_MEMBER_NAME_CHECKED(UPCGTransformPointsSettings, RotationMin));
+		const FRotator RotMax = KernelParams->GetValueRotator(GET_MEMBER_NAME_CHECKED(UPCGTransformPointsSettings, RotationMax));
+		ProxyData.RotationMin = FVector(RotMin.Pitch, RotMin.Yaw, RotMin.Roll);
+		ProxyData.RotationMax = FVector(RotMax.Pitch, RotMax.Yaw, RotMax.Roll);
+		ProxyData.bAbsoluteRotation = KernelParams->GetValueBool(GET_MEMBER_NAME_CHECKED(UPCGTransformPointsSettings, bAbsoluteRotation));
+		ProxyData.ScaleMin = KernelParams->GetValueVector(GET_MEMBER_NAME_CHECKED(UPCGTransformPointsSettings, ScaleMin));
+		ProxyData.ScaleMax = KernelParams->GetValueVector(GET_MEMBER_NAME_CHECKED(UPCGTransformPointsSettings, ScaleMax));
+		ProxyData.bAbsoluteScale = KernelParams->GetValueBool(GET_MEMBER_NAME_CHECKED(UPCGTransformPointsSettings, bAbsoluteScale));
+		ProxyData.bUniformScale = KernelParams->GetValueBool(GET_MEMBER_NAME_CHECKED(UPCGTransformPointsSettings, bUniformScale));
+		ProxyData.bRecomputeSeed = KernelParams->GetValueBool(GET_MEMBER_NAME_CHECKED(UPCGTransformPointsSettings, bRecomputeSeed));
+	}
+
+	// Non-overridable params read directly from Settings.
 	ProxyData.bApplyToAttribute = Settings->bApplyToAttribute;
 	ProxyData.AttributeId = AttributeId;
 
@@ -168,6 +135,7 @@ void UPCGTransformPointsDataProvider::Reset()
 {
 	Super::Reset();
 
+	KernelParams = nullptr;
 	AttributeId = INDEX_NONE;
 }
 
@@ -223,7 +191,7 @@ bool UPCGTransformPointsDataProvider::PrepareForExecute_GameThread(UPCGDataBindi
 		if (AttributeId == INDEX_NONE && !InputDataDesc->GetDataDescriptions().IsEmpty() && bAnyPointsPresent)
 		{
 			PCG_KERNEL_VALIDATION_ERR(Context, Settings, FText::Format(
-				LOCTEXT("TransformPointsAttributeNotFound", "Transfrom points attribute '{0}' not found."),
+				LOCTEXT("TransformPointsAttributeNotFound", "Transform points attribute '{0}' not found."),
 				FText::FromName(Settings->AttributeName)));
 		}
 	}
diff --git a/Engine/Plugins/PCG/Source/PCG/Private/Compute/DataInterfaces/Elements/PCGTransformPointsDataInterface.h b/Engine/Plugins/PCG/Source/PCG/Private/Compute/DataInterfaces/Elements/PCGTransformPointsDataInterface.h
--- a/Engine/Plugins/PCG/Source/PCG/Private/Compute/DataInterfaces/Elements/PCGTransformPointsDataInterface.h
+++ b/Engine/Plugins/PCG/Source/PCG/Private/Compute/DataInterfaces/Elements/PCGTransformPointsDataInterface.h
@@ -8,9 +8,11 @@
 
 class FPCGTransformPointsDataInterfaceParameters;
 
+struct FPCGKernelParams;
+
 /** Data Interface to marshal transform points settings to the GPU. */
 UCLASS(ClassGroup = (Procedural))
-class UPCGTransformPointsDataInterface : public UPCGComputeDataInterface
+class UPCGTransformPointsDataInterface : public UPCGKernelParamsDataInterface
 {
 	GENERATED_BODY()
 
@@ -25,12 +27,13 @@ public:
 };
 
 UCLASS()
-class UPCGTransformPointsDataProvider : public UPCGComputeDataProvider
+class UPCGTransformPointsDataProvider : public UPCGKernelParamsDataProvider
 {
 	GENERATED_BODY()
 
 public:
 	//~ Begin UComputeDataProvider Interface
+	bool PerformPreExecuteReadbacks_GameThread(UPCGDataBinding* InBinding) override;
 	FComputeDataProviderRenderProxy* GetRenderProxy() override;
 	virtual void Reset() override;
 	virtual bool PrepareForExecute_GameThread(UPCGDataBinding* InBinding) override;
@@ -39,6 +42,9 @@ public:
 	/** Output attribute Id. */
 	UPROPERTY()
 	int32 AttributeId = INDEX_NONE;
+
+	/** The kernel params are owned by the KernelParamsCache, so we are not responsible for its lifetime. */
+	const FPCGKernelParams* KernelParams = nullptr;
 };
 
 class FPCGTransformPointsDataProviderProxy : public FComputeDataProviderRenderProxy
diff --git a/Engine/Plugins/PCG/Source/PCG/Private/Elements/PCGTransformPoints.cpp b/Engine/Plugins/PCG/Source/PCG/Private/Elements/PCGTransformPoints.cpp
--- a/Engine/Plugins/PCG/Source/PCG/Private/Elements/PCGTransformPoints.cpp
+++ b/Engine/Plugins/PCG/Source/PCG/Private/Elements/PCGTransformPoints.cpp
@@ -29,6 +29,7 @@ void UPCGTransformPointsSettings::CreateKernels(FPCGGPUCompilationContext& InOut
                                                 TArray<UPCGComputeKernel*>& OutKernels, TArray<FPCGKernelEdge>& OutEdges) const
 {
 	PCGKernelHelpers::FCreateKernelParams CreateParams(InObjectOuter, this);
+	CreateParams.bRequiresOverridableParams = true;
 	PCGKernelHelpers::CreateKernel<UPCGTransformPointsKernel>(InOutContext, CreateParams, OutKernels, OutEdges);
 }
 #endif // WITH_EDITOR
diff --git a/Engine/Plugins/PCG/Source/PCG/Public/Elements/PCGTransformPoints.h b/Engine/Plugins/PCG/Source/PCG/Public/Elements/PCGTransformPoints.h
--- a/Engine/Plugins/PCG/Source/PCG/Public/Elements/PCGTransformPoints.h
+++ b/Engine/Plugins/PCG/Source/PCG/Public/Elements/PCGTransformPoints.h
@@ -44,42 +44,42 @@ public:
 	UPROPERTY(BlueprintReadWrite, EditAnywhere, Category = Settings, meta = (EditCondition="bApplyToAttribute", EditConditionHides, PCG_Overridable))
 	FName AttributeName;
 
-	UPROPERTY(BlueprintReadWrite, EditAnywhere, Category = Settings, meta = (PCG_Overridable))
+	UPROPERTY(BlueprintReadWrite, EditAnywhere, Category = Settings, meta = (PCG_OverridableCPUAndGPU))
 	FVector OffsetMin = FVector::Zero();
 
-	UPROPERTY(BlueprintReadWrite, EditAnywhere, Category = Settings, meta = (PCG_Overridable))
+	UPROPERTY(BlueprintReadWrite, EditAnywhere, Category = Settings, meta = (PCG_OverridableCPUAndGPU))
 	FVector OffsetMax = FVector::Zero();
 
 	/** Set offset in world space */
-	UPROPERTY(BlueprintReadWrite, EditAnywhere, Category = Settings, meta = (PCG_Overridable))
+	UPROPERTY(BlueprintReadWrite, EditAnywhere, Category = Settings, meta = (PCG_OverridableCPUAndGPU))
 	bool bAbsoluteOffset = false;
 
-	UPROPERTY(BlueprintReadWrite, EditAnywhere, Category = Settings, meta = (PCG_Overridable))
+	UPROPERTY(BlueprintReadWrite, EditAnywhere, Category = Settings, meta = (PCG_OverridableCPUAndGPU))
 	FRotator RotationMin = FRotator::ZeroRotator;
 
-	UPROPERTY(BlueprintReadWrite, EditAnywhere, Category = Settings, meta = (PCG_Overridable))
+	UPROPERTY(BlueprintReadWrite, EditAnywhere, Category = Settings, meta = (PCG_OverridableCPUAndGPU))
 	FRotator RotationMax = FRotator::ZeroRotator;
 
 	/** Set rotation directly instead of additively */
-	UPROPERTY(BlueprintReadWrite, EditAnywhere, Category = Settings, meta = (PCG_Overridable))
+	UPROPERTY(BlueprintReadWrite, EditAnywhere, Category = Settings, meta = (PCG_OverridableCPUAndGPU))
 	bool bAbsoluteRotation = false;
 
-	UPROPERTY(BlueprintReadWrite, EditAnywhere, Category = Settings, meta = (AllowPreserveRatio, PCG_Overridable))
+	UPROPERTY(BlueprintReadWrite, EditAnywhere, Category = Settings, meta = (AllowPreserveRatio, PCG_OverridableCPUAndGPU))
 	FVector ScaleMin = FVector::One();
 
-	UPROPERTY(BlueprintReadWrite, EditAnywhere, Category = Settings, meta = (AllowPreserveRatio, PCG_Overridable))
+	UPROPERTY(BlueprintReadWrite, EditAnywhere, Category = Settings, meta = (AllowPreserveRatio, PCG_OverridableCPUAndGPU))
 	FVector ScaleMax = FVector::One();
 
 	/** Set scale directly instead of multiplicatively */
-	UPROPERTY(BlueprintReadWrite, EditAnywhere, Category = Settings, meta = (PCG_Overridable))
+	UPROPERTY(BlueprintReadWrite, EditAnywhere, Category = Settings, meta = (PCG_OverridableCPUAndGPU))
 	bool bAbsoluteScale = false;
 
 	/** Scale uniformly on each axis. Uses the X component of ScaleMin and ScaleMax. */
-	UPROPERTY(BlueprintReadWrite, EditAnywhere, Category = Settings, meta = (PCG_Overridable))
+	UPROPERTY(BlueprintReadWrite, EditAnywhere, Category = Settings, meta = (PCG_OverridableCPUAndGPU))
 	bool bUniformScale = true;
 
 	/** Recompute the seed for each new point using its new location */
-	UPROPERTY(BlueprintReadWrite, EditAnywhere, Category = Settings, meta = (EditCondition = "!bApplyToAttribute", PCG_Overridable))
+	UPROPERTY(BlueprintReadWrite, EditAnywhere, Category = Settings, meta = (EditCondition = "!bApplyToAttribute", PCG_OverridableCPUAndGPU))
 	bool bRecomputeSeed = false;
 };
 
