[PCG] Fix issue where hovering on a PCG pin does not highlight the pin.

#Synced-CL 52394017

diff --git a/Engine/Plugins/PCG/Source/PCGEditor/Private/Pins/SPCGEditorGraphNodePin.cpp b/Engine/Plugins/PCG/Source/PCGEditor/Private/Pins/SPCGEditorGraphNodePin.cpp
--- a/Engine/Plugins/PCG/Source/PCGEditor/Private/Pins/SPCGEditorGraphNodePin.cpp
+++ b/Engine/Plugins/PCG/Source/PCGEditor/Private/Pins/SPCGEditorGraphNodePin.cpp
@@ -230,18 +230,18 @@ void SPCGEditorGraphNodePin::Construct(const FArguments& InArgs, UEdGraphPin* In
 	// Set up a hover for pins that is tinted the color of the pin.
 
 	SBorder::Construct(SBorder::FArguments()
-		.BorderImage(SPCGEditorGraphNodePin::GetPinBorder())
-		.BorderBackgroundColor(SPCGEditorGraphNodePin::GetHighlightColor())
-		.OnMouseButtonDown(this, &SPCGEditorGraphNodePin::OnPinNameMouseDown)
+		.BorderImage(this, &SPCGEditorGraphNodePin::GetPinBorder) // Requires dynamic lambda
+		.BorderBackgroundColor(this, &SPCGEditorGraphNodePin::GetHighlightColor) // Requires dynamic lambda
+		.OnMouseButtonDown(this, &SPCGEditorGraphNodePin::OnPinNameMouseDown) // Requires dynamic lambda
 		.Padding(0) // NOTE: This is different from base class implementation
 		[
 			SNew(SBorder)
 			.BorderImage(CachedImg_Pin_DiffOutline)
-			.BorderBackgroundColor(GetPinDiffColor())
+			.BorderBackgroundColor(this, &SPCGEditorGraphNodePin::GetPinDiffColor) // Requires dynamic lambda
 			.Padding(0) // NOTE: This is different from base class implementation
 			[
 				SNew(SLevelOfDetailBranchNode)
-				.UseLowDetailSlot(UseLowDetailPinNames())
+				.UseLowDetailSlot(this, &SPCGEditorGraphNodePin::UseLowDetailPinNames) // Requires dynamic lambda
 				.LowDetail()
 				[
 					//@TODO: Try creating a pin-colored line replacement that doesn't measure text / call delegates but still renders
