[PCG] Merge FPCGMetadataAttributeGeneric into FPCGMetadataAttributeBase

* To keep the CL to touch too many lines, I keep the `FPCGMetadataAttributeGeneric` as an alias and the few functions to create/get it.
* Added restriction on the constructors of the attributes to always be created by a metadata
* Reduce templated code to the strict minimum, added type erased functions.
* Metadata domain will now verify that the attribute is a valid type. If not the attribute creation will fail. Preliminary work to remove most type validations in the attribute code.
* Replaced in a few files FPCGMetadataAttributeGeneric, as we can't have forward declaration of a class with an alias.
* Try to use Attribute descriptors as much as we can, sunsetting type id for non-typed attributes.
* We need a version bump to know if we need to allocate a typed attribute or a base/generic attribute.

#Synced-CL 52470086

diff --git a/Engine/Plugins/PCG/Source/PCG/Private/Elements/PCGMatchAndSetAttributes.cpp b/Engine/Plugins/PCG/Source/PCG/Private/Elements/PCGMatchAndSetAttributes.cpp
--- a/Engine/Plugins/PCG/Source/PCG/Private/Elements/PCGMatchAndSetAttributes.cpp
+++ b/Engine/Plugins/PCG/Source/PCG/Private/Elements/PCGMatchAndSetAttributes.cpp
@@ -700,7 +700,7 @@ public:
 				{
 					if (PCG::Private::IsBroadcastableOrConstructible(PointAttribute->GetTypeId(), ParamAttribute->GetTypeId()))
 					{
-						if (OutMetadata->ChangeAttributeType(AttributeName, ParamAttribute->GetTypeId()))
+						if (OutMetadata->ChangeAttributeType(AttributeName, ParamAttribute->GetAttributeDesc()))
 						{
 							bShouldDeleteAttribute = false;
 
diff --git a/Engine/Plugins/PCG/Source/PCG/Private/Helpers/PCGSettingsHelpers.cpp b/Engine/Plugins/PCG/Source/PCG/Private/Helpers/PCGSettingsHelpers.cpp
--- a/Engine/Plugins/PCG/Source/PCG/Private/Helpers/PCGSettingsHelpers.cpp
+++ b/Engine/Plugins/PCG/Source/PCG/Private/Helpers/PCGSettingsHelpers.cpp
@@ -14,215 +14,6 @@
 
 namespace PCGSettingsHelpers
 {
-	template<typename T>
-	bool GetParamValue(T& Value, const FPCGMetadataAttributeBase* InAttribute, PCGMetadataEntryKey InKey = 0)
-	{
-		if (!InAttribute->IsGeneric())
-		{
-			return false;
-		}
-		
-		const FPCGMetadataAttributeGeneric* Attribute = static_cast<const FPCGMetadataAttributeGeneric*>(InAttribute);
-		if (Attribute->IsOfType<T>())
-		{
-			Value = Attribute->GetValueFromItemKey<T>(InKey);
-			return true;
-		}
-		else
-		{
-			return false;
-		}
-	}
-
-	template<typename T, typename U, typename... OtherTypes>
-	bool GetParamValue(T& Value, const FPCGMetadataAttributeBase* InAttribute, PCGMetadataEntryKey InKey = 0)
-	{
-		if (!InAttribute->IsGeneric())
-		{
-			return false;
-		}
-		
-		U TempValue{};
-		if (GetParamValue<U>(TempValue, InAttribute, InKey))
-		{
-			Value = T{TempValue};
-			return true;
-		}
-		else
-		{
-			return GetParamValue<T, OtherTypes...>(Value, InAttribute, InKey);
-		}
-	}
-
-	void SetValue(UPCGParamData* InParams, UObject* Object, FProperty* Property)
-	{
-		check(InParams && Property);
-		const FPCGMetadataAttributeBase* MatchingAttribute = InParams->Metadata ? InParams->Metadata->GetConstAttribute(Property->GetFName()) : nullptr;
-
-		// TODO? Support arrays (1) which would be possible if we have a consistent naming scheme
-		if (!MatchingAttribute)
-		{
-			return;
-		}
-
-		bool bTypeError = false;
-
-		if (FNumericProperty* NumericProperty = CastField<FNumericProperty>(Property))
-		{
-			if (NumericProperty->IsFloatingPoint())
-			{
-				double Value = 0;
-				if (GetParamValue<double, float>(Value, MatchingAttribute))
-				{
-					NumericProperty->SetFloatingPointPropertyValue(NumericProperty->ContainerPtrToValuePtr<double>(Object, 0), Value);
-				}
-				else
-				{
-					bTypeError = true;
-				}
-			}
-			else if(NumericProperty->IsInteger())
-			{
-				int64 Value = 0;
-				if (GetParamValue<int64, int32>(Value, MatchingAttribute))
-				{
-					NumericProperty->SetIntPropertyValue(NumericProperty->ContainerPtrToValuePtr<int64>(Object, 0), Value);
-				}
-				else
-				{
-					bTypeError = true;
-				}
-			}
-			else
-			{
-				// Not a float and not an integer...?
-				bTypeError = true;
-			}
-		}
-		else if (FBoolProperty* BoolProperty = CastField<FBoolProperty>(Property))
-		{
-			bool Value = 0;
-			if(GetParamValue<bool>(Value, MatchingAttribute))
-			{
-				BoolProperty->SetPropertyValue(BoolProperty->ContainerPtrToValuePtr<bool>(Object, 0), Value);
-			}
-			else
-			{
-				bTypeError = true;
-			}
-		}
-		else if (FStrProperty* StringProperty = CastField<FStrProperty>(Property))
-		{
-			FString Value;
-			if (GetParamValue<FString>(Value, MatchingAttribute))
-			{
-				StringProperty->SetPropertyValue(StringProperty->ContainerPtrToValuePtr<FString>(Object, 0), Value);
-			}
-			else
-			{
-				bTypeError = true;
-			}
-		}
-		else if (FNameProperty* NameProperty = CastField<FNameProperty>(Property))
-		{
-			FString Value;
-			if (GetParamValue<FString>(Value, MatchingAttribute))
-			{
-				FName NValue(Value);
-				NameProperty->SetPropertyValue(NameProperty->ContainerPtrToValuePtr<FName>(Object), NValue);
-			}
-			else
-			{
-				bTypeError = true;
-			}
-		}
-		else if (FEnumProperty* EnumProperty = CastField<FEnumProperty>(Property))
-		{
-			FNumericProperty* UnderlyingProperty = EnumProperty->GetUnderlyingProperty();
-
-			int64 Value;
-			if (GetParamValue<int64, int32, uint8>(Value, MatchingAttribute))
-			{
-				UnderlyingProperty->SetIntPropertyValue(EnumProperty->ContainerPtrToValuePtr<void>(Object, 0), Value);
-			}
-		}
-		else if (FStructProperty* StructProperty = CastField<FStructProperty>(Property))
-		{
-			const FName PropertyTypeName = StructProperty->Struct->GetFName();
-			void* PropertyValue = StructProperty->ContainerPtrToValuePtr<void>(Object, 0);
-
-			if (PropertyTypeName == NAME_Vector)
-			{
-				FVector Value = FVector::Zero();
-				if(GetParamValue<FVector, double, float>(Value, MatchingAttribute))
-				{
-					*reinterpret_cast<FVector*>(PropertyValue) = Value;
-				}
-				else
-				{
-					bTypeError = true;
-				}
-			}
-			else if (PropertyTypeName == NAME_Vector4)
-			{
-				FVector4 Value = FVector4::Zero();
-				if (GetParamValue<FVector4, double, float>(Value, MatchingAttribute))
-				{
-					*reinterpret_cast<FVector4*>(PropertyValue) = Value;
-				}
-				else
-				{
-					bTypeError = true;
-				}
-			}
-			else if (PropertyTypeName == NAME_Quat)
-			{
-				FQuat Value = FQuat::Identity;
-				if (GetParamValue<FQuat, FRotator>(Value, MatchingAttribute))
-				{
-					*reinterpret_cast<FQuat*>(PropertyValue) = Value;
-				}
-			}
-			else if (PropertyTypeName == NAME_Transform)
-			{
-				FTransform Value = FTransform::Identity;
-				if (GetParamValue<FTransform, FVector>(Value, MatchingAttribute))
-				{
-					*reinterpret_cast<FTransform*>(PropertyValue) = Value;
-				}
-				else
-				{
-					bTypeError = true;
-				}
-			}
-			else if (PropertyTypeName == NAME_Rotator)
-			{
-				FRotator Value = FRotator::ZeroRotator;
-				if (GetParamValue<FRotator, FQuat>(Value, MatchingAttribute))
-				{
-					*reinterpret_cast<FRotator*>(PropertyValue) = Value;
-				}
-				else
-				{
-					bTypeError = true;
-				}
-			}
-			//else if (PropertyTypeName == NAME_Color)
-			//else if (PropertyTypeName == NAME_LinearColor)
-			//else if soft object path to something?
-			else
-			{
-				bTypeError = true;
-			}
-		}
-
-		if (bTypeError)
-		{
-			// Mismatching types
-			UE_LOG(LogPCG, Warning, TEXT("Mismatching type in params vs. property %s"), *Property->GetFName().ToString());
-		}
-	}
-
 	void DeprecationBreakOutParamsToNewPin(UPCGNode* InOutNode, TArray<TObjectPtr<UPCGPin>>& InputPins, TArray<TObjectPtr<UPCGPin>>& OutputPins)
 	{
 		// Check basic conditions for which the code below should run.
diff --git a/Engine/Plugins/PCG/Source/PCG/Private/Metadata/PCGMetadata.cpp b/Engine/Plugins/PCG/Source/PCG/Private/Metadata/PCGMetadata.cpp
--- a/Engine/Plugins/PCG/Source/PCG/Private/Metadata/PCGMetadata.cpp
+++ b/Engine/Plugins/PCG/Source/PCG/Private/Metadata/PCGMetadata.cpp
@@ -421,6 +421,12 @@ void UPCGMetadata::GetResourceSizeEx(FResourceSizeEx& CumulativeResourceSize)
 	});
 }
 
+FPCGMetadataAttributeBase* UPCGMetadata::CreateAttribute(FPCGAttributeIdentifier AttributeName, const FPCGMetadataAttributeDesc& AttributeDesc, bool bAllowsInterpolation, bool bOverrideParent)
+{
+	FPCGMetadataDomain* FoundSubMetadata = FindOrCreateMetadataDomain(AttributeName.MetadataDomain);
+	return FoundSubMetadata ? FoundSubMetadata->CreateAttribute(AttributeDesc, bAllowsInterpolation, bOverrideParent) : nullptr;
+}
+
 void UPCGMetadata::K2_InitializeAsCopy(const UPCGMetadata* InMetadataToCopy, const TArray<int64>& InOptionalEntriesToCopy)
 {
 	if (!IsValid(InMetadataToCopy))
@@ -1012,9 +1018,18 @@ void UPCGMetadata::DeleteAttribute(FPCGAttributeIdentifier AttributeToDelete)
 	return WithMetadataDomain(AttributeToDelete.MetadataDomain, &FPCGMetadataDomain::DeleteAttribute, AttributeToDelete.Name);
 }
 
+bool UPCGMetadata::ChangeAttributeType(FPCGAttributeIdentifier AttributeName, const FPCGMetadataAttributeDesc& AttributeNewDesc)
+{
+	// Need to differentiate overloads.
+	return WithMetadataDomain(AttributeName.MetadataDomain, static_cast<bool(FPCGMetadataDomain::*)(FName, const FPCGMetadataAttributeDesc&)>(&FPCGMetadataDomain::ChangeAttributeType), AttributeName.Name, AttributeNewDesc);
+}
+
 bool UPCGMetadata::ChangeAttributeType(FPCGAttributeIdentifier AttributeName, int16 AttributeNewType)
 {
-	return WithMetadataDomain(AttributeName.MetadataDomain, &FPCGMetadataDomain::ChangeAttributeType, AttributeName.Name, AttributeNewType);
+PRAGMA_DISABLE_DEPRECATION_WARNINGS
+	// Need to differentiate overloads.
+	return WithMetadataDomain(AttributeName.MetadataDomain, static_cast<bool(FPCGMetadataDomain::*)(FName, int16)>(&FPCGMetadataDomain::ChangeAttributeType), AttributeName.Name, AttributeNewType);
+PRAGMA_ENABLE_DEPRECATION_WARNINGS
 }
 
 int64 UPCGMetadata::GetItemCountForChild() const
diff --git a/Engine/Plugins/PCG/Source/PCG/Private/Metadata/PCGMetadataAttribute.cpp b/Engine/Plugins/PCG/Source/PCG/Private/Metadata/PCGMetadataAttribute.cpp
--- a/Engine/Plugins/PCG/Source/PCG/Private/Metadata/PCGMetadataAttribute.cpp
+++ b/Engine/Plugins/PCG/Source/PCG/Private/Metadata/PCGMetadataAttribute.cpp
@@ -2,6 +2,459 @@
 
 #include "Metadata/PCGMetadataAttribute.h"
 #include "Metadata/PCGMetadata.h"
+#include "Metadata/PCGMetadataDomain.h"
+#include "Metadata/Accessors/PCGAttributeAccessorHelpers.h"
+
+#include "Serialization/ArchiveCountMem.h"
+#include "UObject/TextProperty.h"
+
+namespace PCG::Private
+{
+	class FNullArchive : public FArchive
+	{
+		//override FObjectPtr to do nothing
+		virtual FArchive& operator<<(struct FObjectPtr& Value) override
+		{
+			return *this;
+		}
+	};
+
+	/**
+	 * Equivalent of the method in PropertyBag to convert the descriptor to a property.
+	 *
+	 */
+	FProperty* CreatePropertyFromDesc(const FPCGMetadataAttributeDesc& Desc, const FFieldVariant PropertyScope)
+	{
+		// Handle array and nested containers properties
+		if (Desc.ContainerTypes.Num() > 0)
+		{
+			FProperty* Prop = nullptr; // the first created container will fill the return value, nested ones will fill the inner
+
+			// support for nested containers, i.e. : TArray<TArray<float>>
+			FFieldVariant PropertyOwner = PropertyScope;
+			FProperty** ValuePropertyPtr = &Prop;
+
+			// Create the container list
+			for (EPCGMetadataAttributeContainerTypes BagContainerType : Desc.ContainerTypes)
+			{
+				switch(BagContainerType)
+				{
+				case EPCGMetadataAttributeContainerTypes::Array:
+					{
+						// create an array property as a container for the tail
+						FArrayProperty* ArrayProperty = new FArrayProperty(PropertyOwner, Desc.Name);
+						*ValuePropertyPtr = ArrayProperty;
+						ValuePropertyPtr = &ArrayProperty->Inner;
+						PropertyOwner = ArrayProperty;
+						break;
+					}
+				case EPCGMetadataAttributeContainerTypes::Set:
+					{
+						// create a set property as a container for the tail
+						FSetProperty* SetProperty = new FSetProperty(PropertyOwner, Desc.Name);
+						*ValuePropertyPtr = SetProperty;
+						ValuePropertyPtr = &SetProperty->ElementProp;
+						PropertyOwner = SetProperty;
+						break;
+					}
+				case EPCGMetadataAttributeContainerTypes::Map:
+					{
+						// create a map property as a container for the tail
+						FMapProperty* MapProperty = new FMapProperty(PropertyOwner, Desc.Name);
+						*ValuePropertyPtr = MapProperty;
+						ValuePropertyPtr = &MapProperty->ValueProp;
+						PropertyOwner = MapProperty;
+
+						// create the key property (container types are not supported for keys)
+						FPCGMetadataAttributeDesc KeyDesc;
+						KeyDesc.ValueType = Desc.KeyType;
+						KeyDesc.ValueTypeObject = Desc.KeyTypeObject;
+						KeyDesc.Name = Desc.Name;
+						FProperty* KeyProp = CreatePropertyFromDesc(KeyDesc, PropertyOwner);
+						if (!KeyProp)
+						{
+							delete MapProperty;
+							return nullptr;
+						}
+
+						MapProperty->KeyProp = KeyProp;
+						break;
+					}
+				default:
+					ensureMsgf(false, TEXT("Unsupported container type %s"), *UEnum::GetValueAsString(BagContainerType));
+					break;
+				}
+			}
+
+			// finally create the tail type
+			FPCGMetadataAttributeDesc InnerDesc = Desc;
+			InnerDesc.ContainerTypes.Reset();
+			*ValuePropertyPtr = CreatePropertyFromDesc(InnerDesc, PropertyOwner);
+
+			if (*ValuePropertyPtr == nullptr)
+			{
+				delete Prop;
+				return nullptr;
+			}
+			else
+			{
+				return Prop;
+			}
+		}
+
+		const UScriptStruct* ScriptStruct = nullptr;
+
+		switch (Desc.ValueType)
+		{
+		case EPCGMetadataTypes::Boolean:
+			{
+				FBoolProperty* Prop = new FBoolProperty(PropertyScope, Desc.Name);
+				Prop->SetBoolSize(sizeof(bool), true); // Enable native access (init the whole byte, rather than just first bit)
+				Prop->SetPropertyFlags(CPF_HasGetValueTypeHash);
+				return Prop;
+			}
+		case EPCGMetadataTypes::Byte:
+			{
+				FByteProperty* Prop = new FByteProperty(PropertyScope, Desc.Name);
+				Prop->SetPropertyFlags(CPF_HasGetValueTypeHash);
+				return Prop;
+			}
+		case EPCGMetadataTypes::Integer32:
+			{
+				FIntProperty* Prop = new FIntProperty(PropertyScope, Desc.Name);
+				Prop->SetPropertyFlags(CPF_HasGetValueTypeHash);
+				return Prop;
+			}
+		case EPCGMetadataTypes::Integer64:
+			{
+				FInt64Property* Prop = new FInt64Property(PropertyScope, Desc.Name);
+				Prop->SetPropertyFlags(CPF_HasGetValueTypeHash);
+				return Prop;
+			}
+		case EPCGMetadataTypes::Float:
+			{
+				FFloatProperty* Prop = new FFloatProperty(PropertyScope, Desc.Name);
+				Prop->SetPropertyFlags(CPF_HasGetValueTypeHash);
+				return Prop;
+			}
+		case EPCGMetadataTypes::Double:
+			{
+				FDoubleProperty* Prop = new FDoubleProperty(PropertyScope, Desc.Name);
+				Prop->SetPropertyFlags(CPF_HasGetValueTypeHash);
+				return Prop;
+			}
+		case EPCGMetadataTypes::Name:
+			{
+				FNameProperty* Prop = new FNameProperty(PropertyScope, Desc.Name);
+				Prop->SetPropertyFlags(CPF_HasGetValueTypeHash);
+				return Prop;
+			}
+		case EPCGMetadataTypes::String:
+			{
+				FStrProperty* Prop = new FStrProperty(PropertyScope, Desc.Name);
+				Prop->SetPropertyFlags(CPF_HasGetValueTypeHash);
+				return Prop;
+			}
+		case EPCGMetadataTypes::Text:
+			{
+				FTextProperty* Prop = new FTextProperty(PropertyScope, Desc.Name);
+				return Prop;
+			}
+		case EPCGMetadataTypes::Enum:
+			if (const UEnum* Enum = Cast<UEnum>(Desc.ValueTypeObject))
+			{
+				FEnumProperty* Prop = new FEnumProperty(PropertyScope, Desc.Name);
+				Prop->SetPropertyFlags(CPF_HasGetValueTypeHash);
+				FNumericProperty* UnderlyingProp = new FByteProperty(Prop, "UnderlyingType"); // HACK: Hardwire to byte property for now for BP compatibility
+				UnderlyingProp->SetPropertyFlags(CPF_HasGetValueTypeHash | CPF_IsPlainOldData);
+				Prop->SetEnum(const_cast<UEnum*>(Enum));
+				Prop->AddCppProperty(UnderlyingProp);
+				return Prop;
+			}
+			break;
+		case EPCGMetadataTypes::SoftObjectPath:
+			ScriptStruct = TBaseStructure<FSoftObjectPath>::Get();
+			break;
+		case EPCGMetadataTypes::SoftClassPath:
+			ScriptStruct = TBaseStructure<FSoftClassPath>::Get();
+			break;
+		case EPCGMetadataTypes::Vector2:
+			ScriptStruct = TBaseStructure<FVector2D>::Get();
+			break;
+		case EPCGMetadataTypes::Vector:
+			ScriptStruct = TBaseStructure<FVector>::Get();
+			break;
+		case EPCGMetadataTypes::Vector4:
+			ScriptStruct = TBaseStructure<FVector4>::Get();
+			break;
+		case EPCGMetadataTypes::Quaternion:
+			ScriptStruct = TBaseStructure<FQuat>::Get();
+			break;
+		case EPCGMetadataTypes::Rotator:
+			ScriptStruct = TBaseStructure<FRotator>::Get();
+			break;
+		case EPCGMetadataTypes::Transform:
+			ScriptStruct = TBaseStructure<FTransform>::Get();
+			break;
+		case EPCGMetadataTypes::Struct:
+			ScriptStruct = Cast<UScriptStruct>(Desc.ValueTypeObject);
+			break;
+
+		// @todo_pcg: Support class and objects
+// 		case EPCGMetadataTypes::Class:
+// 			if (const UClass* Class = Cast<UClass>(Desc.ValueTypeObject))
+// 			{
+// 				FClassProperty* Prop = new FClassProperty(PropertyScope, Desc.Name);
+// #if WITH_EDITORONLY_DATA
+// 				Prop->SetMetaClass(const_cast<UClass*>(Desc.ValueTypeObject ? Cast<UClass>(Desc.ValueTypeObject) : Class));
+// #else
+// 				Prop->SetMetaClass(const_cast<UClass*>(Class));
+// #endif
+// 				Prop->PropertyClass = UClass::StaticClass();
+// 				Prop->SetPropertyFlags(CPF_HasGetValueTypeHash);
+// 				return Prop;
+// 			}
+// 			break;
+// 		case EPCGMetadataTypes::SoftClass:
+// 			if (const UClass* Class = Cast<UClass>(Desc.ValueTypeObject))
+// 			{
+// 				FSoftClassProperty* Prop = new FSoftClassProperty(PropertyScope, Desc.Name);
+// #if WITH_EDITORONLY_DATA
+// 				Prop->SetMetaClass(const_cast<UClass*>(Desc.ValueTypeObject ? Cast<UClass>(Desc.ValueTypeObject) : Class));
+// #else
+// 				Prop->SetMetaClass(const_cast<UClass*>(Class));
+// #endif
+// 				Prop->PropertyClass = UClass::StaticClass();
+// 				Prop->SetPropertyFlags(CPF_HasGetValueTypeHash);
+// 				return Prop;
+// 			}
+// 			break;
+// 		case EPCGMetadataTypes::Object:
+// 			if (const UClass* Class = Cast<UClass>(Desc.ValueTypeObject))
+// 			{
+// 				FObjectProperty* Prop = new FObjectProperty(PropertyScope, Desc.Name);
+// 				if (Class->HasAnyClassFlags(CLASS_DefaultToInstanced))
+// 				{
+// 					Prop->SetPropertyFlags(CPF_InstancedReference);
+// 				}
+// 				Prop->SetPropertyClass(const_cast<UClass*>(Class));
+// 				Prop->SetPropertyFlags(CPF_HasGetValueTypeHash | CPF_TObjectPtr);
+// 				return Prop;
+// 			}
+// 			break;
+// 		case EPCGMetadataTypes::SoftObject:
+// 			if (const UClass* Class = Cast<UClass>(Desc.ValueTypeObject))
+// 			{
+// 				FSoftObjectProperty* Prop = new FSoftObjectProperty(PropertyScope, Desc.Name);
+// 				if (Class->HasAnyClassFlags(CLASS_DefaultToInstanced))
+// 				{
+// 					Prop->SetPropertyFlags(CPF_InstancedReference);
+// 				}
+// 				Prop->SetPropertyClass(const_cast<UClass*>(Class));
+// 				Prop->SetPropertyFlags(CPF_HasGetValueTypeHash);
+// 				return Prop;
+// 			}
+// 			break;
+		default:
+			ensureMsgf(false, TEXT("Unhandled type %s"), *UEnum::GetValueAsString(Desc.ValueType));
+		}
+
+		if (ScriptStruct)
+		{
+			FStructProperty* Prop = new FStructProperty(PropertyScope, Desc.Name);
+			// Struct Property requires the underlying struct to be non-const. It is done like that in the engine.
+			Prop->Struct = const_cast<UScriptStruct*>(ScriptStruct);
+
+			if (ScriptStruct->GetCppStructOps() && ScriptStruct->GetCppStructOps()->HasGetTypeHash())
+			{
+				Prop->SetPropertyFlags(CPF_HasGetValueTypeHash);
+			}
+
+			if (ScriptStruct->StructFlags & STRUCT_HasInstancedReference)
+			{
+				Prop->SetPropertyFlags(CPF_ContainsInstancedReference);
+			}
+
+			return Prop;
+		}
+
+		return nullptr;
+	}
+
+	/**
+	 * Taken from FArrayProperty::CopyValuesInternal in order to leverage MemCpy when we have plain old data
+	 * with the optimization twist for Object copies (also present in CopyValuesInternal)
+	 * @todo_pcg: Refactorize it in Core to re-use it. Maybe add it to FScriptArrayHelper?
+	 */
+	void CopyArray(const FArrayProperty* InProperty, void* DestPtr, const void* SrcPtr, int32 Count)
+	{
+		const FProperty* Inner = InProperty ? InProperty->Inner : nullptr;
+
+		if (!Inner || Count <= 0 || !DestPtr || !SrcPtr)
+		{
+			return;
+		}
+
+		size_t Size = Inner->GetElementSize();
+		const uint8* SrcData = (uint8*)SrcPtr;
+		uint8* DestData = (uint8*)DestPtr;
+		if (!(Inner->PropertyFlags & CPF_IsPlainOldData))
+		{
+			// If we are in the closed body of a transaction, we have a special case for `FObjectProperty` that vastly
+			// improves performance when copying arrays around (in the order of 4x faster). We **do not want** to do
+			// this optimization with anything that derives from `FObjectProperty` though, so we **do not use** `IsA`
+			// as that would all catch classes that inherited from `FObjectProperty`!
+			if (AutoRTFM::IsClosed() && (Inner->GetClass() == FObjectProperty::StaticClass()))
+			{
+				#if UE_WITH_REMOTE_OBJECT_HANDLE
+				// CopyCompleteValue will resolve handles in the RawPtr <- TObjectPtr case.
+				// Since we can't abort while in an open scope, use a speculation scope to defer any
+				// migrations until this scope ends.
+				UE::RemoteExecutor::FSpeculationExecutionScope AvoidMigrationsInTheOpen;
+				#endif
+
+				// Since we've already written to each location that we are going to be copying over below, we know that
+				// all of the record-writes that the transactional system would have inserted are no-ops. Also, each copy
+				// of a `TObjectPtr` will go into the open to do mark the object as reachable, so by doing a single open
+				// here we can avoid doing N calls to record-write and N calls to open for each object in the array. Big
+				// saving on large arrays!
+				UE_AUTORTFM_OPEN_NO_VALIDATION
+				{
+					// This optimization assumes that `FObjectProperty` is just backed by a `TObjectPtr` which when copying
+					// does not do any heap allocations, just pulls over the pointer to the underlying `UObject` and then
+					// marks it as reachable for incremental reachability analysis. Lets check that this is still the case
+					// here so that if anyone ever changes that in future this line will shout at them!
+					static_assert(std::is_same_v<decltype(FObjectProperty::GetDefaultPropertyValue()), TObjectPtr<UObject>>);
+
+					for (int32 i = 0; i < Count; i++)
+					{
+						Inner->CopyCompleteValue(DestData + i * Size, SrcData + i * Size);
+					}
+				};
+			}
+			else
+			{
+				for (int32 i = 0; i < Count; i++)
+				{
+					Inner->CopyCompleteValue(DestData + i * Size, SrcData + i * Size);
+				}
+			}
+		}
+		else
+		{
+			FMemory::Memcpy(DestData, SrcData, Count * Size);
+		}
+	}
+}
+
+/**
+ * Create a new wrapper using the Attribute desc. This is expected to be stored in a SharedPtr to avoid having
+ * to always recreate it.
+ */
+FPCGAttributeProperty::FPCGAttributeProperty(const FPCGMetadataAttributeDesc& InAttributeDesc)
+{
+	// Use an array property to be able to leverage FScriptArrayHelper. It will contain the property
+	// that matches with the Descriptor.
+	// Note: CreatePropertyFromDesc will allocate a raw pointer, and Property will become its owner. When Property
+	// is destroyed, all allocated properties will also be deleted. It's all done internally.
+	Property = MakeUnique<FArrayProperty>(FFieldVariant{nullptr}, NAME_None);
+	FProperty* NewProperty = PCG::Private::CreatePropertyFromDesc(InAttributeDesc, Property.Get());
+
+	if (!NewProperty)
+	{
+		Property.Reset();
+		return;
+	}
+
+	Property->AddCppProperty(NewProperty);
+	// We need to link the property so it fills the Property Flags and other members correctly.
+	// We don't need an archive (that's how it is done in the FProperty::StaticLink), but FNullArchive is private so we copied it.
+	PCG::Private::FNullArchive ArDummy;
+	Property->Link(ArDummy);
+
+	// Cache info if the data is plain old data.
+	bIsPlainOldData = !!(NewProperty->PropertyFlags & CPF_IsPlainOldData);
+
+	// Compress strings and soft paths, if they are not in a container.
+	bCompressData = InAttributeDesc.ContainerTypes.IsEmpty() &&
+		(InAttributeDesc.ValueType == EPCGMetadataTypes::String
+		|| InAttributeDesc.ValueType == EPCGMetadataTypes::Text
+		|| InAttributeDesc.ValueType == EPCGMetadataTypes::SoftObjectPath
+		|| InAttributeDesc.ValueType == EPCGMetadataTypes::SoftClassPath
+		|| InAttributeDesc.ValueTypeObject == TBaseStructure<FSoftObjectPath>().Get()
+		|| InAttributeDesc.ValueTypeObject == TBaseStructure<FSoftClassPath>().Get());
+
+	InnerElementCachedSize = NewProperty->GetElementSize();
+}
+
+FPCGAttributeProperty::~FPCGAttributeProperty() = default;
+
+void FPCGAttributeProperty::Copy(void* DestPtr, const void* SrcPtr, int32 Count)
+{
+	PCG::Private::CopyArray(Property.Get(), DestPtr, SrcPtr, Count);
+}
+
+bool FPCGAttributeProperty::IsValid() const
+{
+	return Property.IsValid() && Property->Inner != nullptr;
+}
+
+const void* FPCGAttributeProperty::GetPtrInArray(const void* InArrayStart, int32 Index) const
+{
+	check(InnerElementCachedSize > 0);
+	return static_cast<const uint8*>(InArrayStart) + Index * InnerElementCachedSize;
+}
+
+FPCGAttributeDefaultValue::~FPCGAttributeDefaultValue()
+{
+	Destroy();
+}
+
+void FPCGAttributeDefaultValue::Init(TSharedPtr<FPCGAttributeProperty> InUnderlyingProperty)
+{
+	if (!InUnderlyingProperty || !InUnderlyingProperty->IsValid())
+	{
+		return;
+	}
+
+	// Check if we already have a property set. If the property is the same type, we can keep the allocated memory.
+	// Otherwise we have to destroy it.
+	if (UnderlyingProperty && UnderlyingProperty->IsValid() && Memory)
+	{
+		if (UnderlyingProperty->Property->SameType(InUnderlyingProperty->Property.Get()))
+		{
+			UnderlyingProperty = MoveTemp(InUnderlyingProperty);
+			return;
+		}
+	}
+
+	Destroy();
+
+	UnderlyingProperty = MoveTemp(InUnderlyingProperty);
+	Memory = UnderlyingProperty->Property->Inner->AllocateAndInitializeValue();
+}
+
+void FPCGAttributeDefaultValue::CopyFrom(const void* SrcData)
+{
+	check(Memory);
+	UnderlyingProperty->Copy(Memory, SrcData, /*Count=*/1);
+}
+
+void FPCGAttributeDefaultValue::Destroy()
+{
+	if (Memory)
+	{
+		if (ensureMsgf(UnderlyingProperty && UnderlyingProperty->IsValid(),
+			TEXT("Memory is allocated but we do not have an underlying property. This is ill-formed and will leak memory")))
+		{
+			UnderlyingProperty->Property->Inner->DestroyValue(Memory);
+		}
+
+		FMemory::Free(Memory);
+	}
+
+	Memory = nullptr;
+}
 
 namespace PCGMetadataAttributeConstants
 {
@@ -42,24 +495,111 @@ FPCGMetadataAttributeBase::FPCGMetadataAttributeBase(FPCGMetadataDomain* InMetad
 {
 }
 
-PRAGMA_DISABLE_DEPRECATION_WARNINGS
-FPCGMetadataAttributeBase::FPCGMetadataAttributeBase(UPCGMetadata* InMetadata, FName InName, const FPCGMetadataAttributeBase* InParent, bool bInAllowsInterpolation)
-	: FPCGMetadataAttributeBase(InMetadata ? InMetadata->GetDefaultMetadataDomain() : nullptr, InName, InParent, bInAllowsInterpolation)
-{}
+FPCGMetadataAttributeBase::FPCGMetadataAttributeBase(const FPCGMetadataAttributeDesc& InAttributeDesc, FPCGMetadataDomain* InMetadata, bool bInAllowsInterpolation)
+	: Metadata(InMetadata)
+	, bAllowsInterpolation(bInAllowsInterpolation)
+	, CachedDesc(InAttributeDesc)
+	, UnderlyingProperty(MakeShared<FPCGAttributeProperty>(CachedDesc))
+	, Name(InAttributeDesc.Name)
+{
+	// @todo_pcg support multi-containers
+	if (!ensure(CachedDesc.ContainerTypes.Num() <= 1))
+	{
+		UnderlyingProperty.Reset();
+		CachedDesc = FPCGMetadataAttributeDesc{};
+		return;
+	}
+
+	Init();
+}
 
-void FPCGMetadataAttributeBase::Serialize(UPCGMetadata* InMetadata, FArchive& InArchive)
+FPCGMetadataAttributeBase::FPCGMetadataAttributeBase(TNotNull<const FPCGMetadataAttributeBase*> InParent, FPCGMetadataDomain* InMetadata, bool bInAllowsInterpolation, TOptional<FName> InName)
+	: Metadata(InMetadata)
+	, Parent(InParent)
+	, bAllowsInterpolation(bInAllowsInterpolation)
+	, CachedDesc(InParent->CachedDesc)
+	, UnderlyingProperty(InParent->UnderlyingProperty)
+	, Name(InName.Get(InParent->CachedDesc.Name))
 {
-	Serialize(InMetadata ? InMetadata->GetDefaultMetadataDomain() : nullptr, InArchive);
+	CachedDesc.Name = InName.Get(InParent->CachedDesc.Name);
+
+	Init();
+
+	ValueKeyOffset = InParent->GetValueKeyOffsetForChild();
+
+	if (ensure(DefaultValue.IsValid() && InParent->DefaultValue.IsValid()))
+	{
+		// Copy the default value from the parent.
+		DefaultValue.CopyFrom(InParent->DefaultValue.GetRawPtr());
+	}
 }
 
-const UPCGMetadata* FPCGMetadataAttributeBase::GetMetadata() const
+FPCGMetadataAttributeBase::FPCGMetadataAttributeBase(const FPCGMetadataAttributeDesc& InAttributeDesc, const FPCGMetadataAttributeBase* InParent, FPCGMetadataDomain* InMetadata, bool bInAllowsInterpolation)
+	: Metadata(InMetadata)
+	, Parent(InParent)
+	, bAllowsInterpolation(bInAllowsInterpolation)
+	, CachedDesc(InAttributeDesc)
+	, UnderlyingProperty(InParent ? InParent->UnderlyingProperty : MakeShared<FPCGAttributeProperty>(CachedDesc))
+	, Name(InAttributeDesc.Name)
 {
-	return Metadata ? Metadata->GetTopMetadata() : nullptr;
+	check(!InParent || InParent->GetAttributeDesc().IsSameType(CachedDesc));
+
+	Init();
+
+	ValueKeyOffset = InParent ? InParent->GetValueKeyOffsetForChild() : 0;
+
+	if (InParent && ensure(DefaultValue.IsValid() && InParent->DefaultValue.IsValid()))
+	{
+		// Copy the default value from the parent.
+		DefaultValue.CopyFrom(InParent->DefaultValue.GetRawPtr());
+	}
+}
+
+FPCGMetadataAttributeBase::~FPCGMetadataAttributeBase()
+{
+	// We need to make sure to empty the array so the elements are destroyed and the memory is freed.
+	if (TOptional<FScriptArrayHelper> Helper = MakeArrayHelper(); ensure(Values.IsEmpty() || Helper.IsSet()) && Helper)
+	{
+		Helper->EmptyValues();
+	}
 }
 
-PRAGMA_ENABLE_DEPRECATION_WARNINGS
 
-void FPCGMetadataAttributeBase::Serialize(FPCGMetadataDomain* InMetadata, FArchive& InArchive)
+void FPCGMetadataAttributeBase::GetResourceSizeEx(FResourceSizeEx& CumulativeResourceSize) const
+{
+	{
+		PCG::TSharedScopeLock ReadLock(EntryMapLock);
+		CumulativeResourceSize.AddDedicatedSystemMemoryBytes(EntryToValueKeyMap.GetAllocatedSize());
+	}
+
+	PCG::TSharedScopeLock ReadLock(ValueLock);
+	// Note - can replace this with serialize call once it's there.
+	FArchiveCountMem Archive(nullptr);
+	// Serialize desc
+	FPCGMetadataAttributeDesc::StaticStruct()->SerializeItem(FStructuredArchiveFromArchive(Archive).GetSlot(), const_cast<FPCGMetadataAttributeDesc*>(&CachedDesc), /*Defaults*/nullptr);
+
+	if (UnderlyingProperty && UnderlyingProperty->Property)
+	{
+		if (UnderlyingProperty->Property->Inner && DefaultValue.GetRawPtr())
+		{
+			// Serialize default value
+			UnderlyingProperty->Property->Inner->SerializeItem(FStructuredArchiveFromArchive(Archive).GetSlot(), const_cast<void*>(DefaultValue.GetRawPtr()), /*Defaults*/nullptr);
+		}
+
+		// Serialize values
+		UnderlyingProperty->Property->SerializeItem(FStructuredArchiveFromArchive(Archive).GetSlot(), const_cast<FScriptArray*>(&Values), /*Defaults*/nullptr);
+	}
+
+	CumulativeResourceSize.AddDedicatedSystemMemoryBytes(Archive.GetNum());
+}
+
+PCGMetadataValueKey FPCGMetadataAttributeBase::GetValueKeyOffsetForChild() const
+{
+	PCG::TSharedScopeLock ReadLock(ValueLock);
+	return Values.Num() + ValueKeyOffset;
+}
+
+void FPCGMetadataAttributeBase::BaseSerialize(FPCGMetadataDomain* InMetadata, FArchive& InArchive)
 {
 	InArchive << EntryToValueKeyMap;
 	Metadata = InMetadata;
@@ -82,10 +622,1414 @@ void FPCGMetadataAttributeBase::Serialize(FPCGMetadataDomain* InMetadata, FArchi
 	InArchive << AttributeId;
 }
 
-void FPCGMetadataAttributeBase::GetResourceSizeEx(FResourceSizeEx& CumulativeResourceSize) const
+void FPCGMetadataAttributeBase::Serialize(FPCGMetadataDomain* InMetadata, FArchive& InArchive)
+{
+	BaseSerialize(InMetadata, InArchive);
+
+	// If we don't have a parent, we need to serialize the CachedDesc
+	if (!Parent)
+	{
+		FPCGMetadataAttributeDesc::StaticStruct()->SerializeItem(FStructuredArchiveFromArchive(InArchive).GetSlot(), &CachedDesc, /*Defaults*/nullptr);
+		if (InArchive.IsLoading() && !UnderlyingProperty)
+		{
+			UnderlyingProperty = MakeShared<FPCGAttributeProperty>(CachedDesc);
+		}
+	}
+	else if (InArchive.IsLoading() )
+	{
+		CachedDesc = Parent->CachedDesc;
+		CachedDesc.Name = Name;
+		UnderlyingProperty = Parent->UnderlyingProperty;
+	}
+
+	if (InArchive.IsLoading() && !DefaultValue.IsValid())
+	{
+		// Initialize the attribute
+		Init();
+	}
+
+	if (UnderlyingProperty && UnderlyingProperty->IsValid())
+	{
+		// Then serialize the default value
+		if (DefaultValue.IsValid())
+		{
+			UnderlyingProperty->Property->Inner->SerializeItem(FStructuredArchiveFromArchive(InArchive).GetSlot(), DefaultValue.GetRawPtr(), /*Defaults*/nullptr);
+		}
+
+		// And finally serialize the values
+		UnderlyingProperty->Property->SerializeItem(FStructuredArchiveFromArchive(InArchive).GetSlot(), &Values, /*Defaults*/nullptr);
+
+		if (InArchive.IsLoading())
+		{
+			ValueKeyOffset = GetParent() ? GetParent()->GetValueKeyOffsetForChild() : 0;
+		}
+	}
+}
+
+void FPCGMetadataAttributeBase::AddReferencedObjects(FReferenceCollector& Collector)
+{
+	// @todo_pcg: We should register all the stored objects
+	// We'll have to see the cost of this, and if we don't hold object references for too long.
+}
+
+TArray<PCGMetadataValueKey> FPCGMetadataAttributeBase::AddValues_Internal(const PCG::Private::FInValues& InValues, int32 Count, int32 ExistingIndex, bool bLockless, TOptional<TConstArrayView<int32>> OptionalIndices)
 {
-	PCG::TSharedScopeLock ReadLock(EntryMapLock);
-	CumulativeResourceSize.AddDedicatedSystemMemoryBytes(EntryToValueKeyMap.GetAllocatedSize());
+	using namespace PCG::Private;
+	
+	check(!OptionalIndices || Count == OptionalIndices->Num());
+	
+	if (const FInValuesSubset* InValuesSubset = InValues.TryGet<FInValuesSubset>())
+	{
+		if (OptionalIndices)
+		{
+			// If we have multiple layers of subsets, we need to extract the indices.
+			// Note that the order we extract the layers is inner to outer, so Optional indices MUST be always smaller than the indices in the subset.
+			// For example, if we have FirstLayer = Subset(OriginalValues, [0,1,3]) and SecondLayer = Subset(FirstLayer, [0, 2]), the first time we call AddValues_Internal,
+			// we pass `SecondLayer`. We then call the function again with FirstLayer and OptionalIndices = [0, 2].
+			// The second time, we call the function again with OriginalValues and OptionalIndices = [0, 3] (First and Third indices ([0, 2]) of the FirstLayer subset indices [0, 1, 3])
+			check(InValuesSubset->Indices.Num() >= OptionalIndices->Num())
+
+			TArray<int32, TInlineAllocator<256>> IndicesSubset;
+			IndicesSubset.Reserve(OptionalIndices->Num());
+			for (int32 i : *OptionalIndices)
+			{
+				IndicesSubset.Add(InValuesSubset->Indices[i]);
+			}
+			
+			return AddValues_Internal(*static_cast<const FInValues*>(InValuesSubset->InValues), IndicesSubset.Num(), ExistingIndex, bLockless, MakeConstArrayView(IndicesSubset));
+		}
+
+		check(InValuesSubset->Indices.Num() <= Count);
+		
+		// Otherwise, forward the call with the underlying values and its indices.
+		return AddValues_Internal(*static_cast<const FInValues*>(InValuesSubset->InValues), InValuesSubset->Indices.Num(), ExistingIndex, bLockless, MakeConstArrayView(InValuesSubset->Indices));
+	}
+
+	TArray<PCGMetadataValueKey> ValueKeys;
+
+	int32 StartWriteIndex = INDEX_NONE;
+
+	TOptional<FScriptArrayHelper> Helper = MakeArrayHelper();
+	if (!Helper.IsSet() || Count <= 0)
+	{
+		return ValueKeys;
+	}
+
+	// Hard-requirement that when we write to the default value, we can only write a single value
+	if (!ensure(ExistingIndex != ExistingIndexForDefaultValue || Count == 1))
+	{
+		Count = 1;
+	}
+
+	{
+		PCG::TUniqueScopeLock WriteLock(ValueLock, /*bShouldLock=*/!bLockless);
+		if (ExistingIndex == INDEX_NONE)
+		{
+			check(!bLockless);
+			StartWriteIndex = AllocateValues_Unsafe(Count);
+		}
+		else if (ExistingIndex != ExistingIndexForDefaultValue)
+		{
+			// Expected to be already allocated.
+			check(Helper->Num() >= ExistingIndex + Count);
+			StartWriteIndex = ExistingIndex;
+		}
+		else
+		{
+			// Expected the default value to always be allocated.
+			check(DefaultValue.GetRawPtr() != nullptr);
+			StartWriteIndex = ExistingIndex;
+		}
+
+		check(StartWriteIndex != INDEX_NONE);
+
+		auto GetRawPtr = [this, &Helper, StartWriteIndex](const int32 Index)
+		{
+			return StartWriteIndex == ExistingIndexForDefaultValue ? DefaultValue.GetRawPtr() : Helper->GetRawPtr(Index);
+		};
+		
+		auto Loop = [&OptionalIndices, Count, StartWriteIndex](auto Callback)
+		{
+			if (!OptionalIndices)
+			{
+				for (int32 i = 0; i < Count; ++i)
+				{
+					Callback(/*ReadIndex=*/i, /*WriteIndex=*/StartWriteIndex + i);
+				}
+			}
+			else
+			{
+				for (int32 i = 0; i < Count; ++i)
+				{
+					Callback(/*ReadIndex=*/(*OptionalIndices)[i], /*WriteIndex=*/StartWriteIndex + i);
+				}
+			}
+		};
+
+		Visit([this, &GetRawPtr, StartWriteIndex, Count, &OptionalIndices, &Loop](auto&& Value)
+			{
+				using T = std::decay_t<decltype(Value)>;
+				if constexpr(std::is_same_v<T, FInValuesByValue>)
+				{
+					if (!OptionalIndices)
+					{
+						const void* SrcPtr = Value.InValues;
+						void* DestPtr = GetRawPtr(StartWriteIndex);
+						UnderlyingProperty->Copy(DestPtr, SrcPtr, Count);
+					}
+					else
+					{
+						for (int32 i = 0; i < OptionalIndices->Num(); ++i)
+						{
+							const void* SrcPtr = UnderlyingProperty->GetPtrInArray(Value.InValues, (*OptionalIndices)[i]);
+							void* DestPtr = GetRawPtr(StartWriteIndex + i);
+							UnderlyingProperty->Copy(DestPtr, SrcPtr, 1);
+						}
+					}
+				}
+				else if constexpr(std::is_same_v<T, FInValuesSubset>)
+				{
+					// Must be caught earlier
+					checkNoEntry();
+				}
+				else if constexpr (std::is_same_v<T, FInValuesAsSet>)
+				{
+					const FSetProperty* InnerProperty = CastFieldChecked<FSetProperty>(UnderlyingProperty->Property->Inner);
+
+					Loop([InnerProperty, &GetRawPtr, &Value](int32 ReadIndex, int32 WriteIndex)
+					{
+						FScriptSetHelper OtherHelper(InnerProperty, GetRawPtr(WriteIndex));
+						const int32 Num = Value.InValues[ReadIndex].Num();
+
+						for (int j = 0; j < Num; ++j)
+						{
+							OtherHelper.AddElement(Value.InValues[ReadIndex][j]);
+						}
+					});
+				}
+				else if constexpr (std::is_same_v<T, FInValuesAsMap>)
+				{
+					const FMapProperty* InnerProperty = CastFieldChecked<FMapProperty>(UnderlyingProperty->Property->Inner);
+
+					Loop([InnerProperty, &GetRawPtr, &Value](int32 ReadIndex, int32 WriteIndex)
+					{
+						FScriptMapHelper OtherHelper(InnerProperty, GetRawPtr(WriteIndex));
+						const int32 Num = Value.InValues[ReadIndex].Num();
+
+						for (int j = 0; j < Num; ++j)
+						{
+							OtherHelper.AddPair(Value.InValues[ReadIndex][j].Key, Value.InValues[ReadIndex][j].Value);
+						}
+					});
+				}
+				else if constexpr (std::is_same_v<T, FInValuesAsArray>)
+				{
+					const FArrayProperty* InnerProperty = CastFieldChecked<FArrayProperty>(UnderlyingProperty->Property->Inner);
+
+					Loop([InnerProperty, &GetRawPtr, &Value](int32 ReadIndex, int32 WriteIndex)
+					{
+						FScriptArrayHelper OtherHelper(InnerProperty, GetRawPtr(WriteIndex));
+						const int32 Num = Value.InValues[ReadIndex].template Get<1>();
+						if (InnerProperty->HasAnyPropertyFlags(CPF_IsPlainOldData))
+						{
+							OtherHelper.AddUninitializedValues(Num);
+						}
+						else
+						{
+							OtherHelper.AddValues(Num);
+						}
+
+						PCG::Private::CopyArray(InnerProperty, OtherHelper.GetRawPtr(), Value.InValues[ReadIndex].template Get<0>(), Num);
+					});
+				}
+				else
+				{
+					static_assert(!std::is_same_v<T, T>, "Missing variant case for AddValues_Internal");
+				}
+			}, InValues);
+	}
+
+	ValueKeys.SetNum(Count);
+	for (int ValueIndex = 0; ValueIndex < Count; ++ValueIndex)
+	{
+		ValueKeys[ValueIndex] = ValueKeyOffset + StartWriteIndex + ValueIndex;
+	}
+
+	return ValueKeys;
+}
+
+TArray<PCGMetadataValueKey> FPCGMetadataAttributeBase::AddCompressedValues(PCG::Private::FInValues InUniqueValues, int32 Count, TArray<PCGMetadataValueKey>& FoundValueKeys)
+{
+	using namespace PCG::Private;
+
+	// For now only support FInValuesByValue or FInValuesSubset
+	if (!ensure(InUniqueValues.IsType<FInValuesByValue>() || InUniqueValues.IsType<FInValuesSubset>()))
+	{
+		return {};
+	}
+	
+	// FoundValueKeys should always be bigger than Count (since FoundValueKeys should represent the full range of values to add, while InUniqueValues/Count represent the unique values)
+	if (!ensure(Count <= FoundValueKeys.Num()))
+	{
+		return {};
+	}
+
+	const bool bHasDuplicateValues = (Count != FoundValueKeys.Num());
+
+	TArray<PCGMetadataValueKey> FoundUniqueValueKeys;
+	TArray<PCGMetadataValueKey>& FoundKeys = (bHasDuplicateValues ? FoundUniqueValueKeys : FoundValueKeys);
+
+	// Implementation note: when we don't have any duplicate values, the previously-set values in FoundValueKeys will be wiped out - this is intended
+	const bool bAtLeastOneValueNotFound = !FindValues(FoundKeys, InUniqueValues, Count);
+
+	if (bAtLeastOneValueNotFound)
+	{
+		// If there is at least one value not found, go through all the found keys (from the original unique values)
+		// and keep the indices of all the original values that need to be added.
+		FInValuesSubset ValueSubset{};
+		ValueSubset.InValues = &InUniqueValues;
+		ValueSubset.Indices.Reserve(Count);
+
+		for (int ValueIndex = 0; ValueIndex < Count; ++ValueIndex)
+		{
+			if (FoundKeys[ValueIndex] == PCGNotFoundValueKey)
+			{
+				ValueSubset.Indices.Add(ValueIndex);
+			}
+		}
+
+		const int32 SubCount = ValueSubset.Indices.Num();
+		TArray<PCGMetadataValueKey> NewValueKeys = AddValues_Internal(FInValues{TInPlaceType<FInValuesSubset>{}, MoveTemp(ValueSubset)}, SubCount);
+
+		int32 NewlyAddedValueKeyIndex = 0;
+		for (int ValueIndex = 0; ValueIndex < Count; ++ValueIndex)
+		{
+			if (FoundKeys[ValueIndex] == PCGNotFoundValueKey)
+			{
+				FoundKeys[ValueIndex] = NewValueKeys[NewlyAddedValueKeyIndex++];
+			}
+		}
+	}
+
+	// Remap to full array if needed
+	if (bHasDuplicateValues)
+	{
+		for (PCGMetadataValueKey& ValueToRemap : FoundValueKeys)
+		{
+			ValueToRemap = FoundUniqueValueKeys[ValueToRemap];
+		}
+	}
+
+	return FoundValueKeys;
+}
+
+TArray<PCGMetadataValueKey> FPCGMetadataAttributeBase::AddCompressedValues(PCG::Private::FInValuesByValue InValues, int32 Count)
+{
+	using namespace PCG::Private;
+	
+	if (!UnderlyingProperty || !UnderlyingProperty->IsValid())
+	{
+		return {};
+	}
+
+	check(UnderlyingProperty->bCompressData);
+
+	// Since we're getting raw values here, we might have duplicates
+	// so we should aim to remove duplicates here so we preserve our 'compress data' idea, otherwise it will break other foundational blocks (e.g. partition)
+	TArray<int32> UniqueValuesIndices;
+
+	auto Compare = [this, &InValues](int32 LHS, int32 RHS)
+	{	
+		const void* Value1 = UnderlyingProperty->GetPtrInArray(InValues.InValues, LHS);
+		const void* Value2 = UnderlyingProperty->GetPtrInArray(InValues.InValues, RHS);
+			
+		return Value1 != nullptr && Value2 != nullptr && UnderlyingProperty->Property->Inner->Identical(Value1, Value2);
+	};
+
+	// Initially, fill with mapping to unique values so we can remap them at the end if needed
+	TArray<PCGMetadataValueKey> FoundValueKeys;
+	FoundValueKeys.Reserve(Count);
+
+	const FProperty* InnerProperty = UnderlyingProperty->Property->Inner;
+	
+	// Hash-accelerated deduplication: O(n) average case.
+	// Algorithm:
+	//	- For each value in our incoming values
+	//		* Get the value pointer
+	//		* Compute the hash for this value
+	//		* Find the hash in our map of already computed hashes
+	//		* For each unique value with the same hash (collision), do the comparison
+	//		* If it is found, we have our value key
+	//		* If it is not found, we have a new unique value, so add it to our map.
+	//
+	// This reduces drastically the amount of comparisons to do at the cost of extra memory.
+	if (InnerProperty->HasAllPropertyFlags(CPF_HasGetValueTypeHash))
+	{
+		TMap<uint32, TArray<int32, TInlineAllocator<4>>> HashBuckets;
+
+		for (int32 i = 0; i < Count; ++i)
+		{
+			const void* ValuePtr = UnderlyingProperty->GetPtrInArray(InValues.InValues, i);
+
+			int32 FoundValueKey = INDEX_NONE;
+
+			if (ValuePtr)
+			{
+				const uint32 Hash = InnerProperty->GetValueTypeHash(ValuePtr);
+
+				if (TArray<int32, TInlineAllocator<4>>* Bucket = HashBuckets.Find(Hash))
+				{
+					// Check within collision bucket using Identical for correctness
+					for (const int32 UniqueIdx : *Bucket)
+					{
+						if (Compare(i, UniqueValuesIndices[UniqueIdx]))
+						{
+							FoundValueKey = UniqueIdx;
+							break;
+						}
+					}
+				}
+
+				if (FoundValueKey == INDEX_NONE)
+				{
+					FoundValueKey = UniqueValuesIndices.Add(i);
+					HashBuckets.FindOrAdd(Hash).Add(FoundValueKey);
+				}
+			}
+			else
+			{
+				// Null pointer - treat as unique (consistent with Compare returning false for nulls)
+				FoundValueKey = UniqueValuesIndices.Add(i);
+			}
+
+			FoundValueKeys.Emplace(FoundValueKey);
+		}
+	}
+	else
+	{
+		// Fallback for types without hash support: O(n*m) where n = Count and m is the number of unique values (worst case: all unique values -> n^2)
+		for (int32 i = 0; i < Count; ++i)
+		{
+			int32 FoundValueKey = INDEX_NONE;
+			for (int32 j = 0; j < UniqueValuesIndices.Num(); ++j)
+			{
+				if (Compare(i, UniqueValuesIndices[j]))
+				{
+					FoundValueKey = j;
+					break;
+				}
+			}
+
+			if (FoundValueKey == INDEX_NONE)
+			{
+				FoundValueKey = UniqueValuesIndices.Add(i);
+			}
+
+			FoundValueKeys.Emplace(FoundValueKey);
+		}
+	}
+
+	FInValues TempValues{TInPlaceType<FInValuesByValue>{}, MoveTemp(InValues)};
+	const int32 UniqueCount = UniqueValuesIndices.Num();
+	return AddCompressedValues(FInValues{TInPlaceType<FInValuesSubset>{}, &TempValues, MoveTemp(UniqueValuesIndices)}, UniqueCount, FoundValueKeys);
+}
+
+bool FPCGMetadataAttributeBase::FindValues(TArray<PCGMetadataValueKey>& OutValueKeys, const PCG::Private::FInValues& InValues, int32 Count) const
+{
+	OutValueKeys.Init(PCGNotFoundValueKey, Count);
+
+	return FindValues(MakeArrayView(OutValueKeys), InValues, Count);
+}
+
+bool FPCGMetadataAttributeBase::FindValues(TArrayView<PCGMetadataValueKey> OutValueKeys, const PCG::Private::FInValues& InValues, int32 Count) const
+{
+	check(OutValueKeys.Num() == Count);
+	int ValueKeysSet = 0;
+	FindValuesInternal(InValues, Count, OutValueKeys, ValueKeysSet, /*bIsRoot=*/true);
+
+	return (ValueKeysSet == Count);
+}
+
+void FPCGMetadataAttributeBase::FindValuesInternal(const PCG::Private::FInValues& InValues, int32 Count, TArrayView<PCGMetadataValueKey> ValueKeys, int& ValueKeysSet, bool bIsRoot) const
+{
+	using namespace PCG::Private;
+
+	check(Count == ValueKeys.Num());
+
+	if (!ensure(UnderlyingProperty))
+	{
+		return;
+	}
+	
+	// For now only support FInValuesByValue and FInValuesSubset
+	const TArray<int32>* OptionalIndices = nullptr;
+	const void* UnderlyingValues = nullptr;
+
+	Visit([this, &OptionalIndices, &UnderlyingValues](auto&& InValue)
+	{
+		using T = std::decay_t<decltype(InValue)>;
+		if constexpr (std::is_same_v<T, FInValuesAsArray> || std::is_same_v<T, FInValuesAsSet> || std::is_same_v<T, FInValuesAsMap>)
+		{
+			// Unsupported
+			return;
+		}
+		else if constexpr (std::is_same_v<T, FInValuesByValue>)
+		{
+			UnderlyingValues = InValue.InValues;
+		}
+		else if constexpr (std::is_same_v<T, FInValuesSubset>)
+		{
+			const FInValues* UnderlyingVariant = static_cast<const FInValues*>(InValue.InValues);
+			
+			if (const FInValuesByValue* UnderlyingValue = UnderlyingVariant->TryGet<FInValuesByValue>())
+			{
+				UnderlyingValues = UnderlyingValue->InValues;
+				OptionalIndices = &InValue.Indices;
+			}
+			else
+			{
+				// Unsupported
+				return;
+			}
+		}
+		else
+		{
+			static_assert(!std::is_same_v<T, T>, "Missing variant case for FindValuesInternal");
+		}
+	}, InValues);
+
+	if (!UnderlyingValues)
+	{
+		// Unsupported
+		return;
+	}
+
+	auto Compare = [this, UnderlyingValues](int32 InValueIndex, PCGMetadataValueKey ValueKey)
+	{
+		const void* Value1 = GetReadAddressForValueKey_Unsafe(ValueKey);
+		const void* Value2 = UnderlyingProperty->GetPtrInArray(UnderlyingValues, InValueIndex);
+
+		return Value1 != nullptr && Value2 != nullptr && UnderlyingProperty->Property->Inner->Identical(Value1, Value2);
+	};
+
+	if (bIsRoot)
+	{
+		for (int ValueIndex = 0; ValueIndex < Count; ++ValueIndex)
+		{
+			if (Compare(OptionalIndices ? (*OptionalIndices)[ValueIndex] : ValueIndex, PCGDefaultValueKey))
+			{
+				ValueKeys[ValueIndex] = PCGDefaultValueKey;
+				++ValueKeysSet;
+			}
+		}
+	}
+
+	if (ValueKeysSet != Count)
+	{
+		const FPCGMetadataAttributeBase* ThisParent = GetParent();
+		if (ThisParent)
+		{
+			ThisParent->FindValuesInternal(InValues, Count, ValueKeys, ValueKeysSet, /*bRoot=*/false);
+		}
+	}
+
+	if (ValueKeysSet != Count)
+	{
+		PCG::TSharedScopeLock ReadLock(ValueLock, /*bShouldLock=*/bIsRoot);
+		for (int ValueIndex = 0; ValueIndex < Count; ++ValueIndex)
+		{
+			if (ValueKeys[ValueIndex] != PCGNotFoundValueKey)
+			{
+				continue;
+			}
+
+			for (PCGMetadataValueKey ValueKey = 0; ValueKey < Values.Num(); ++ValueKey)
+			{
+				const PCGMetadataValueKey OffsetKey = ValueKey + ValueKeyOffset;
+				if (Compare(OptionalIndices ? (*OptionalIndices)[ValueIndex] : ValueIndex, OffsetKey))
+				{
+					ValueKeys[ValueIndex] = OffsetKey;
+					++ValueKeysSet;
+					break;
+				}
+			}
+		}
+	}
+}
+
+void FPCGMetadataAttributeBase::Init()
+{
+	if (!ensure(UnderlyingProperty) || !ensure(!DefaultValue.IsValid()) || !UnderlyingProperty->IsValid())
+	{
+		return;
+	}
+
+	// In order to be consistent between multiple equivalent descriptors, we will convert structs that matches a predefined type.
+	CachedDesc.FixLegacyTypeId();
+
+	// Initialize the default value
+	DefaultValue.Init(UnderlyingProperty);
+}
+
+void FPCGMetadataAttributeBase::Flatten()
+{
+	// Implementation notes:
+	// We don't need to flatten the EntryToValueKeyMap - this will have been taken care of in the metadata flatten
+
+	// Flatten values, from root to current attribute
+	if (Parent && UnderlyingProperty)
+	{
+		PCG::TUniqueScopeLock WriteLock(ValueLock);
+
+		TArray<const FScriptArray*> OriginalValues;
+
+		int32 ValueCount = 0;
+
+		const FPCGMetadataAttributeBase* Current = this;
+		while (Current)
+		{
+			ValueCount += Current->Values.Num();
+			OriginalValues.Add(&Current->Values);
+
+			if (Current->Parent)
+			{
+				Current = Current->Parent;
+			}
+			else
+			{
+				Current = nullptr;
+			}
+		}
+
+		// Swap the current values to an empty one
+		FScriptArray NewValues{};
+		FScriptArrayHelper NewValuesHelper{UnderlyingProperty->Property.Get(), &NewValues};
+
+		if (UnderlyingProperty->bIsPlainOldData)
+		{
+			NewValuesHelper.AddUninitializedValues(ValueCount);
+		}
+		else
+		{
+			NewValuesHelper.AddValues(ValueCount);
+		}
+
+		// Then copy all the values in reverse order (oldest first)
+		int32 Count = 0;
+		for (int32 ValuesIndex = OriginalValues.Num() - 1; ValuesIndex >= 0; --ValuesIndex)
+		{
+			FScriptArrayHelper OtherHelper{UnderlyingProperty->Property.Get(), OriginalValues[ValuesIndex]};
+			const int32 OtherCount = OtherHelper.Num();
+			UnderlyingProperty->Copy(NewValuesHelper.GetRawPtr(Count), OtherHelper.GetRawPtr(0), /*Count=*/OtherCount);
+			Count += OtherCount;
+		}
+
+		// Finally move the new values in place of the old values.
+		if (TOptional<FScriptArrayHelper> Helper = MakeArrayHelper(); ensure(Helper.IsSet()))
+		{
+			Helper->MoveAssign(&NewValues);
+		}
+	}
+
+	// Reset value offset, and lose parent
+	ValueKeyOffset = 0;
+	Parent = nullptr;
+}
+
+void FPCGMetadataAttributeBase::FlattenAndCompress(const TArrayView<const PCGMetadataEntryKey>& InEntryKeysToKeep)
+{
+	// No entries, we can just delete everything in the attribute.
+	if (InEntryKeysToKeep.IsEmpty() || !UnderlyingProperty)
+	{
+		Reset();
+		return;
+	}
+
+	TArray<PCGMetadataValueKey> AllValueKeys;
+	TArray<PCGMetadataValueKey> AllUniqueValueKeys;
+	AllValueKeys.Reserve(InEntryKeysToKeep.Num());
+	bool bUseValueKeys = UsesValueKeys();
+	if (bUseValueKeys)
+	{
+		AllUniqueValueKeys.Reserve(InEntryKeysToKeep.Num());
+	}
+
+	// First gather all value keys associated with the entry keys to keep.
+	// If we compress data, we also store the unique value keys used (that is not default).
+	for (const PCGMetadataEntryKey& EntryKey : InEntryKeysToKeep)
+	{
+		AllValueKeys.Add(GetValueKey(EntryKey));
+		if (bUseValueKeys)
+		{
+			if (AllValueKeys.Last() != PCGDefaultValueKey)
+			{
+				AllUniqueValueKeys.AddUnique(AllValueKeys.Last());
+			}
+		}
+	}
+
+	// Then for each value key (or unique values keys), gather the value in a NewValues array
+	// and also keep a mapping between old value key and new value key.
+	// Only done if the old value key is not the default one.
+
+	TMap<PCGMetadataValueKey, PCGMetadataValueKey> ValueKeyMapping;
+	FScriptArray NewValues{};
+	FScriptArrayHelper NewValuesHelper{UnderlyingProperty->Property.Get(), &NewValues};
+
+	const TArray<PCGMetadataValueKey>& AllValueKeysRef = !bUseValueKeys ? AllValueKeys : AllUniqueValueKeys;
+	if (UnderlyingProperty->bIsPlainOldData)
+	{
+		NewValuesHelper.AddUninitializedValues(AllValueKeysRef.Num());
+	}
+	else
+	{
+		NewValuesHelper.AddValues(AllValueKeysRef.Num());
+	}
+
+	ValueKeyMapping.Reserve(AllValueKeysRef.Num());
+
+	int32 Count = 0;
+	for (PCGMetadataValueKey ValueKey : AllValueKeysRef)
+	{
+		if (ValueKey != PCGDefaultValueKey)
+		{
+			ValueKeyMapping.Add(ValueKey, Count);
+
+			// @todo_pcg: Improve to find all contiguous values before copying.
+			const void* SrcPtr = GetReadAddressForValueKey_Unsafe(ValueKey);
+			UnderlyingProperty->Copy(NewValuesHelper.GetRawPtr(Count++), SrcPtr, /*Count=*/1);
+		}
+	}
+
+	// Move the new values in place of the old values.
+	if (TOptional<FScriptArrayHelper> Helper = MakeArrayHelper(); ensure(Helper.IsSet()))
+	{
+		PCG::TUniqueScopeLock WriteLock(ValueLock);
+		Helper->MoveAssign(&NewValues);
+	}
+
+	// And finally, create a new entry to value mapping.
+	// Logic is that each entry to keep will have their "index" as new entry key
+	// (like if the entries to keep are [25, 47, 54], the new entries would be [0, 1, 2]).
+	// So the operation is:
+	// All pairs Old EK -> Old VK transform to New EK -> New VK.
+	TMap<PCGMetadataEntryKey, PCGMetadataValueKey> NewMap;
+	for (int32 i = 0; i < InEntryKeysToKeep.Num(); ++i)
+	{
+		PCGMetadataValueKey ValueKey = AllValueKeys[i];
+		if (ValueKey != PCGDefaultValueKey && ensure(InEntryKeysToKeep[i] != PCGInvalidEntryKey))
+		{
+			NewMap.Add(i, ValueKeyMapping[ValueKey]);
+		}
+	}
+
+	// And move the map.
+	{
+		PCG::TUniqueScopeLock WriteLock(EntryMapLock);
+		EntryToValueKeyMap = MoveTemp(NewMap);
+	}
+
+	// At the end, reset value offset, and lose parent.
+	ValueKeyOffset = 0;
+	Parent = nullptr;
+}
+
+void FPCGMetadataAttributeBase::Reset()
+{
+	ValueKeyOffset = 0;
+	Parent = nullptr;
+
+	{
+		PCG::TUniqueScopeLock WriteLock(EntryMapLock);
+		EntryToValueKeyMap.Empty();
+	}
+
+	if (TOptional<FScriptArrayHelper> Helper = MakeArrayHelper(); Helper.IsSet())
+	{
+		PCG::TUniqueScopeLock WriteLock(ValueLock);
+		Helper->EmptyValues();
+	}
+}
+
+FPCGMetadataAttributeBase* FPCGMetadataAttributeBase::Copy(FName NewName, FPCGMetadataDomain* InMetadata, bool bKeepParent, bool bCopyEntries, bool bCopyValues) const
+{
+	// If we copy an attribute where we don't want to keep the parent, while copying entries and/or values, we'll lose data.
+	// In that case, we will copy all the data from this attribute and all its ancestors.
+
+	// We can't keep the parent if we don't have the same root.
+	checkSlow(!bKeepParent || PCGMetadataHelpers::HasSameRoot(Metadata, InMetadata));
+
+	// Validate that the new name is valid
+	if (!IsValidName(NewName))
+	{
+		UE_LOG(LogPCG, Error, TEXT("Try to create a new attribute with an invalid name: %s"), *NewName.ToString());
+		return nullptr;
+	}
+
+	// This copies to a new attribute.
+	FPCGMetadataAttributeBase* AttributeCopy = nullptr;
+	if (bKeepParent)
+	{
+		AttributeCopy = new FPCGMetadataAttributeBase(this, InMetadata, bAllowsInterpolation, NewName);
+	}
+	else
+	{
+		FPCGMetadataAttributeDesc NewDesc = CachedDesc;
+		NewDesc.Name = NewName;
+		AttributeCopy = new FPCGMetadataAttributeBase(NewDesc, InMetadata, bAllowsInterpolation);
+	}
+
+	CopyInternal(AttributeCopy, bKeepParent, bCopyEntries, bCopyValues);
+
+	return AttributeCopy;
+}
+
+FPCGMetadataAttributeBase* FPCGMetadataAttributeBase::CopyToAnotherType(int16 Type) const
+{
+	// This was done to support casting from one type to another, where we knew at compile time both types (source and target).
+	// Without it, we would need a way at runtime to know what is compatible and how to convert the values. In the meantime, this is not supported.
+	return nullptr;
+}
+
+void FPCGMetadataAttributeBase::CopyInternal(FPCGMetadataAttributeBase* NewAttribute, bool bKeepParent, bool bCopyEntries, bool bCopyValues, FCopyFunctor CopyFunctor) const
+{
+	check(NewAttribute);
+
+	if (!CopyFunctor)
+	{
+		// If we do not have an explicit functor, just use the default copy.
+		CopyFunctor = [&UnderlyingProperty = this->UnderlyingProperty](void* DestData, const void* SrcData, const int32 Count)
+		{
+			UnderlyingProperty->Copy(DestData, SrcData, Count);
+		};
+	}
+
+	check(CopyFunctor);
+
+	if (ensure(DefaultValue.IsValid()) && ensure(NewAttribute->DefaultValue.IsValid()))
+	{
+		// Copy this default value to the attribute copy.
+		CopyFunctor(NewAttribute->DefaultValue.GetRawPtr(), DefaultValue.GetRawPtr(), 1);
+	}
+
+	// Gather the chain of parents if we don't keep the parent and we want to copy entries/values.
+	// We always have at least one item, "this".
+	TArray<const FPCGMetadataAttributeBase*, TInlineAllocator<2>> Parents = { this };
+	if (!bKeepParent && (bCopyEntries || bCopyValues))
+	{
+		const FPCGMetadataDomain* CurrentMetadata = Metadata;
+		const FPCGMetadataAttributeBase* Current = this;
+
+		const FPCGMetadataDomain* ParentMetadata = PCGMetadataHelpers::GetParentMetadata(CurrentMetadata);
+		while (ParentMetadata && Current->Parent)
+		{
+			CurrentMetadata = ParentMetadata;
+			Current = Current->Parent;
+			Parents.Add(Current);
+
+			ParentMetadata = PCGMetadataHelpers::GetParentMetadata(CurrentMetadata);
+		}
+	}
+
+	if (bCopyEntries)
+	{
+		// We go backwards, since we need to preserve order (root -> this)
+		// Latest entry in our Parents array is the root.
+		for (int32 i = Parents.Num() - 1; i >= 0; --i)
+		{
+			const FPCGMetadataAttributeBase* Current = Parents[i];
+
+			// Only need to lock if Current == this, since the parents are assumed to be const.
+			PCG::TSharedScopeLock ReadLock(Current->EntryMapLock, /*bShouldLock=*/Current == this);
+			NewAttribute->EntryToValueKeyMap.Append(Current->EntryToValueKeyMap);
+		}
+	}
+
+	if (bCopyValues)
+	{
+		// We go backwards, since we need to preserve order (root -> this)
+		// Latest entry in our Parents array is the root.
+		if (TOptional<FScriptArrayHelper> Helper = NewAttribute->MakeArrayHelper(); Helper.IsSet())
+		{
+			// Make sure to start from a fresh start, as we already have a default value allocated.
+			Helper->EmptyValues();
+
+			for (int32 i = Parents.Num() - 1; i >= 0; --i)
+			{
+				const FPCGMetadataAttributeBase* Current = Parents[i];
+
+				// Only need to lock if Current == this, since the parents are assumed to be const.
+				PCG::TSharedScopeLock ReadLock(Current->ValueLock, /*bShouldLock=*/Current == this);
+
+				TOptional<FScriptArrayHelper> ParentHelper = Current->MakeArrayHelper();
+				if (!ParentHelper.IsSet())
+				{
+					continue;
+				}
+
+				const int32 NumValuesToCopy = ParentHelper->Num();
+				if (NumValuesToCopy <= 0)
+				{
+					continue;
+				}
+
+				const int32 StartIndex = NewAttribute->AllocateValues_Unsafe(NumValuesToCopy);
+
+				if (StartIndex == INDEX_NONE)
+				{
+					continue;
+				}
+
+				CopyFunctor(Helper->GetRawPtr(StartIndex), ParentHelper->GetRawPtr(0), NumValuesToCopy);
+			}
+		}
+	}
+}
+
+void FPCGMetadataAttributeBase::SetValue(PCGMetadataEntryKey ItemKey, const FPCGMetadataAttributeBase* InAttribute, PCGMetadataEntryKey InEntryKey)
+{
+	using namespace PCG::Private;
+
+	if (!InAttribute)
+	{
+		return;
+	}
+
+	TOptional<FScriptArrayHelper> ThisHelper = MakeArrayHelper();
+	if (!ThisHelper.IsSet())
+	{
+		return;
+	}
+
+	check(ItemKey != PCGInvalidEntryKey);
+
+	const void* SrcPtr = InAttribute->GetReadAddressFromEntryKey_Unsafe(InEntryKey);
+	if (!SrcPtr)
+	{
+		return;
+	}
+
+	PCGMetadataValueKey ValueKeyToAdd = PCGNotFoundValueKey;
+
+	if (UnderlyingProperty->bCompressData)
+	{
+		const bool bValueKeyFound = FindValues(MakeArrayView(&ValueKeyToAdd, 1), FInValues{TInPlaceType<FInValuesByValue>{}, SrcPtr, 1}, 1);
+		check(!bValueKeyFound || ValueKeyToAdd != PCGNotFoundValueKey);
+	}
+
+	// If we are not compressed or we haven't find the value, we need to add it.
+	if (ValueKeyToAdd == PCGNotFoundValueKey)
+	{
+		PCG::TUniqueScopeLock WriteLock(ValueLock);
+
+		// Add the value in the buffer
+		const int32 NewValueIndex = ThisHelper->AddValue();
+
+		// If InAttribute == this, SrcPtr might be dangling now because of possible reallocation, so fix it there
+		if (InAttribute == this)
+		{
+			SrcPtr = InAttribute->GetReadAddressFromEntryKey_Unsafe(InEntryKey);
+		}
+
+		check(SrcPtr);
+
+		ValueKeyToAdd = ValueKeyOffset + NewValueIndex;
+		void* DestPtr = ThisHelper->GetRawPtr(NewValueIndex);
+		UnderlyingProperty->Copy(DestPtr, SrcPtr);
+	}
+
+	if (ensure(ValueKeyToAdd != PCGNotFoundValueKey))
+	{
+		PCG::TUniqueScopeLock WriteLock(EntryMapLock);
+		EntryToValueKeyMap.Emplace(ItemKey, ValueKeyToAdd);
+	}
+}
+
+void FPCGMetadataAttributeBase::SetZeroValue(PCGMetadataEntryKey ItemKey)
+{
+	// Not supported, will need to be specialized in the templated versions.
+}
+
+void FPCGMetadataAttributeBase::AccumulateValue(PCGMetadataEntryKey ItemKey, const FPCGMetadataAttributeBase* InAttributeA,
+	PCGMetadataEntryKey InEntryKeyA, float Weight)
+{
+	// Not supported, will need to be specialized in the templated versions.
+}
+
+void FPCGMetadataAttributeBase::SetWeightedValue(PCGMetadataEntryKey ItemKey, const FPCGMetadataAttributeBase* InAttribute,
+	const TArrayView<const TPair<PCGMetadataEntryKey, float>>& InWeightedKeys)
+{
+	// Not supported, will need to be specialized in the templated versions.
+}
+
+void FPCGMetadataAttributeBase::SetValue(PCGMetadataEntryKey ItemKey, const FPCGMetadataAttributeBase* InAttributeA,
+	PCGMetadataEntryKey InEntryKeyA, const FPCGMetadataAttributeBase* InAttributeB, PCGMetadataEntryKey InEntryKeyB, EPCGMetadataOp Op)
+{
+	// Not supported, will need to be specialized in the templated versions.
+}
+
+bool FPCGMetadataAttributeBase::IsEqualToDefaultValue(PCGMetadataValueKey ValueKey) const
+{
+	if (ValueKey == PCGDefaultValueKey)
+	{
+		return true;
+	}
+
+	if (ensure(UnderlyingProperty) && UnderlyingProperty->bCompressData)
+	{
+		return false;
+	}
+	else
+	{
+		return AreValuesEqual(ValueKey, PCGDefaultValueKey);
+	}
+}
+
+void FPCGMetadataAttributeBase::SetDefaultValueToFirstEntry()
+{
+	if (Values.Num() == 0 || Values.Num() > 1 || !ensure(UnderlyingProperty) || !DefaultValue.IsValid())
+	{
+		return;
+	}
+
+	{
+		PCG::TUniqueScopeLock WriteLock(ValueLock);
+
+		// First value key of this attribute.
+		const void* SrcPtr = GetReadAddressForValueKey_Unsafe(ValueKeyOffset);
+		if (!SrcPtr)
+		{
+			return;
+		}
+
+		void* DestPtr = DefaultValue.GetRawPtr();
+
+		UnderlyingProperty->Copy(DestPtr, SrcPtr);
+	}
+}
+
+bool FPCGMetadataAttributeBase::UsesValueKeys() const
+{
+	return ensure(UnderlyingProperty) ? UnderlyingProperty->bCompressData : false;
+}
+
+bool FPCGMetadataAttributeBase::AreValuesEqualForEntryKeys(PCGMetadataEntryKey EntryKey1, PCGMetadataEntryKey EntryKey2) const
+{
+	return AreValuesEqual(GetValueKey(EntryKey1), GetValueKey(EntryKey2));
+}
+
+bool FPCGMetadataAttributeBase::AreValuesEqual(PCGMetadataValueKey ValueKey1, PCGMetadataValueKey ValueKey2) const
+{
+	if (!ensure(UnderlyingProperty))
+	{
+		return false;
+	}
+
+	if (ValueKey1 == ValueKey2)
+	{
+		return true;
+	}
+
+	{
+		PCG::TSharedScopeLock ReadLock(ValueLock);
+
+		const void* Value1 = GetReadAddressForValueKey_Unsafe(ValueKey1);
+		const void* Value2 = GetReadAddressForValueKey_Unsafe(ValueKey2);
+
+		if (Value1 == nullptr || Value2 == nullptr)
+		{
+			return false;
+		}
+
+		return UnderlyingProperty->Property->Inner->Identical(Value1, Value2);
+	}
+}
+
+/**
+ * In this function, we want to gather all the values are that given by the values keys, but only those who are at this level (aka. ValueKey >= ValueKeyOffset).
+ * UnretrievedValues will hold all the values that didn't managed to get because they are stored higher in the attribute hierarchy. The algorithm stops when
+ * all the values are retrieved.
+ * Extraction of the values will depend on the FOutValues variant, since we need to distinguished between those 5 possibilities:
+ * - Output by Values (copy)
+ * - Output by Pointer
+ * - Output as Array
+ * - Output as Set
+ * - Output as Map
+ *
+ * On top of that, in order to reduce as much as we can virtual function calls, we also track for the different un-retrieved values if they are contiguous in memory, for input AND output.
+ * Values are contiguous in memory if their index in the output array are contiguous AND the values keys associated with them are also contiguous.
+ * For example, if we have indices {0, 1, 2, 5} (because we still haven't retrieved values for those indices), and we have value keys {1, 2, 8, 9}, we have a contiguous
+ * range for the first two (their indices AND values keys are contiguous), but not for the rest.
+ * So we can do 3 extractions (in this example [Index:ValueKey]): [0:1 ,1:2], [2:8] and [5:9] and optimize the first one to reduce the number of virtual calls.
+ */
+void FPCGMetadataAttributeBase::GetValues_Internal(const TArrayView<const PCGMetadataValueKey> ValueKeys, PCG::Private::FOutValues OutValues, TBitArray<>& UnretrievedValues) const
+{
+	using namespace PCG::Private;
+
+	bool bFoundAllKeys = true;
+	TConstSetBitIterator<> It(UnretrievedValues);
+	if (!It)
+	{
+		return;
+	}
+
+	const FPCGMetadataAttributeBase* ThisParent = static_cast<const FPCGMetadataAttributeBase*>(GetParent());
+	TOptional<FScriptArrayHelper> Helper = MakeArrayHelper();
+	if (!Helper)
+	{
+		return;
+	}
+
+	{
+		PCG::TSharedScopeLock ReadLock(ValueLock);
+
+		struct FContinuousRange
+		{
+			void Reset()
+			{
+				StartIndex = INDEX_NONE;
+				StartValueKey = INDEX_NONE;
+				Count = 0;
+			}
+
+			void Init(int32 InStartIndex, int64 InStartValueKey)
+			{
+				StartIndex = InStartIndex;
+				StartValueKey = InStartValueKey;
+				Count = 1;
+			}
+
+			int32 StartIndex = INDEX_NONE;
+			int64 StartValueKey = INDEX_NONE;
+			int32 Count = 0;
+		};
+
+		FContinuousRange ContinuousRange{};
+
+		auto RetrieveValues = [this, &OutValues, &Helper](FContinuousRange& Range)
+		{
+			const uint8* SrcPtr = nullptr;
+			if (Range.StartValueKey == PCG::Private::ExistingIndexForDefaultValue)
+			{
+				check(Range.Count == 1);
+				SrcPtr = static_cast<const uint8*>(DefaultValue.GetRawPtr());
+			}
+			else
+			{
+				SrcPtr = Helper->GetRawPtr(Range.StartValueKey);
+			}
+
+			check(SrcPtr);
+
+			// Since the values are stored in an array, they are all contiguous in memory so we can just do pointer arithmetic using the
+			// element size.
+			auto ApplyOnLoop = [this, &Range, SrcPtr](auto Func)
+			{
+				for (int32 i = 0; i < Range.Count; ++i)
+				{
+					Func(static_cast<const uint8*>(UnderlyingProperty->GetPtrInArray(SrcPtr, i)), Range.StartIndex + i);
+				}
+			};
+
+			Visit([this, &ApplyOnLoop, SrcPtr, &Range](auto&& OutValue)
+				{
+					using T = std::decay_t<decltype(OutValue)>;
+					if constexpr (std::is_same_v<T, FOutValuesAsArray>)
+					{
+						ApplyOnLoop([this, &OutValue](const uint8* SrcPtr, int32 Index)
+						{
+							FScriptArrayHelper OtherHelper(static_cast<FArrayProperty*>(UnderlyingProperty->Property->Inner), SrcPtr);
+							OutValue.OutValues[Index].template Get<0>() = OtherHelper.GetRawPtr();
+							OutValue.OutValues[Index].template Get<1>() = OtherHelper.Num();
+						});
+					}
+					else if constexpr (std::is_same_v<T, FOutValuesAsSet>)
+					{
+						ApplyOnLoop([this, &OutValue](const uint8* SrcPtr, int32 Index)
+						{
+							new (OutValue.OutValues[Index]) FScriptSetHelper(static_cast<FSetProperty*>(UnderlyingProperty->Property->Inner), SrcPtr);
+						});
+					}
+					else if constexpr (std::is_same_v<T, FOutValuesAsMap>)
+					{
+						ApplyOnLoop([this, &OutValue](const uint8* SrcPtr, int32 Index)
+						{
+							new (OutValue.OutValues[Index]) FScriptMapHelper(static_cast<FMapProperty*>(UnderlyingProperty->Property->Inner), SrcPtr);
+						});
+					}
+					else if constexpr (std::is_same_v<T, FOutValuesByPtr>)
+					{
+						ApplyOnLoop([this, &OutValue](const uint8* SrcPtr, int32 Index)
+						{
+							OutValue.OutValues[Index] = SrcPtr;
+						});
+					}
+					else if constexpr (std::is_same_v<T, FOutValuesByValue>)
+					{
+						void* DestPtr = const_cast<void*>(UnderlyingProperty->GetPtrInArray(OutValue.OutValues, Range.StartIndex));
+						UnderlyingProperty->Copy(DestPtr, SrcPtr, Range.Count);
+					}
+					else
+					{
+						static_assert(!std::is_same_v<T, T>, "Missing variant case for GetValues_Internal");
+					}
+			}, OutValues);
+		};
+
+		for (; It; ++It)
+		{
+			const int32 Index = It.GetIndex();
+			const PCGMetadataValueKey ValueKey = ValueKeys[Index];
+
+			int32 ValueKeyIndex = INDEX_NONE;
+
+			if (ValueKey == PCGDefaultValueKey)
+			{
+				ValueKeyIndex = PCG::Private::ExistingIndexForDefaultValue;
+			}
+			else if (ValueKey >= ValueKeyOffset)
+			{
+				int32 ValueIndex = ValueKey - ValueKeyOffset;
+				ValueKeyIndex = ValueIndex < Values.Num() ? ValueIndex : PCG::Private::ExistingIndexForDefaultValue;
+			}
+			else if (!ThisParent)
+			{
+				ValueKeyIndex = PCG::Private::ExistingIndexForDefaultValue;
+			}
+			else
+			{
+				bFoundAllKeys = false;
+			}
+
+			if (ValueKeyIndex < 0)
+			{
+				// Reached the end
+				if (ContinuousRange.StartIndex != INDEX_NONE)
+				{
+					RetrieveValues(ContinuousRange);
+					ContinuousRange.Reset();
+				}
+
+				// Then copy the default value if the value key is PCG::Private::ExistingIndexForDefaultValue
+				if (ValueKeyIndex == PCG::Private::ExistingIndexForDefaultValue)
+				{
+					ContinuousRange.Init(Index, PCG::Private::ExistingIndexForDefaultValue);
+					RetrieveValues(ContinuousRange);
+					ContinuousRange.Reset();
+				}
+			}
+			else
+			{
+				// Check if we are still in the same continuous range. If so, increment and continue. We also need to mark the value as retrieved.
+				if (ContinuousRange.StartIndex != INDEX_NONE && ContinuousRange.StartIndex + ContinuousRange.Count == Index && ContinuousRange.StartValueKey + ContinuousRange.Count == ValueKeyIndex)
+				{
+					ContinuousRange.Count++;
+					UnretrievedValues[Index] = false;
+					continue;
+				}
+				else if (ContinuousRange.StartIndex != INDEX_NONE)
+				{
+					// End of the range
+					RetrieveValues(ContinuousRange);
+					ContinuousRange.Reset();
+				}
+
+				ContinuousRange.Init(Index, ValueKeyIndex);
+			}
+
+			// If the value key index is not none, the value was (or will be) retrieved.
+			if (ValueKeyIndex != INDEX_NONE)
+			{
+				UnretrievedValues[Index] = false;
+			}
+		}
+
+		// One last time
+		if (ContinuousRange.StartIndex != INDEX_NONE)
+		{
+			// End of the range
+			RetrieveValues(ContinuousRange);
+			ContinuousRange.Reset();
+		}
+	}
+
+	ensure(ThisParent || bFoundAllKeys);
+
+	if (ThisParent && !bFoundAllKeys)
+	{
+		ThisParent->GetValues_Internal(ValueKeys, OutValues, UnretrievedValues);
+	}
+}
+
+FInstancedPropertyBag FPCGMetadataAttributeBase::BuildStructForDebug(PCGMetadataEntryKey EntryKey) const
+{
+	if (!ensure(UnderlyingProperty))
+	{
+		return {};
+	}
+
+	FInstancedPropertyBag Result;
+
+	FPropertyBagPropertyDesc Desc{Name, UnderlyingProperty->Property->Inner};
+
+	const UPropertyBag* PropertyBag = UPropertyBag::GetOrCreateFromDescs(MakeConstArrayView(&Desc, 1));
+	Result.InitializeFromBagStruct(PropertyBag);
+
+	const FProperty* ResultProperty = PropertyBag->GetPropertyDescs()[0].CachedProperty;
+
+	void* DestPtr = ResultProperty->ContainerPtrToValuePtr<void>(Result.GetMutableValue().GetMemory());
+
+	{
+		PCG::TSharedScopeLock ReadLock(ValueLock);
+		if (const void* SrcPtr = GetReadAddressFromEntryKey_Unsafe(EntryKey))
+		{
+			ResultProperty->CopyCompleteValue(DestPtr, SrcPtr);
+		}
+	}
+
+	return Result;
+}
+
+TOptional<FScriptArrayHelper> FPCGMetadataAttributeBase::MakeArrayHelper() const
+{
+	if (!ensure(UnderlyingProperty) || !UnderlyingProperty->Property)
+	{
+		return {};
+	}
+
+	return FScriptArrayHelper(UnderlyingProperty->Property.Get(), &Values);
+}
+
+const void* FPCGMetadataAttributeBase::GetReadAddressForValueKey_Unsafe(const PCGMetadataValueKey InValueKey) const
+{
+	if (InValueKey == PCGDefaultValueKey)
+	{
+		check(DefaultValue.IsValid());
+		return DefaultValue.GetRawPtr();
+	}
+
+	if (InValueKey < ValueKeyOffset)
+	{
+		return ensure(Parent) ? Parent->GetReadAddressForValueKey_Unsafe(InValueKey) : nullptr;
+	}
+
+	const int32 ValueIndex = InValueKey - ValueKeyOffset;
+
+	TOptional<FScriptArrayHelper> Helper = MakeArrayHelper();
+	if (!Helper || ValueIndex >= Values.Num())
+	{
+		return nullptr;
+	}
+
+	return Helper->GetRawPtr(ValueIndex);
+}
+
+const void* FPCGMetadataAttributeBase::GetReadAddressFromEntryKey_Unsafe(const PCGMetadataEntryKey InEntryKey) const
+{
+	return GetReadAddressForValueKey_Unsafe(GetValueKey(InEntryKey));
+}
+
+void FPCGMetadataAttributeBase::SetValueFromProperty(PCGMetadataEntryKey ItemKey, const void* SrcPtr, const FProperty* Property)
+{
+	using namespace PCG::Private;
+
+	if (!SrcPtr || !Property)
+	{
+		return;
+	}
+
+	if (TOptional<FScriptArrayHelper> Helper = MakeArrayHelper(); Helper.IsSet())
+	{
+		if (!UnderlyingProperty->Property->Inner->SameType(Property))
+		{
+			return;
+		}
+
+		const int32 ExistingIndex = ItemKey == PCGInvalidEntryKey && ensure(DefaultValue.IsValid()) ? PCG::Private::ExistingIndexForDefaultValue : INDEX_NONE;
+
+		PCGMetadataValueKey ValueKeyToAdd = PCGNotFoundValueKey;
+
+		if (UnderlyingProperty->bCompressData && ItemKey != PCGInvalidEntryKey)
+		{
+			const bool bValueKeyFound = FindValues(MakeArrayView(&ValueKeyToAdd, 1), FInValues{TInPlaceType<FInValuesByValue>{}, SrcPtr, 1}, 1);
+			check(!bValueKeyFound || ValueKeyToAdd != PCGNotFoundValueKey);
+		}
+
+		if (ValueKeyToAdd == PCGNotFoundValueKey)
+		{
+			SetValuesFromValueKeys({ItemKey}, AddValues_Internal(PCG::Private::FInValues{TInPlaceType<PCG::Private::FInValuesByValue>{}, SrcPtr, 1}, 1, ExistingIndex), /*bResetDefaultValueOnDefaultValueKey=*/ItemKey==PCGInvalidEntryKey);
+		}
+		else
+		{
+			SetValueFromValueKey(ItemKey, ValueKeyToAdd);
+		}
+	}
+}
+
+int32 FPCGMetadataAttributeBase::PreallocateValues(TArrayView<PCGMetadataEntryKey*> EntryKeys, bool bLockless)
+{
+	int32 StartIndex = INDEX_NONE;
+
+	if (!ensure(UnderlyingProperty))
+	{
+		return StartIndex;
+	}
+
+	if (!UnderlyingProperty->bCompressData)
+	{
+		PCG::TUniqueScopeLock WriteValueLock(ValueLock, /*bShouldLock=*/!bLockless);
+		StartIndex = AllocateValues_Unsafe(EntryKeys.Num());
+	}
+
+	{
+		PCG::TUniqueScopeLock WriteEntryMapLock(EntryMapLock, /*bShouldLock=*/!bLockless);
+		EntryToValueKeyMap.Reserve(EntryToValueKeyMap.Num() + EntryKeys.Num());
+
+		if (!UnderlyingProperty->bCompressData && StartIndex != INDEX_NONE)
+		{
+			for (int32 i = 0; i < EntryKeys.Num(); ++i)
+			{
+				const PCGMetadataValueKey ValueKey = StartIndex + i + ValueKeyOffset;
+				EntryToValueKeyMap.Emplace(*EntryKeys[i], ValueKey);
+			}
+		}
+	}
+
+	return StartIndex;
+}
+
+void FPCGMetadataAttributeBase::Prepare(int32 Count)
+{
+	PCG::TUniqueScopeLock Lock(EntryMapLock);
+	EntryToValueKeyMap.Reserve(EntryToValueKeyMap.Num() + Count);
+}
+
+int32 FPCGMetadataAttributeBase::AllocateValues_Unsafe(int32 Count)
+{
+	if (TOptional<FScriptArrayHelper> Helper = MakeArrayHelper(); Helper.IsSet())
+	{
+		const int32 StartIndex = Helper->Num();
+
+		if (UnderlyingProperty->bIsPlainOldData)
+		{
+			Helper->AddUninitializedValues(Count);
+		}
+		else
+		{
+			Helper->AddValues(Count);
+		}
+
+		return StartIndex;
+	}
+	else
+	{
+		return INDEX_NONE;
+	}
+}
+
+const UPCGMetadata* FPCGMetadataAttributeBase::GetMetadata() const
+{
+	return Metadata ? Metadata->GetTopMetadata() : nullptr;
 }
 
 void FPCGMetadataAttributeBase::SetValueFromValueKey(PCGMetadataEntryKey EntryKey, PCGMetadataValueKey ValueKey, bool bResetValueOnDefaultValueKey)
diff --git a/Engine/Plugins/PCG/Source/PCG/Private/Metadata/PCGMetadataAttributeGeneric.cpp b/Engine/Plugins/PCG/Source/PCG/Private/Metadata/PCGMetadataAttributeGeneric.cpp
deleted file mode 100644
diff --git a/Engine/Plugins/PCG/Source/PCG/Private/Metadata/PCGMetadataDomain.cpp b/Engine/Plugins/PCG/Source/PCG/Private/Metadata/PCGMetadataDomain.cpp
--- a/Engine/Plugins/PCG/Source/PCG/Private/Metadata/PCGMetadataDomain.cpp
+++ b/Engine/Plugins/PCG/Source/PCG/Private/Metadata/PCGMetadataDomain.cpp
@@ -61,6 +61,7 @@ namespace PCGMetadata
 				return false;
 			}
 
+			// @todo_pcg To be modified to use accessors.
 			// Allow to set the value if both type matches or if we can construct AttributeType from PropertyType.
 			return PCGMetadataAttribute::CallbackWithRightType(BaseAttribute->GetTypeId(), [&EntryKey, &PropertyValue, BaseAttribute, Metadata]<typename AttributeType>(AttributeType AttributeValue) -> bool
 				{
@@ -148,6 +149,11 @@ namespace PCGMetadata
 void FPCGMetadataDomain::Serialize(FArchive& InArchive)
 {
 	LLM_SCOPE_BYTAG(PCG);
+	
+	InArchive.UsingCustomVersion(FFortniteMainBranchObjectVersion::GUID);
+	
+	// To delta serialize against the default
+	static const FPCGMetadataAttributeDesc DefaultDesc{};
 
 	int32 NumAttributes = (InArchive.IsLoading() ? 0 : Attributes.Num());
 	// We need to keep track of the max attribute Id, since it won't necessary be equal to the number of attributes + 1.
@@ -161,16 +167,24 @@ void FPCGMetadataDomain::Serialize(FArchive& InArchive)
 		{
 			FName AttributeName = NAME_None;
 			InArchive << AttributeName;
-
-			int32 AttributeTypeId = 0;
-			InArchive << AttributeTypeId;
-
-			FPCGMetadataAttributeBase* SerializedAttribute = PCGMetadataAttribute::AllocateEmptyAttributeFromType(static_cast<int16>(AttributeTypeId));
-			if (!SerializedAttribute)
+			
+			FPCGMetadataAttributeBase* SerializedAttribute = nullptr;
+			FPCGMetadataAttributeDesc AttributeDesc;
+			AttributeDesc.Name = AttributeName;
+		
+			if (InArchive.CustomVer(FFortniteMainBranchObjectVersion::GUID) < FFortniteMainBranchObjectVersion::MergePCGMetadataAttributeBaseAndGeneric)
 			{
-				SerializedAttribute = new FPCGMetadataAttributeGeneric();
+				int32 AttributeTypeId = 0;
+				InArchive << AttributeTypeId;
+				
+				AttributeDesc.ValueType = static_cast<EPCGMetadataTypes>(AttributeTypeId);
 			}
-
+			else
+			{
+				FPCGMetadataAttributeDesc::StaticStruct()->SerializeItem(InArchive, &AttributeDesc, /*Defaults=*/&DefaultDesc);
+			}
+			
+			SerializedAttribute = AllocateAttributeFromDesc(AttributeDesc);
 			check(SerializedAttribute);
 
 			SerializedAttribute->Name = AttributeName;
@@ -181,14 +195,13 @@ void FPCGMetadataDomain::Serialize(FArchive& InArchive)
 	}
 	else
 	{
-		for (auto& AttributePair : Attributes)
+		for (auto& [Name, Attribute] : Attributes)
 		{
-			InArchive << AttributePair.Key;
+			InArchive << Name;
 			
-			int32 AttributeTypeId = AttributePair.Value->GetTypeId();
-			InArchive << AttributeTypeId;
+			FPCGMetadataAttributeDesc::StaticStruct()->SerializeItem(InArchive, &Attribute->CachedDesc, /*Defaults=*/&DefaultDesc);
 
-			AttributePair.Value->Serialize(this, InArchive);
+			Attribute->Serialize(this, InArchive);
 		}
 	}
 
@@ -226,80 +239,6 @@ void FPCGMetadataDomain::GetResourceSizeEx(FResourceSizeEx& CumulativeResourceSi
 	}
 }
 
-FPCGMetadataAttributeGeneric* FPCGMetadataDomain::CreateGenericAttribute(const FPCGMetadataAttributeDesc& Desc, bool bAllowsInterpolation, bool bOverrideParent)
-{
-	const FName AttributeName = Desc.Name;
-	
-	if (!FPCGMetadataAttributeBase::IsValidName(AttributeName))
-	{
-		UE_LOG(LogPCG, Error, TEXT("Attribute name '%s' is invalid"), *AttributeName.ToString());
-		return nullptr;
-	}
-
-	const FPCGMetadataAttributeGeneric* ParentAttribute = nullptr;
-
-	if (bOverrideParent && Parent)
-	{
-		const FPCGMetadataAttributeBase* BaseParentAttribute = Parent->GetConstAttribute(AttributeName);
-		ParentAttribute = BaseParentAttribute && BaseParentAttribute->IsGeneric() ? static_cast<const FPCGMetadataAttributeGeneric*>(BaseParentAttribute) : nullptr;
-	}
-
-	FPCGMetadataAttributeGeneric* NewAttribute = nullptr;
-
-	if (!ParentAttribute || (!ParentAttribute->GetAttributeDesc().IsSameType(Desc)))
-	{
-		// Can't parent if the types do not match
-		NewAttribute = new FPCGMetadataAttributeGeneric(Desc, this, bAllowsInterpolation);
-	}
-	else
-	{
-		NewAttribute = new FPCGMetadataAttributeGeneric(static_cast<const FPCGMetadataAttributeGeneric*>(ParentAttribute), this, bAllowsInterpolation, AttributeName);
-	}
-	
-	{
-		PCG::TUniqueScopeLock WriteLock(AttributeLock);
-
-		if (FPCGMetadataAttributeBase** BaseExistingAttribute = Attributes.Find(AttributeName))
-		{
-			FPCGMetadataAttributeGeneric* ExistingAttribute = *BaseExistingAttribute && (*BaseExistingAttribute)->IsGeneric() ? static_cast<FPCGMetadataAttributeGeneric*>(*BaseExistingAttribute) : nullptr;
-			
-			delete NewAttribute;
-			if (!ExistingAttribute || !ExistingAttribute->GetAttributeDesc().IsSameType(Desc))
-			{
-				UE_LOG(LogPCG, Error, TEXT("Attribute %s already exists but is not the right type. Abort."), *AttributeName.ToString());
-				return nullptr;
-			}
-			else
-			{
-				UE_LOG(LogPCG, Warning, TEXT("Attribute %s already exists"), *AttributeName.ToString());
-				NewAttribute = ExistingAttribute;
-			}
-		}
-		else
-		{
-			NewAttribute->AttributeId = NextAttributeId++;
-			AddAttributeInternal(AttributeName, NewAttribute);
-
-			// Also when creating an attribute, notify the PCG Data owner that the latest attribute manipulated is this one.
-			SetLastCachedSelectorOnOwner(AttributeName);
-		}
-	}
-
-	return NewAttribute;
-}
-
-FPCGMetadataAttributeGeneric* FPCGMetadataDomain::GetMutableGenericAttribute(FName AttributeName)
-{
-	FPCGMetadataAttributeBase* Attribute = GetMutableAttribute(AttributeName);
-	return Attribute && Attribute->IsGeneric() ? static_cast<FPCGMetadataAttributeGeneric*>(Attribute) : nullptr;
-}
-
-const FPCGMetadataAttributeGeneric* FPCGMetadataDomain::GetConstGenericAttribute(FName AttributeName) const
-{
-	const FPCGMetadataAttributeBase* Attribute = GetConstAttribute(AttributeName);
-	return Attribute && Attribute->IsGeneric() ? static_cast<const FPCGMetadataAttributeGeneric*>(Attribute) : nullptr;
-}
-
 void FPCGMetadataDomain::AddToCrc(FArchiveCrc32& Ar, const UPCGData* Data, bool bFullDataCrc) const
 {
 	check(Data);
@@ -538,12 +477,15 @@ bool FPCGMetadataDomain::AddAttributes(const FPCGMetadataDomainInitializeParams&
 			// If none of this is true - do nothing
 			const FPCGMetadataAttributeBase* Attribute = GetConstAttribute(OtherAttribute.Key);
 			check(Attribute);
+			
+			const FPCGMetadataAttributeDesc& AttributeDesc =  Attribute->GetAttributeDesc();
+			const FPCGMetadataAttributeDesc& OtherAttributeDesc =  OtherAttribute.Value->GetAttributeDesc();
 
-			if (Attribute->GetTypeId() != OtherAttribute.Value->GetTypeId() &&
-				!PCG::Private::IsBroadcastable(OtherAttribute.Value->GetTypeId(), Attribute->GetTypeId()) &&
-				PCG::Private::IsBroadcastable(Attribute->GetTypeId(), OtherAttribute.Value->GetTypeId()))
+			if (!AttributeDesc.IsSameType(OtherAttributeDesc) &&
+				!PCG::Private::IsBroadcastable(OtherAttributeDesc, AttributeDesc) &&
+				PCG::Private::IsBroadcastable(AttributeDesc, OtherAttributeDesc))
 			{
-				ChangeAttributeType(OtherAttribute.Key, OtherAttribute.Value->GetTypeId());
+				ChangeAttributeType(OtherAttribute.Key, OtherAttributeDesc);
 			}
 		}
 		else if (CopyAttribute(OtherAttribute.Value, OtherAttribute.Key, /*bKeepParent=*/InParams.Parent == Parent, /*bCopyEntries=*/false, /*bCopyValues=*/false))
@@ -739,19 +681,50 @@ void FPCGMetadataDomain::SetLastCachedSelectorOnOwner(FName AttributeName)
 	TopMetadata->SetLastCachedSelectorOnOwner(AttributeName, DomainID);
 }
 
+FPCGMetadataAttributeBase* FPCGMetadataDomain::AllocateAttributeFromDesc(const FPCGMetadataAttributeDesc& InDesc, const FPCGMetadataAttributeBase* InAttributeParent)
+{
+	// Single value and "old" types will be allocated as FPCGMetadataAttribute<T>, so we can still do `static_cast<FPCGMetadataAttribute<T>>(Attribute)`
+	FPCGMetadataAttributeBase* Result = nullptr;
+	if (InDesc.IsSingleValue())
+	{
+		Result = PCGMetadataAttribute::CallbackWithRightType(static_cast<uint16>(InDesc.ValueType), [this, &InDesc, InAttributeParent]<typename T>(T&&) -> FPCGMetadataAttributeBase*
+		{
+			// Making sure we do not change those types.
+			static_assert(PCG::Private::TIsBasicType<T>::Value);
+			return new FPCGMetadataAttribute<T>(typename FPCGMetadataAttribute<T>::FProtectedToken(), this, InDesc.Name, InAttributeParent, T{}, /*bInAllowsInterpolation=*/false);
+		});
+	}
+	
+	if (!Result)
+	{
+		Result = new FPCGMetadataAttributeBase(InDesc, InAttributeParent, this, /*bInAllowsInterpolation=*/false);
+	}
+	
+	if (Result && !Result->UnderlyingProperty)
+	{
+		// The attribute is invalid
+		delete Result;
+		Result = nullptr;
+	}
+	
+	return Result;
+}
+
 FPCGMetadataAttributeBase* FPCGMetadataDomain::GetMutableAttribute(FName AttributeName)
+{
+	PCG::TSharedScopeLock ReadLock(AttributeLock);
+	return GetMutableAttribute_Unsafe(AttributeName);
+}
+	
+FPCGMetadataAttributeBase* FPCGMetadataDomain::GetMutableAttribute_Unsafe(FName AttributeName)
 {
 	FPCGMetadataAttributeBase* Attribute = nullptr;
-
+	if (FPCGMetadataAttributeBase** FoundAttribute = Attributes.Find(AttributeName))
 	{
-		PCG::TSharedScopeLock ReadLock(AttributeLock);
-		if (FPCGMetadataAttributeBase** FoundAttribute = Attributes.Find(AttributeName))
-		{
-			Attribute = *FoundAttribute;
+		Attribute = *FoundAttribute;
 
-			// Also when accessing an attribute, notify the PCG Data owner that the latest attribute manipulated is this one.
-			SetLastCachedSelectorOnOwner(AttributeName);
-		}
+		// Also when accessing an attribute, notify the PCG Data owner that the latest attribute manipulated is this one.
+		SetLastCachedSelectorOnOwner(AttributeName);
 	}
 
 	return Attribute;
@@ -759,16 +732,18 @@ FPCGMetadataAttributeBase* FPCGMetadataDomain::GetMutableAttribute(FName Attribu
 
 const FPCGMetadataAttributeBase* FPCGMetadataDomain::GetConstAttribute(FName AttributeName) const
 {
-	const FPCGMetadataAttributeBase* Attribute = nullptr;
+	PCG::TSharedScopeLock ReadLock(AttributeLock);
+	return GetConstAttribute_Unsafe(AttributeName);
+}
 
+const FPCGMetadataAttributeBase* FPCGMetadataDomain::GetConstAttribute_Unsafe(FName AttributeName) const
+{
+	const FPCGMetadataAttributeBase* Attribute = nullptr;
+	if (const FPCGMetadataAttributeBase* const* FoundAttribute = Attributes.Find(AttributeName))
 	{
-		PCG::TSharedScopeLock ReadLock(AttributeLock);
-		if (const FPCGMetadataAttributeBase* const* FoundAttribute = Attributes.Find(AttributeName))
-		{
-			Attribute = *FoundAttribute;
-		}
+		Attribute = *FoundAttribute;
 	}
-
+	
 	return Attribute;
 }
 
@@ -837,6 +812,7 @@ void FPCGMetadataDomain::GetAttributes(TArray<FName>& AttributeNames, TArray<EPC
 		check(Attribute.Value && Attribute.Value->Name == Attribute.Key);
 		AttributeNames.Add(Attribute.Key);
 
+		// @todo_pcg: To be revisited with descriptors.
 		if (Attribute.Value->GetTypeId() < static_cast<uint16>(EPCGMetadataTypes::Unknown))
 		{
 			AttributeTypes.Add(static_cast<EPCGMetadataTypes>(Attribute.Value->GetTypeId()));
@@ -1021,6 +997,15 @@ void FPCGMetadataDomain::DeleteAttribute(FName AttributeToDelete)
 }
 
 bool FPCGMetadataDomain::ChangeAttributeType(FName AttributeName, int16 AttributeNewType)
+{
+	FPCGMetadataAttributeDesc NewDesc{};
+	NewDesc.Name = AttributeName;
+	NewDesc.ValueType = static_cast<EPCGMetadataTypes>(AttributeNewType);
+	
+	return ChangeAttributeType(AttributeName, NewDesc);
+}
+
+bool FPCGMetadataDomain::ChangeAttributeType(FName AttributeName, const FPCGMetadataAttributeDesc& AttributeNewDesc)
 {
 	FPCGMetadataAttributeBase* Attribute = GetMutableAttribute(AttributeName);
 
@@ -1030,13 +1015,19 @@ bool FPCGMetadataDomain::ChangeAttributeType(FName AttributeName, int16 Attribut
 		return false;
 	}
 
-	if (Attribute->GetTypeId() == AttributeNewType)
+	if (Attribute->GetAttributeDesc().IsSameType(AttributeNewDesc))
 	{
 		// Nothing to do, attribute is already the type we want
 		return true;
 	}
+	
+	if (!Attribute->GetAttributeDesc().IsSingleValue() || !AttributeNewDesc.IsSingleValue())
+	{
+		// Not possible at the moment to transform not single value attributes.
+		return false;
+	}
 
-	if (FPCGMetadataAttributeBase* NewAttribute = Attribute->CopyToAnotherType(AttributeNewType))
+	if (FPCGMetadataAttributeBase* NewAttribute = Attribute->CopyToAnotherType(static_cast<int16>(AttributeNewDesc.ValueType)))
 	{
 		NewAttribute->AttributeId = Attribute->AttributeId;
 
@@ -1340,7 +1331,7 @@ void FPCGMetadataDomain::MergeAttributesSubset(PCGMetadataEntryKey InKeyA, const
 			AttributeA = InMetadataA->GetConstAttribute(AttributeName);
 		}
 
-		if (AttributeA && AttributeA->GetTypeId() != Attribute->GetTypeId())
+		if (AttributeA && !AttributeA->GetAttributeDesc().IsSameType(Attribute->GetAttributeDesc()))
 		{
 			UE_LOG(LogPCG, Error, TEXT("Metadata type mismatch with attribute %s"), *AttributeName.ToString());
 			AttributeA = nullptr;
@@ -1353,7 +1344,7 @@ void FPCGMetadataDomain::MergeAttributesSubset(PCGMetadataEntryKey InKeyA, const
 			AttributeB = InMetadataB->GetConstAttribute(AttributeName);
 		}
 
-		if (AttributeB && AttributeB->GetTypeId() != Attribute->GetTypeId())
+		if (AttributeB && !AttributeB->GetAttributeDesc().IsSameType(Attribute->GetAttributeDesc()))
 		{
 			UE_LOG(LogPCG, Error, TEXT("Metadata type mismatch with attribute %s"), *AttributeName.ToString());
 			AttributeB = nullptr;
@@ -1401,7 +1392,7 @@ void FPCGMetadataDomain::AccumulateWeightedAttributes(PCGMetadataEntryKey InKey,
 
 		if (const FPCGMetadataAttributeBase* OtherAttribute = InMetadata->GetConstAttribute(AttributeName))
 		{
-			if (OtherAttribute->GetTypeId() != Attribute->GetTypeId())
+			if (!OtherAttribute->GetAttributeDesc().IsSameType(Attribute->GetAttributeDesc()))
 			{
 				UE_LOG(LogPCG, Error, TEXT("Metadata type mismatch with attribute %s"), *AttributeName.ToString());
 				continue;
@@ -1440,7 +1431,7 @@ void FPCGMetadataDomain::ComputeWeightedAttribute(PCGMetadataEntryKey& OutKey, c
 
 		if (const FPCGMetadataAttributeBase* OtherAttribute = InMetadata->GetConstAttribute(AttributeName))
 		{
-			if (OtherAttribute->GetTypeId() != Attribute->GetTypeId())
+			if (!OtherAttribute->GetAttributeDesc().IsSameType(Attribute->GetAttributeDesc()))
 			{
 				UE_LOG(LogPCG, Error, TEXT("Metadata type mismatch with attribute %s"), *AttributeName.ToString());
 				continue;
@@ -1477,7 +1468,7 @@ void FPCGMetadataDomain::SetAttributes(PCGMetadataEntryKey InKey, const FPCGMeta
 
 		if (const FPCGMetadataAttributeBase* OtherAttribute = InMetadata->GetConstAttribute(AttributeName))
 		{
-			if (OtherAttribute->GetTypeId() != Attribute->GetTypeId())
+			if (!OtherAttribute->GetAttributeDesc().IsSameType(Attribute->GetAttributeDesc()))
 			{
 				UE_LOG(LogPCG, Error, TEXT("Metadata type mismatch with attribute %s"), *AttributeName.ToString());
 				continue;
@@ -1569,7 +1560,7 @@ void FPCGMetadataDomain::SetAttributes(const TArrayView<const PCGMetadataEntryKe
 
 				if (const FPCGMetadataAttributeBase* OtherAttribute = InMetadata->GetConstAttribute(AttributeName))
 				{
-					if (!PCG::Private::IsBroadcastableOrConstructible(OtherAttribute->GetTypeId(), Attribute->GetTypeId()))
+					if (!PCG::Private::IsBroadcastableOrConstructible(OtherAttribute->GetAttributeDesc(), Attribute->GetAttributeDesc()))
 					{
 						PCGE_LOG_C(Error, GraphAndLog, OptionalContext, FText::Format(NSLOCTEXT("PCGMetadata", "TypeMismatch", "Metadata type mismatch with attribute '{0}'"), FText::FromName(AttributeName)));
 						return;
@@ -1594,19 +1585,20 @@ void FPCGMetadataDomain::SetAttributes(const TArrayView<const PCGMetadataEntryKe
 							return;
 						}
 
+						// @todo_pcg: To be modified to work with all types.
 						auto GetAndSetValues = [Attribute, &OutKeys, &OtherAttributeAccessor, &OtherAttributeKeys](auto Dummy) -> bool
-							{
-								using Type = decltype(Dummy);
-
-								auto SetValues = [Attribute, &OutKeys](const TArrayView<Type>& View, const int32 Start, const int32 Range)
-									{
-										TArrayView<PCGMetadataEntryKey> Keys(OutKeys.GetData() + Start, Range);
-										static_cast<FPCGMetadataAttribute<Type>*>(Attribute)->SetValues(Keys, View);
-									};
+						{
+							using Type = decltype(Dummy);
 
-								return PCGMetadataElementCommon::ApplyOnAccessorRange<Type>(OtherAttributeKeys, *OtherAttributeAccessor, SetValues, EPCGAttributeAccessorFlags::AllowBroadcast | EPCGAttributeAccessorFlags::AllowConstructible);
-							};
+							auto SetValues = [Attribute, &OutKeys](const TArrayView<Type>& View, const int32 Start, const int32 Range)
+								{
+									TArrayView<PCGMetadataEntryKey> Keys(OutKeys.GetData() + Start, Range);
+									static_cast<FPCGMetadataAttribute<Type>*>(Attribute)->SetValues(Keys, View);
+								};
 
+							return PCGMetadataElementCommon::ApplyOnAccessorRange<Type>(OtherAttributeKeys, *OtherAttributeAccessor, SetValues, EPCGAttributeAccessorFlags::AllowBroadcast | EPCGAttributeAccessorFlags::AllowConstructible);
+						};
+						
 						PCGMetadataAttribute::CallbackWithRightType(Attribute->GetTypeId(), GetAndSetValues);
 					}
 				}
@@ -1630,3 +1622,105 @@ void FPCGMetadataDomain::SetAttributes(const TArrayView<const PCGMetadataEntryKe
 {
 	SetAttributes(InKeys, InMetadata, &OutKeys, OptionalContext);
 }
+
+FPCGMetadataAttributeBase* FPCGMetadataDomain::CreateAttribute(const FPCGMetadataAttributeDesc& AttributeDesc, bool bAllowsInterpolation,bool bOverrideParent)
+{
+	return CreateAttributeInternal(AttributeDesc, bAllowsInterpolation, bOverrideParent, [](FPCGMetadataAttributeBase*){});
+}
+
+FPCGMetadataAttributeBase* FPCGMetadataDomain::CreateAttributeInternal(const FPCGMetadataAttributeDesc& AttributeDesc, bool bAllowsInterpolation, bool bOverrideParent, TFunctionRef<void(FPCGMetadataAttributeBase*)> SetupDefaultValue, bool bAlreadyLocked)
+{
+	const FName AttributeName(AttributeDesc.Name);
+	if (!FPCGMetadataAttributeBase::IsValidName(AttributeName))
+	{
+		UE_LOG(LogPCG, Error, TEXT("Attribute name '%s' is invalid"), *AttributeName.ToString());
+		return nullptr;
+	}
+
+	const FPCGMetadataAttributeBase* ParentAttribute = nullptr;
+
+	if (bOverrideParent && Parent)
+	{
+		ParentAttribute = Parent->GetConstAttribute(AttributeName);
+	}
+	
+	if (ParentAttribute && (!ParentAttribute->GetAttributeDesc().IsSameType(AttributeDesc)))
+	{
+		// Can't parent if the types do not match
+		ParentAttribute = nullptr;
+	}
+
+	FPCGMetadataAttributeBase* NewAttribute = AllocateAttributeFromDesc(AttributeDesc, ParentAttribute);
+	if (!NewAttribute)
+	{
+		return nullptr;
+	}
+	
+	NewAttribute->bAllowsInterpolation = bAllowsInterpolation;
+	SetupDefaultValue(NewAttribute);
+
+	{
+		PCG::TUniqueScopeLock WriteLock(AttributeLock, !bAlreadyLocked);
+
+		if (FPCGMetadataAttributeBase** ExistingAttribute = Attributes.Find(AttributeName))
+		{
+			delete NewAttribute;
+			if (!(*ExistingAttribute)->GetAttributeDesc().IsSameType(AttributeDesc))
+			{
+				UE_LOG(LogPCG, Error, TEXT("Attribute %s already exists but is not the right type. Abort."), *AttributeName.ToString());
+				return nullptr;
+			}
+			else
+			{
+				UE_LOG(LogPCG, Warning, TEXT("Attribute %s already exists"), *AttributeName.ToString());
+				NewAttribute = *ExistingAttribute;
+			}
+		}
+		else
+		{
+			NewAttribute->AttributeId = NextAttributeId++;
+			AddAttributeInternal(AttributeName, NewAttribute);
+
+			// Also when creating an attribute, notify the PCG Data owner that the latest attribute manipulated is this one.
+			SetLastCachedSelectorOnOwner(AttributeName);
+		}
+	}
+
+	return NewAttribute;
+}
+
+FPCGMetadataAttributeBase* FPCGMetadataDomain::FindOrCreateAttributeInternal(const FPCGMetadataAttributeDesc& AttributeDesc, bool bAllowsInterpolation, bool bOverrideParent, bool bOverwriteIfTypeMismatch, TFunctionRef<void(FPCGMetadataAttributeBase*)> SetupDefaultValue)
+{
+	const FName AttributeName = AttributeDesc.Name;
+	
+	{
+		PCG::TSharedScopeLock ReadLock(AttributeLock);
+		if (FPCGMetadataAttributeBase* Attribute = GetMutableAttribute_Unsafe(AttributeName); Attribute && Attribute->GetAttributeDesc().IsSameType(AttributeDesc))
+		{
+			return Attribute;
+		}
+	}
+
+	PCG::TUniqueScopeLock WriteLock(AttributeLock);
+	if (FPCGMetadataAttributeBase* Attribute = GetMutableAttribute_Unsafe(AttributeName); Attribute && Attribute->GetAttributeDesc().IsSameType(AttributeDesc))
+	{
+		return Attribute;
+	}
+
+	// If an attribute with this name exists here, there is a type mismatch.
+	if (FPCGMetadataAttributeBase** FoundAttribute = Attributes.Find(AttributeName))
+	{
+		if (bOverwriteIfTypeMismatch)
+		{
+			delete *FoundAttribute;
+			RemoveAttributeInternal(AttributeName);
+		}
+		else
+		{
+			return nullptr;
+		}
+	}
+
+	// A new attribute will be created.
+	return CreateAttributeInternal(AttributeDesc, bAllowsInterpolation, bOverrideParent, SetupDefaultValue, /*bAlreadyLocked=*/true);
+}
diff --git a/Engine/Plugins/PCG/Source/PCG/Public/Helpers/PCGSettingsHelpers.h b/Engine/Plugins/PCG/Source/PCG/Public/Helpers/PCGSettingsHelpers.h
--- a/Engine/Plugins/PCG/Source/PCG/Public/Helpers/PCGSettingsHelpers.h
+++ b/Engine/Plugins/PCG/Source/PCG/Public/Helpers/PCGSettingsHelpers.h
@@ -55,9 +55,6 @@ public:
 
 namespace PCGSettingsHelpers
 {
-	/** Sets data from the params to a given property, matched on a name basis */
-	void SetValue(UPCGParamData* Params, UObject* Object, FProperty* Property);
-
 	/** Utility to call from before-node-update deprecation. A dedicated pin for params will be added when the pins are updated. Here we detect any params
 	*   connections to the In pin and disconnect them, and move the first params connection to a new params pin.
 	*/
diff --git a/Engine/Plugins/PCG/Source/PCG/Public/Metadata/PCGMetadata.h b/Engine/Plugins/PCG/Source/PCG/Public/Metadata/PCGMetadata.h
--- a/Engine/Plugins/PCG/Source/PCG/Public/Metadata/PCGMetadata.h
+++ b/Engine/Plugins/PCG/Source/PCG/Public/Metadata/PCGMetadata.h
@@ -254,13 +254,14 @@ public:
 	UFUNCTION(BlueprintCallable, Category = "PCG|Metadata")
 	UE_API int32 GetAttributeCount() const;
 
-	template <typename T>
+	template <typename T> requires (PCG::Private::TIsBasicType<T>::Value == true)
 	FPCGMetadataAttribute<T>* GetMutableTypedAttribute(FPCGAttributeIdentifier AttributeName);
 
-	template <typename T>
+	template <typename T> requires (PCG::Private::TIsBasicType<T>::Value == true)
+	UE_DEPRECATED(5.8, "Should not be used anymore.")
 	FPCGMetadataAttribute<T>* GetMutableTypedAttribute_Unsafe(FPCGAttributeIdentifier AttributeName);
 
-	template <typename T>
+	template <typename T> requires (PCG::Private::TIsBasicType<T>::Value == true)
 	const FPCGMetadataAttribute<T>* GetConstTypedAttribute(FPCGAttributeIdentifier AttributeName) const;
 
 	/** Get all the attributes names and type on the default domain. If you need all the attributes on all domains, use GetAllAttributes. */
@@ -297,8 +298,11 @@ public:
 	UE_API void BP_ClearAttribute(FName AttributeToClear);
 	
 	UE_API void ClearAttribute(FPCGAttributeIdentifier AttributeToClear);
-
+	
 	/** Change type of an attribute */
+	UE_API bool ChangeAttributeType(FPCGAttributeIdentifier AttributeName, const FPCGMetadataAttributeDesc& AttributeNewDesc);
+	
+	UE_DEPRECATED(5.8, "Use the version with the Attribute Desc")
 	UE_API bool ChangeAttributeType(FPCGAttributeIdentifier AttributeName, int16 AttributeNewType);
 
 	/** Adds a unique entry key to the metadata */
@@ -392,8 +396,24 @@ public:
 	* If the attribute already exists but is of the wrong type, it will fail and return nullptr. Same if the name is invalid.
 	* Return a typed attribute pointer, of the requested type T.
 	*/
-	template<typename T>
+	template<typename T> requires (PCG::Private::TIsBasicType<T>::Value == true)
 	FPCGMetadataAttribute<T>* CreateAttribute(FPCGAttributeIdentifier AttributeName, const T& DefaultValue, bool bAllowsInterpolation, bool bOverrideParent);
+	
+	/**
+	* Create a new attribute, that will be generic and not typed. This should never be static cast to a FPCGMetadataAttribute<T>* as this would
+	* be undefined behavior.
+	* If the attribute already exists, it will raise a warning (use FindOrCreateAttribute if this usecase can arise)
+	* If the attribute already exists but is of the wrong type, it will fail and return nullptr. Same if the name is invalid.
+	*/
+	template<typename T> requires (PCG::Private::TIsBasicType<T>::Value == false)
+	FPCGMetadataAttributeBase* CreateAttribute(FPCGAttributeIdentifier AttributeName, const T& DefaultValue, bool bAllowsInterpolation = true, bool bOverrideParent = true);
+	
+	/**
+	* Create a new attribute, that will be either generic or typed, depending the type given by the desc. The default value will not be set.
+	* If the attribute already exists, it will raise a warning (use FindOrCreateAttribute if this usecase can arise)
+	* If the attribute already exists but is of the wrong type, it will fail and return nullptr. Same if the name is invalid.
+	*/
+	UE_API FPCGMetadataAttributeBase* CreateAttribute(FPCGAttributeIdentifier AttributeName, const FPCGMetadataAttributeDesc& AttributeDesc, bool bAllowsInterpolation = true, bool bOverrideParent = true);
 
 	/**
 	* Find or create an attribute. Follows CreateAttribute signature.
@@ -401,8 +421,17 @@ public:
 	* Same as CreateAttribute, it will return nullptr if the attribute name is invalid.
 	* Return a typed attribute pointer, of the requested type T.
 	*/
-	template<typename T>
+	template<typename T> requires (PCG::Private::TIsBasicType<T>::Value == true)
 	FPCGMetadataAttribute<T>* FindOrCreateAttribute(FPCGAttributeIdentifier AttributeName, const T& DefaultValue = T{}, bool bAllowsInterpolation = true, bool bOverrideParent = true, bool bOverwriteIfTypeMismatch = true);
+	
+	/**
+	* Find or create an attribute. Follows CreateAttribute signature.
+	* Extra boolean bOverwriteIfTypeMismatch allows to overwrite an existing attribute if the type mismatch.
+	* Same as CreateAttribute, it will return nullptr if the attribute name is invalid.
+	* Return a generic attribute pointer
+	*/
+	template<typename T> requires (PCG::Private::TIsBasicType<T>::Value == false)
+	FPCGMetadataAttributeBase* FindOrCreateAttribute(FPCGAttributeIdentifier AttributeName, const T& DefaultValue = T{}, bool bAllowsInterpolation = true, bool bOverrideParent = true, bool bOverwriteIfTypeMismatch = true);
 
 	// Need this gymnastic for those 2 functions, because blueprints doesn't support default arguments for Arrays, and we don't want to force c++ user to provide OptionalNewEntriesOrder.
 
@@ -482,43 +511,55 @@ protected:
 	PCG::FLock MetadataDomainsLock;
 };
 
-template<typename T>
+template<typename T> requires (PCG::Private::TIsBasicType<T>::Value == true)
 FPCGMetadataAttribute<T>* UPCGMetadata::CreateAttribute(FPCGAttributeIdentifier AttributeName, const T& DefaultValue, bool bAllowsInterpolation, bool bOverrideParent)
 {
 	FPCGMetadataDomain* FoundSubMetadata = FindOrCreateMetadataDomain(AttributeName.MetadataDomain);
 	return FoundSubMetadata ? FoundSubMetadata->CreateAttribute<T>(AttributeName.Name, DefaultValue, bAllowsInterpolation, bOverrideParent) : nullptr;
 }
 
-template<typename T>
-FPCGMetadataAttribute<T>* UPCGMetadata::GetMutableTypedAttribute_Unsafe(FPCGAttributeIdentifier AttributeName)
+template<typename T> requires (PCG::Private::TIsBasicType<T>::Value == false)
+FPCGMetadataAttributeBase* UPCGMetadata::CreateAttribute(FPCGAttributeIdentifier AttributeName, const T& DefaultValue, bool bAllowsInterpolation, bool bOverrideParent)
 {
-	FPCGMetadataDomain* FoundSubMetadata = GetMetadataDomain(AttributeName.MetadataDomain);
-	return FoundSubMetadata ? FoundSubMetadata->GetMutableTypedAttribute_Unsafe<T>(AttributeName.Name) : nullptr;
+	FPCGMetadataDomain* FoundSubMetadata = FindOrCreateMetadataDomain(AttributeName.MetadataDomain);
+	return FoundSubMetadata ? FoundSubMetadata->CreateAttribute<T>(AttributeName.Name, DefaultValue, bAllowsInterpolation, bOverrideParent) : nullptr;
 }
 
-template<typename T>
+template<typename T> requires (PCG::Private::TIsBasicType<T>::Value == true)
 FPCGMetadataAttribute<T>* UPCGMetadata::FindOrCreateAttribute(FPCGAttributeIdentifier AttributeName, const T& DefaultValue, bool bAllowsInterpolation, bool bOverrideParent, bool bOverwriteIfTypeMismatch)
 {
 	FPCGMetadataDomain* FoundSubMetadata = FindOrCreateMetadataDomain(AttributeName.MetadataDomain);
 	return FoundSubMetadata ? FoundSubMetadata->FindOrCreateAttribute<T>(AttributeName.Name, DefaultValue, bAllowsInterpolation, bOverrideParent, bOverwriteIfTypeMismatch) : nullptr;
 }
 
-template <typename T>
+template<typename T> requires (PCG::Private::TIsBasicType<T>::Value == false)
+FPCGMetadataAttributeBase* UPCGMetadata::FindOrCreateAttribute(FPCGAttributeIdentifier AttributeName, const T& DefaultValue, bool bAllowsInterpolation, bool bOverrideParent, bool bOverwriteIfTypeMismatch)
+{
+	FPCGMetadataDomain* FoundSubMetadata = FindOrCreateMetadataDomain(AttributeName.MetadataDomain);
+	return FoundSubMetadata ? FoundSubMetadata->FindOrCreateAttribute<T>(AttributeName.Name, DefaultValue, bAllowsInterpolation, bOverrideParent, bOverwriteIfTypeMismatch) : nullptr;
+}
+
+template<typename T> requires (PCG::Private::TIsBasicType<T>::Value == true)
+FPCGMetadataAttribute<T>* UPCGMetadata::GetMutableTypedAttribute_Unsafe(FPCGAttributeIdentifier AttributeName)
+{
+	FPCGMetadataDomain* FoundSubMetadata = GetMetadataDomain(AttributeName.MetadataDomain);
+PRAGMA_DISABLE_DEPRECATION_WARNINGS
+	return FoundSubMetadata ? FoundSubMetadata->GetMutableTypedAttribute_Unsafe<T>(AttributeName.Name) : nullptr;
+PRAGMA_ENABLE_DEPRECATION_WARNINGS
+}
+
+template <typename T> requires (PCG::Private::TIsBasicType<T>::Value == true)
 FPCGMetadataAttribute<T>* UPCGMetadata::GetMutableTypedAttribute(FPCGAttributeIdentifier AttributeName)
 {
-	FPCGMetadataAttributeBase* BaseAttribute = GetMutableAttribute(AttributeName);
-	return (BaseAttribute && (BaseAttribute->GetTypeId() == PCG::Private::MetadataTypes<T>::Id))
-		? static_cast<FPCGMetadataAttribute<T>*>(BaseAttribute)
-		: nullptr;
+	FPCGMetadataDomain* FoundSubMetadata = GetMetadataDomain(AttributeName.MetadataDomain);
+	return FoundSubMetadata ? FoundSubMetadata->GetMutableTypedAttribute<T>(AttributeName.Name) : nullptr;
 }
 
-template <typename T>
+template <typename T> requires (PCG::Private::TIsBasicType<T>::Value == true)
 const FPCGMetadataAttribute<T>* UPCGMetadata::GetConstTypedAttribute(FPCGAttributeIdentifier AttributeName) const
 {
-	const FPCGMetadataAttributeBase* BaseAttribute = GetConstAttribute(AttributeName);
-	return (BaseAttribute && (BaseAttribute->GetTypeId() == PCG::Private::MetadataTypes<T>::Id))
-		? static_cast<const FPCGMetadataAttribute<T>*>(BaseAttribute)
-		: nullptr;
+	const FPCGMetadataDomain* FoundSubMetadata = GetConstMetadataDomain(AttributeName.MetadataDomain);
+	return FoundSubMetadata ? FoundSubMetadata->GetConstTypedAttribute<T>(AttributeName.Name) : nullptr;
 }
 
 #undef UE_API
diff --git a/Engine/Plugins/PCG/Source/PCG/Public/Metadata/PCGMetadataAttribute.h b/Engine/Plugins/PCG/Source/PCG/Public/Metadata/PCGMetadataAttribute.h
--- a/Engine/Plugins/PCG/Source/PCG/Public/Metadata/PCGMetadataAttribute.h
+++ b/Engine/Plugins/PCG/Source/PCG/Public/Metadata/PCGMetadataAttribute.h
@@ -4,9 +4,19 @@
 
 #include "PCGCommon.h"
 #include "PCGMetadataCommon.h"
+#include "PCGMetadataContainerTypes.h"
+#include "PCGModule.h"
+
+#include "Helpers/PCGConcepts.h"
+#include "Helpers/PCGMetadataHelpers.h"
+#include "Metadata/PCGMetadataAttributeTraits.h"
 
 #include "Utils/PCGValueRange.h"
 
+#include "Containers/StaticArray.h"
+#include "Misc/TVariant.h"
+#include "StructUtils/PropertyBag.h"
+
 #define UE_API PCG_API
 
 class FPCGMetadataDomain;
@@ -20,56 +30,114 @@ namespace PCGMetadataAttributeConstants
 	extern PCG_API const FName SourceNameAttributeName;
 }
 
+struct FPCGAttributeDefaultValue;
+
+namespace PCG::Private
+{
+	void CopyArray(const FArrayProperty* InProperty, void* DestPtr, const void* SrcPtr, int32 Count);
+
+	/** Magic number for AddValues to know if we need to write to the DefaultValue or not. */
+	constexpr int32 ExistingIndexForDefaultValue = -2;
+}
+
+struct FPCGAttributeProperty
+{
+public:
+	PCG_API FPCGAttributeProperty(const FPCGMetadataAttributeDesc& InAttributeDesc);
+	PCG_API ~FPCGAttributeProperty();
+
+	PCG_API void Copy(void* DestPtr, const void* SrcPtr, int32 Count = 1);
+	PCG_API bool IsValid() const;
+
+private:
+	friend FPCGMetadataAttributeBase;
+	friend FPCGAttributeDefaultValue;
+
+	/** Unsafe method that will apply pointer arithmetic to get to the element in the array at the specified index. */
+	const void* GetPtrInArray(const void* InArrayStart, int32 Index) const;
+
+	TUniquePtr<FArrayProperty> Property;
+	size_t InnerElementCachedSize = 0;
+	bool bCompressData = false;
+	bool bIsPlainOldData = false;
+};
+
+// Simple wrapper around a raw pointer to hold the default value to mimic how it was stored in the typed attributes
+struct FPCGAttributeDefaultValue
+{
+	FPCGAttributeDefaultValue() = default;
+	PCG_API ~FPCGAttributeDefaultValue();
+
+	// Since this is an internal struct and attributes are never copied with constructor/assignment, just delete the copy/move
+	UE_NONCOPYABLE(FPCGAttributeDefaultValue);
+
+	PCG_API void Init(TSharedPtr<FPCGAttributeProperty> InUnderlyingProperty);
+	PCG_API void CopyFrom(const void* SrcData);
+
+	void* GetRawPtr() { return Memory; }
+	const void* GetRawPtr() const { return Memory; };
+	bool IsValid() const { return Memory != nullptr; }
+
+private:
+	void Destroy();
+
+	TSharedPtr<FPCGAttributeProperty> UnderlyingProperty;
+	void* Memory = nullptr;
+};
+
 class FPCGMetadataAttributeBase
 {
 public:
-	FPCGMetadataAttributeBase() = default;
+	UE_DEPRECATED(5.8, "Metadata attributes should never be constructed individually and should be created from the Metadata domain.")
 	UE_API FPCGMetadataAttributeBase(FPCGMetadataDomain* InMetadata, FName InName, const FPCGMetadataAttributeBase* InParent, bool bInAllowsInterpolation);
 
-	UE_DEPRECATED(5.6, "Use the version with FPCGMetadataDomain")
-	UE_API FPCGMetadataAttributeBase(UPCGMetadata* InMetadata, FName InName, const FPCGMetadataAttributeBase* InParent, bool bInAllowsInterpolation);
+	UE_API virtual ~FPCGMetadataAttributeBase();
 
-	virtual ~FPCGMetadataAttributeBase() = default;
-	
 	UE_API virtual void Serialize(FPCGMetadataDomain* InMetadata, FArchive& InArchive);
 	UE_API virtual void GetResourceSizeEx(FResourceSizeEx& CumulativeResourceSize) const;
+	UE_API void AddReferencedObjects(FReferenceCollector& Collector);
+
+	const FPCGMetadataAttributeDesc& GetAttributeDesc() const { return CachedDesc; }
+	bool DoesCompressData() const { return UnderlyingProperty && UnderlyingProperty->bCompressData; }
 
-	UE_DEPRECATED(5.6, "Use the version with FPCGMetadataDomain")
-	UE_API virtual void Serialize(UPCGMetadata* InMetadata, FArchive& InArchive);
+	template <typename T>
+	bool IsOfType() const
+	{
+		return PCG::Private::IsEquivalentDesc<T>(CachedDesc);
+	}
+
+	// @todo_pcg All generic methods to be removed in a follow-up CL. Will ease the refactor review.
+	bool IsGeneric() const { return true; }
 
 	/** Unparents current attribute by flattening the values, entries, etc. */
-	virtual void Flatten() = 0;
+	UE_API virtual void Flatten();
 	/** Unparents current attribute by flattening the values, entries, etc while only keeping the entries referenced in InEntryKeysToKeep. There must be NO invalid entry keys. */
-	virtual void FlattenAndCompress(const TArrayView<const PCGMetadataEntryKey>& InEntryKeysToKeep) = 0;
+	UE_API virtual void FlattenAndCompress(const TArrayView<const PCGMetadataEntryKey>& InEntryKeysToKeep);
 
 	/** Remove all entries, values and parenting. */
-	virtual void Reset() = 0;
-	
+	UE_API virtual void Reset();
+
 	UE_API const UPCGMetadata* GetMetadata() const;
 	const FPCGMetadataDomain* GetMetadataDomain() const { return Metadata; }
-	
+
 	int16 GetTypeId() const { return TypeId; }
 
-	virtual FPCGMetadataAttributeBase* Copy(FName NewName, FPCGMetadataDomain* InMetadata, bool bKeepParent, bool bCopyEntries = true, bool bCopyValues = true) const = 0;
+	UE_API virtual FPCGMetadataAttributeBase* Copy(FName NewName, FPCGMetadataDomain* InMetadata, bool bKeepParent, bool bCopyEntries = true, bool bCopyValues = true) const;
+	UE_API virtual FPCGMetadataAttributeBase* CopyToAnotherType(int16 Type) const;
 
-	UE_DEPRECATED(5.6, "Use the version with FPCGMetadataDomain")
-	virtual FPCGMetadataAttributeBase* Copy(FName NewName, UPCGMetadata* InMetadata, bool bKeepParent, bool bCopyEntries = true, bool bCopyValues = true) const = 0;
-	
-	virtual FPCGMetadataAttributeBase* CopyToAnotherType(int16 Type) const = 0;
-
-	virtual PCGMetadataValueKey GetValueKeyOffsetForChild() const = 0;
-	virtual void SetValue(PCGMetadataEntryKey ItemKey, const FPCGMetadataAttributeBase* InAttribute, PCGMetadataEntryKey InEntryKey) = 0;
-	virtual void SetZeroValue(PCGMetadataEntryKey ItemKey) = 0;
-	virtual void AccumulateValue(PCGMetadataEntryKey ItemKey, const FPCGMetadataAttributeBase* InAttributeA, PCGMetadataEntryKey InEntryKeyA, float Weight) = 0;
-	virtual void SetWeightedValue(PCGMetadataEntryKey ItemKey, const FPCGMetadataAttributeBase* InAttribute, const TArrayView<const TPair<PCGMetadataEntryKey, float>>& InWeightedKeys) = 0;
-	virtual void SetValue(PCGMetadataEntryKey ItemKey, const FPCGMetadataAttributeBase* InAttributeA, PCGMetadataEntryKey InEntryKeyA, const FPCGMetadataAttributeBase* InAttributeB, PCGMetadataEntryKey InEntryKeyB, EPCGMetadataOp Op) = 0;
-	virtual bool IsEqualToDefaultValue(PCGMetadataValueKey ValueKey) const = 0;
+	UE_API virtual PCGMetadataValueKey GetValueKeyOffsetForChild() const;
+	UE_API virtual void SetValue(PCGMetadataEntryKey ItemKey, const FPCGMetadataAttributeBase* InAttribute, PCGMetadataEntryKey InEntryKey);
+	UE_API virtual void SetZeroValue(PCGMetadataEntryKey ItemKey);
+	UE_API virtual void AccumulateValue(PCGMetadataEntryKey ItemKey, const FPCGMetadataAttributeBase* InAttributeA, PCGMetadataEntryKey InEntryKeyA, float Weight);
+	UE_API virtual void SetWeightedValue(PCGMetadataEntryKey ItemKey, const FPCGMetadataAttributeBase* InAttribute, const TArrayView<const TPair<PCGMetadataEntryKey, float>>& InWeightedKeys);
+	UE_API virtual void SetValue(PCGMetadataEntryKey ItemKey, const FPCGMetadataAttributeBase* InAttributeA, PCGMetadataEntryKey InEntryKeyA, const FPCGMetadataAttributeBase* InAttributeB, PCGMetadataEntryKey InEntryKeyB, EPCGMetadataOp Op);
+	UE_API virtual bool IsEqualToDefaultValue(PCGMetadataValueKey ValueKey) const;
 	/** In the case of multi entry attribute and after some operations, we might have a single entry attribute with a default value that is different than the first entry. Use this function to fix that. Only valid if there is one and only one value. */
-	virtual void SetDefaultValueToFirstEntry() = 0;
+	UE_API virtual void SetDefaultValueToFirstEntry();
 
-	virtual bool UsesValueKeys() const = 0;
-	virtual bool AreValuesEqualForEntryKeys(PCGMetadataEntryKey EntryKey1, PCGMetadataEntryKey EntryKey2) const = 0;
-	virtual bool AreValuesEqual(PCGMetadataValueKey ValueKey1, PCGMetadataValueKey ValueKey2) const = 0;
+	UE_API virtual bool UsesValueKeys() const;
+	UE_API virtual bool AreValuesEqualForEntryKeys(PCGMetadataEntryKey EntryKey1, PCGMetadataEntryKey EntryKey2) const;
+	UE_API virtual bool AreValuesEqual(PCGMetadataValueKey ValueKey1, PCGMetadataValueKey ValueKey2) const;
 
 	UE_API void SetValueFromValueKey(PCGMetadataEntryKey EntryKey, PCGMetadataValueKey ValueKey, bool bResetValueOnDefaultValueKey = false);
 	UE_API PCGMetadataValueKey GetValueKey(PCGMetadataEntryKey EntryKey) const;
@@ -113,27 +181,528 @@ public:
 	/** Replaces any invalid characters in name with underscores. Returns true if Name was changed. */
 	static UE_API bool SanitizeName(FString& InOutName);
 
-	// @todo_pcg: To be removed
-	virtual bool IsGeneric() const { return false; }
+	/**
+	 * GetValues section
+	 * Templated code to support all container types.
+	 * If the provided output container does not match the underlying type, it will fail.
+	 * It is also mandatory that the number of value keys and output container num match.
+	 *
+	 * Supported types are:
+	 * - Single types that are valid for metadata
+	 * - Pointers to single types that are valid for metadata
+	 * - Typed ConstArrayViews for arrays
+	 * - PCG::TScriptSetWrapper<T> for set, wrapper that can be manipulated like a set
+	 * - PCG::TScriptMapWrapper<T> for map, wrapper that can be manipulated like a map
+	 */
+	template <typename T> requires (!std::is_pointer_v<T>)
+	void GetValues(TArrayView<const PCGMetadataValueKey> ValueKeys, TArrayView<T> OutValues) const
+	{
+		if (ValueKeys.IsEmpty())
+		{
+			return;
+		}
+
+		if (!IsOfType<T>())
+		{
+			// @todo_pcg Better logging
+			UE_LOG(LogPCG, Error, TEXT("Type mismatch on GetValues"));
+			return;
+		}
+
+		if (!ensure(ValueKeys.Num() == OutValues.Num()))
+		{
+			return;
+		}
+
+		GetValues_Internal(ValueKeys, PCG::Private::FOutValues{TInPlaceType<PCG::Private::FOutValuesByValue>{}, OutValues.GetData()});
+	}
+
+	template <typename T>
+	void GetValues(TArrayView<const PCGMetadataValueKey> ValueKeys, TArrayView<const T*> OutValues) const
+	{
+		if (ValueKeys.IsEmpty())
+		{
+			return;
+		}
+
+		if (!IsOfType<T>())
+		{
+			// @todo_pcg Better logging
+			UE_LOG(LogPCG, Error, TEXT("Type mismatch on GetValues"));
+			return;
+		}
+
+		if (!ensure(ValueKeys.Num() == OutValues.Num()))
+		{
+			return;
+		}
+
+		// We need to type erase the pointers, but to be fully C++ compliant about aliasing rules we can't reinterpret_cast. So make a copy of the pointers.
+		TArray<const void*, TInlineAllocator<256>> Temp;
+		Temp.SetNumUninitialized(OutValues.Num());
+		GetValues_Internal(ValueKeys, PCG::Private::FOutValues{ TInPlaceType<PCG::Private::FOutValuesByPtr>{}, Temp });
+		for (int32 i = 0; i < OutValues.Num(); ++i)
+		{
+			OutValues[i] = static_cast<const T*>(Temp[i]);
+		}
+	}
+
+	template <typename T>
+	void GetValues(TArrayView<const PCGMetadataValueKey> ValueKeys, TArrayView<TConstArrayView<T>> OutValues) const
+	{
+		if (ValueKeys.IsEmpty())
+		{
+			return;
+		}
+
+		if (!IsOfType<TArray<T>>())
+		{
+			// @todo_pcg Better logging
+			UE_LOG(LogPCG, Error, TEXT("Type mismatch on GetValues"));
+			return;
+		}
+
+		if (!ensure(ValueKeys.Num() == OutValues.Num()))
+		{
+			return;
+		}
+
+		TArray<TTuple<const void*, int32>, TInlineAllocator<256>> Temp;
+		Temp.SetNumUninitialized(OutValues.Num());
+		GetValues_Internal(ValueKeys, PCG::Private::FOutValues{TInPlaceType<PCG::Private::FOutValuesAsArray>{}, Temp});
+		for (int32 i = 0; i < OutValues.Num(); ++i)
+		{
+			OutValues[i] = MakeConstArrayView<T>(static_cast<const T*>(Temp[i].Get<0>()), Temp[i].Get<1>());
+		}
+	}
+
+	// To support sets
+	template <typename T>
+	void GetValues(TArrayView<const PCGMetadataValueKey> ValueKeys, TArrayView<PCG::TScriptSetWrapper<T>> OutValues) const
+	{
+		if (ValueKeys.IsEmpty())
+		{
+			return;
+		}
+
+		if (!IsOfType<TSet<T>>())
+		{
+			// @todo_pcg Better logging
+			UE_LOG(LogPCG, Error, TEXT("Type mismatch on GetValues"));
+			return;
+		}
+
+		if (!ensure(ValueKeys.Num() == OutValues.Num()))
+		{
+			return;
+		}
+
+		TArray<FScriptSetHelper*, TInlineAllocator<256>> Temp;
+		Temp.Reserve(OutValues.Num());
+		for (PCG::TScriptSetWrapper<T>& Value : OutValues)
+		{
+			Temp.Add(&Value.Helper);
+		}
+		GetValues_Internal(ValueKeys, PCG::Private::FOutValues{TInPlaceType<PCG::Private::FOutValuesAsSet>{}, Temp});
+	}
+
+	// To support maps
+	template <typename KeyType, typename ValueType>
+	void GetValues(TArrayView<const PCGMetadataValueKey> ValueKeys, TArrayView<PCG::TScriptMapWrapper<KeyType, ValueType>> OutValues) const
+	{
+		if (ValueKeys.IsEmpty())
+		{
+			return;
+		}
+
+		if (!IsOfType<TMap<KeyType, ValueType>>())
+		{
+			// @todo_pcg Better logging
+			UE_LOG(LogPCG, Error, TEXT("Type mismatch on GetValues"));
+			return;
+		}
+
+		if (!ensure(ValueKeys.Num() == OutValues.Num()))
+		{
+			return;
+		}
+
+		TArray<FScriptMapHelper*, TInlineAllocator<256>> Temp;
+		Temp.Reserve(OutValues.Num());
+		for (PCG::TScriptMapWrapper<KeyType, ValueType>& Value : OutValues)
+		{
+			Temp.Add(&Value.Helper);
+		}
+		GetValues_Internal(ValueKeys, PCG::Private::FOutValues{TInPlaceType<PCG::Private::FOutValuesAsMap>{}, Temp});
+	}
+
+	////////////// End GetValues Section
+
+	// Fill the OutValues using the array of entry keys.
+	// cf. GetValues for more info on supported types.
+	template <typename T>
+	void GetValuesFromItemKeys(TConstArrayView<PCGMetadataEntryKey> EntryKeys, TArrayView<T> OutValues) const
+	{
+		GetValuesFromItemKeys(PCGValueRangeHelpers::MakeConstValueRange(EntryKeys), OutValues);
+	}
+
+	// Fill the OutValues using the range of entry keys.
+	// cf. GetValues for more info on supported types.
+	template <typename T>
+	void GetValuesFromItemKeys(TConstPCGValueRange<PCGMetadataEntryKey> EntryKeys, TArrayView<T> OutValues) const
+	{
+		if (EntryKeys.IsEmpty() || !ensure(EntryKeys.Num() == OutValues.Num()))
+		{
+			return;
+		}
+
+		TArray<PCGMetadataValueKey> ValueKeys;
+		GetValueKeys(EntryKeys, ValueKeys);
+		GetValues(ValueKeys, OutValues);
+	}
+
+	template <typename T>
+	T GetValueFromItemKey(const PCGMetadataEntryKey InEntryKey) const
+	{
+		T OutValue{};
+		GetValuesFromItemKeys(MakeArrayView(&InEntryKey, 1), MakeArrayView(&OutValue, 1));
+		return OutValue;
+	}
+
+	// Override the default value for this attribute.
+	// Will fail if the type mismatch.
+	template <typename T>
+	void SetDefaultValue(const T& Value)
+	{
+		if (TOptional<FScriptArrayHelper> Helper = MakeArrayHelper(); Helper.IsSet())
+		{
+			check(DefaultValue.IsValid());
+			AddValues(MakeArrayView(&Value, 1), /*ExistingIndex=*/PCG::Private::ExistingIndexForDefaultValue);
+		}
+	}
+
+	template <typename T>
+	void SetValue(const PCGMetadataEntryKey InEntryKey, const T& InValue)
+	{
+		SetValuesFromValueKeys(MakeArrayView(&InEntryKey, 1), AddValues<T>({InValue}));
+	}
+
+	template <typename T>
+	void SetValues(TArrayView<const PCGMetadataEntryKey * const> ItemKeys, TArrayView<const T> InValues)
+	{
+		SetValuesFromValueKeys(ItemKeys, AddValues(InValues));
+	}
+
+	template <typename T>
+	void SetValues(TArrayView<const PCGMetadataEntryKey> ItemKeys, TArrayView<const T> InValues)
+	{
+		SetValuesFromValueKeys(ItemKeys, AddValues(InValues));
+	}
+
+	template <typename T>
+	void SetValues(TConstPCGValueRange<PCGMetadataEntryKey> ItemKeys, TArrayView<const T> InValues)
+	{
+		SetValuesFromValueKeys(ItemKeys, AddValues(InValues));
+	}
+
+	template <typename T>
+	void SetValues(TPCGValueRange<PCGMetadataEntryKey> ItemKeys, TArrayView<const T> InValues)
+	{
+		SetValuesFromValueKeys(PCGValueRangeHelpers::MakeConstValueRange(ItemKeys), AddValues(InValues));
+	}
+
+	// Special version to write on pre-allocated data if we do not compress, in order to be able to write in concurrency.
+	template <typename T>
+	void SetValues_TryLockless(TArrayView<PCGMetadataEntryKey*> ItemKeys, const TArrayView<const T>& InValues, int32 StartIndex)
+	{
+		if (InValues.IsEmpty())
+		{
+			return;
+		}
+
+		if (UnderlyingProperty->bCompressData)
+		{
+			SetValues(ItemKeys, InValues);
+		}
+		else
+		{
+			check(StartIndex != INDEX_NONE && Values.IsValidIndex(StartIndex + InValues.Num() - 1));
+
+			// Since values are already preallocated, we can take advantage of the AddValues with existing index.
+			AddValues(InValues, StartIndex, /*bLockless=*/true);
+		}
+	}
+
+	// Copy values from SrcPtr to the attribute. Property must match the attribute property, and the SrcPtr is expected to be
+	// reinterpreted as the property type. The ItemKey can be the invalid entry key to set the default value.
+	UE_API void SetValueFromProperty(PCGMetadataEntryKey ItemKey, const void* SrcPtr, const FProperty* Property);
+
+	UE_API int32 PreallocateValues(TArrayView<PCGMetadataEntryKey*> EntryKeys, bool bLockless);
+
+	UE_API void Prepare(int32 Count);
+
+	// @todo_pcg: Temporary for debug. Build a struct that will contain the value at the entry key.
+	UE_API FInstancedPropertyBag BuildStructForDebug(PCGMetadataEntryKey EntryKey) const;
+
+	// Utility function to get the pointer to a value using a value key. Return nullptr if invalid. Advanced function, use at your own risk.
+	UE_API const void* GetReadAddressForValueKey_Unsafe(const PCGMetadataValueKey InValueKey) const;
+
+	// Utility function to get the pointer to a value using an entry key. Return nullptr if invalid. Advanced function, use at your own risk.
+	UE_API const void* GetReadAddressFromEntryKey_Unsafe(const PCGMetadataEntryKey InEntryKey) const;
+
+protected:
+	// Base serialize for deprecation, what was done before the merge with generic
+	UE_API void BaseSerialize(FPCGMetadataDomain* InMetadata, FArchive& InArchive);
+	
+	/** Attributes can't be created individually, they need to be created through a Metadata Domain. */
+	friend FPCGMetadataDomain;
+	friend class FPCGAttributeGenericAccessor;
+
+	// New attribute with no parent, name is contained in the InAttributeDesc.
+	UE_API FPCGMetadataAttributeBase(const FPCGMetadataAttributeDesc& InAttributeDesc, FPCGMetadataDomain* InMetadata, bool bInAllowsInterpolation);
+
+	// New attribute with an existing parent, can override the name.
+	UE_API FPCGMetadataAttributeBase(TNotNull<const FPCGMetadataAttributeBase*> InParent, FPCGMetadataDomain* InMetadata, bool bInAllowsInterpolation, TOptional<FName> InName = {});
+
+	// For deprecation purposes of the typed attributes
+	UE_API FPCGMetadataAttributeBase(const FPCGMetadataAttributeDesc& InAttributeDesc, const FPCGMetadataAttributeBase* InParent, FPCGMetadataDomain* InMetadata, bool bInAllowsInterpolation);
+
+	// Allocate the values (without initialization if it can) and return the first index of that new memory space.
+	// Return INDEX_NONE if it didn't work.
+	UE_API int32 AllocateValues_Unsafe(int32 Count);
+
+	// Used to make an array helper of the local FScriptArray.
+	// Optional, since the attribute can be ill-formed, so the user should always check if it is set before
+	// using it.
+	UE_API TOptional<FScriptArrayHelper> MakeArrayHelper() const;
+
+	/**
+	 * GetValues implementation section
+	 * GetValues_Internal is working on a Variant that will type erase the underlying output container
+	 * in order to reduce template code size.
+	 */
+	void GetValues_Internal(const TArrayView<const PCGMetadataValueKey> ValueKeys, PCG::Private::FOutValues OutValues) const
+	{
+		// Bitset with all unretrieved values. If we have any unretrieved value, we will ask the parent for those.
+		TBitArray<> UnretrievedValues(true, ValueKeys.Num());
+
+		return GetValues_Internal(ValueKeys, MoveTemp(OutValues), UnretrievedValues);
+	}
+
+	/**
+	 * AddValues section
+	 * Values will be dispatched to a type erased variant to reduce template code size
+	 * Variant contains all supported containers: Values/Array/Set/Map
+	 *
+	 * Note that to support setting the default value (and potentially write on existing values if it
+	 * was preallocated), all AddValues have an extra parameter `ExistingIndex` that will indicate if allocation
+	 * is necessary or not. This is a very niche feature, and it is expected that everything is already allocated
+	 * or it will crash. There is also an option for lockless, if we expect to be able to write without any lock.
+	 */
+
+	/**
+	 * Basic version that will accept all single values (as an array view).
+	 */
+	template <typename T>
+	TArray<PCGMetadataValueKey> AddValues(TArrayView<const T> InValues, int32 ExistingIndex = INDEX_NONE, bool bLockless = false)
+	{
+		if (InValues.IsEmpty())
+		{
+			return {};
+		}
+
+		if (!IsOfType<T>())
+		{
+			// @todo_pcg Better logging
+			UE_LOG(LogPCG, Error, TEXT("Type mismatch on AddValues"));
+			return {};
+		}
+
+		// For compressed data, we need to analyze which values actually need to be added. If existing index is not none, ignore it.
+		// Need the concept for Equals, since `AddUnique` require the operator '==` to be defined.
+		if constexpr (PCG::Concepts::CIsComparable<T, T>)
+		{
+			if (UnderlyingProperty->bCompressData && ExistingIndex == INDEX_NONE)
+			{
+				// Since we're getting raw values here, we might have duplicates
+				// so we should aim to remove duplicates here so we preserve our 'compress data' idea, otherwise it will break other foundational blocks (e.g. partition)
+				TArray<T, TInlineAllocator<256>> UniqueValues;
+
+				// Initially, fill with mapping to unique values so we can remap them at the end if needed
+				TArray<PCGMetadataValueKey> FoundValueKeys;
+				FoundValueKeys.Reserve(InValues.Num());
+
+				for (int ValueIndex = 0; ValueIndex < InValues.Num(); ++ValueIndex)
+				{
+					FoundValueKeys.Emplace(UniqueValues.AddUnique(InValues[ValueIndex]));
+				}
+
+				return AddCompressedValues(PCG::Private::FInValues{TInPlaceType<PCG::Private::FInValuesByValue>{}, UniqueValues.GetData(), UniqueValues.Num()}, UniqueValues.Num(), FoundValueKeys);
+			}
+		}
+
+		return AddValues_Internal(PCG::Private::FInValues{TInPlaceType<PCG::Private::FInValuesByValue>{}, InValues.GetData(), InValues.Num()}, InValues.Num(), ExistingIndex, bLockless);
+	}
+
+	/**
+	 * Version that accept anything that is array like, like `TArray<T>`, `TArrayView<T>`, `TStaticArray<T, 5>`, `TArray<T, TInlineAllocator<5>>`, etc...
+	 */
+	template <typename ArrayContainer> requires (PCG::Concepts::CIsArrayLike<ArrayContainer>)
+	TArray<PCGMetadataValueKey> AddValues(TConstArrayView<ArrayContainer> InValues, int32 ExistingIndex = INDEX_NONE, bool bLockless = false)
+	{
+		if (InValues.IsEmpty())
+		{
+			return {};
+		}
+
+		using T = typename ArrayContainer::ElementType;
+		if (!IsOfType<TArray<T>>())
+		{
+			// @todo_pcg Better logging
+			UE_LOG(LogPCG, Error, TEXT("Type mismatch on AddValues"));
+			return {};
+		}
+
+		TArray<TTuple<const void*, int32>, TInlineAllocator<256>> InValuesTuple;
+		InValuesTuple.Reserve(InValues.Num());
+		Algo::Transform(InValues, InValuesTuple, [](const auto& Item) { return TTuple<const void*, int32>{Item.GetData(), Item.Num()}; });
+		return AddValues_Internal(PCG::Private::FInValues{TInPlaceType<PCG::Private::FInValuesAsArray>{}, InValuesTuple}, InValues.Num(), ExistingIndex, bLockless);
+	}
+
+	/**
+	 * Version that accepts Sets
+	 */
+	template <typename T, typename ...Extra>
+	TArray<PCGMetadataValueKey> AddValues(TConstArrayView<TSet<T, Extra...>> InValues, int32 ExistingIndex = INDEX_NONE, bool bLockless = false)
+	{
+		if (InValues.IsEmpty())
+		{
+			return {};
+		}
+
+		if (!IsOfType<TSet<T>>())
+		{
+			// @todo_pcg Better logging
+			UE_LOG(LogPCG, Error, TEXT("Type mismatch on AddValues"));
+			return {};
+		}
+
+		TArray<TArray<const void*>, TInlineAllocator<256>> InSetValues;
+		InSetValues.Reserve(InValues.Num());
+		Algo::Transform(InValues, InSetValues, [](const auto& Item) -> TArray<const void*>
+			{
+				TArray<const void*> SetValues;
+				SetValues.Reserve(Item.Num());
+				Algo::Transform(Item, SetValues, [](const T& Value) -> const T* { return &Value;});
+				return SetValues;
+			});
+		return AddValues_Internal(PCG::Private::FInValues{TInPlaceType<PCG::Private::FInValuesAsSet>{}, InSetValues}, InValues.Num(), ExistingIndex, bLockless);
+	}
+
+	/**
+	 * Version that accepts Maps
+	 */
+	template <typename KeyType, typename ValueType, typename ...Extra>
+	TArray<PCGMetadataValueKey> AddValues(TConstArrayView<TMap<KeyType, ValueType, Extra...>> InValues, int32 ExistingIndex = INDEX_NONE, bool bLockless = false)
+	{
+		if (InValues.IsEmpty())
+		{
+			return {};
+		}
+
+		if (!IsOfType<TMap<KeyType, ValueType>>())
+		{
+			// @todo_pcg Better logging
+			UE_LOG(LogPCG, Error, TEXT("Type mismatch on AddValues"));
+			return {};
+		}
+
+		TArray<TArray<TPair<const void*, const void*>>, TInlineAllocator<256>> InMapValues;
+		InMapValues.Reserve(InValues.Num());
+		Algo::Transform(InValues, InMapValues, [](const auto& Item) -> TArray<TPair<const void*, const void*>>
+			{
+				TArray<TPair<const void*, const void*>> MapValues;
+				MapValues.Reserve(Item.Num());
+				Algo::Transform(Item, MapValues, [](const auto& It) -> TPair<const void*, const void*> { return {&It.Key, &It.Value};});
+				return MapValues;
+			});
+		return AddValues_Internal(PCG::Private::FInValues{TInPlaceType<PCG::Private::FInValuesAsMap>{}, InMapValues}, InValues.Num(), ExistingIndex, bLockless);
+	}
+
+	// Internal implementation that will do the dispatch using the typed erased variant.
+	UE_API TArray<PCGMetadataValueKey> AddValues_Internal(const PCG::Private::FInValues& InValues, int32 Count, int32 ExistingIndex = INDEX_NONE, bool bLockless = false, TOptional<TConstArrayView<int32>> OptionalIndices = {});
+
+	// Internal implementation for compressed data
+	UE_API TArray<PCGMetadataValueKey> AddCompressedValues(PCG::Private::FInValues InUniqueValues, int32 Count, TArray<PCGMetadataValueKey>& FoundValueKeys);
+
+	// Internal implementation for typed-erased compressed data. Will create a temporary storage to find all the unique values then will call AddCompressedValues with those.
+	UE_API TArray<PCGMetadataValueKey> AddCompressedValues(PCG::Private::FInValuesByValue InValues, int32 Count);
+
+	/**
+	 * Find all the value keys associated to InValues for compressed data. Return true if all the values were found.
+	 */
+	UE_API bool FindValues(TArray<PCGMetadataValueKey>& OutValueKeys, const PCG::Private::FInValues& InValues, int32 Count) const;
+
+	/**
+	 * Find all the value keys associated to InValues for compressed data. Return true if all the values were found. Preallocated version.
+	 */
+	UE_API bool FindValues(TArrayView<PCGMetadataValueKey> OutValueKeys, const PCG::Private::FInValues& InValues, int32 Count) const;
+
+	template <typename T>
+	void InitForDeprecation(TArray<T> InValues, const T& InDefaultValue)
+	{
+		CachedDesc.Name = Name;
+
+		if (TOptional<FScriptArrayHelper> Helper = MakeArrayHelper(); ensure(Helper))
+		{
+			Helper->EmptyValues(InValues.Num());
+			SetDefaultValue<T>(InDefaultValue);
+
+			// FScriptArray and TArray are interchangeable, if and only if they are the exact same type
+			Helper->MoveAssign(&InValues);
+
+			ValueKeyOffset = GetParent() ? GetParent()->GetValueKeyOffsetForChild() : 0;
+		}
+	}
+
+	using FCopyFunctor = TFunction<void(void* /*DestData*/, const void* /*SrcData*/, const int32 /*Count*/)>;
+	UE_API void CopyInternal(FPCGMetadataAttributeBase* NewAttribute, bool bKeepParent, bool bCopyEntries = true, bool bCopyValues = true, FCopyFunctor CopyFunctor = {}) const;
 
 private:
+	void Init();
+
 	// Unsafe version, needs to be write lock protected.
 	UE_API void SetValueFromValueKey_Unsafe(PCGMetadataEntryKey EntryKey, PCGMetadataValueKey ValueKey, bool bResetValueOnDefaultValueKey, bool bAllowInvalidEntries = false);
 
 	// Gather the value keys for the list of entry keys.
 	// Because we need to update the entry keys if we look for the value keys in the parent, but because the EntryKeys are coming from the outside, we can't modify them.
-	// (as entry keys need to be put into the parent referential when looking for value keys). 
+	// (as entry keys need to be put into the parent referential when looking for value keys).
 	// So we will copy internally the EntryKeys to modify them (and only once) if we are not owner of the memory.
 	UE_API void GetValueKeys_Internal(TConstPCGValueRange<PCGMetadataEntryKey> EntryKeys, TArrayView<PCGMetadataValueKey> OutValueKeys, TBitArray<>& UnsetValues, bool bOwnerOfEntryKeysView = false) const;
 
+	UE_API void GetValues_Internal(const TArrayView<const PCGMetadataValueKey> ValueKeys, PCG::Private::FOutValues OutValues, TBitArray<>& UnretrievedValues) const;
+	void FindValuesInternal(const PCG::Private::FInValues& InValues, int32 Count, TArrayView<PCGMetadataValueKey> ValueKeys, int& ValueKeysSet, bool bIsRoot) const;
+
 protected:
 	TMap<PCGMetadataEntryKey, PCGMetadataValueKey> EntryToValueKeyMap;
 	mutable PCG::FSharedLock EntryMapLock;
 
 	FPCGMetadataDomain* Metadata = nullptr;
 	const FPCGMetadataAttributeBase* Parent = nullptr;
-	int16 TypeId = 0;
 	bool bAllowsInterpolation = false;
+	
+	// Legacy Id that is only set if the attribute is typed. This is purely for deprecation purposes, and the only reliable source to
+	// know if you can `static_cast` or not. All the info is now contained in the Desc that can be queried with `GetAttributeDesc`
+	int16 TypeId = static_cast<int16>(EPCGMetadataTypes::Unknown);
+
+private:
+	FPCGMetadataAttributeDesc CachedDesc;
+	TSharedPtr<FPCGAttributeProperty> UnderlyingProperty;
+	FScriptArray Values;
+	FPCGAttributeDefaultValue DefaultValue;
+
+	mutable PCG::FSharedLock ValueLock;
+	PCGMetadataValueKey ValueKeyOffset = 0;
 
 public:
 	FName Name = NAME_None;
diff --git a/Engine/Plugins/PCG/Source/PCG/Public/Metadata/PCGMetadataAttributeGeneric.h b/Engine/Plugins/PCG/Source/PCG/Public/Metadata/PCGMetadataAttributeGeneric.h
--- a/Engine/Plugins/PCG/Source/PCG/Public/Metadata/PCGMetadataAttributeGeneric.h
+++ b/Engine/Plugins/PCG/Source/PCG/Public/Metadata/PCGMetadataAttributeGeneric.h
@@ -2,631 +2,7 @@
 
 #pragma once
 
-#include "PCGMetadataContainerTypes.h"
-#include "Metadata/PCGMetadataCommon.h"
-#include "PCGModule.h"
-
-#include "Helpers/PCGConcepts.h"
-#include "Helpers/PCGMetadataHelpers.h"
+// @todo_pcg To be removed in a follow-up CL. Will ease the refactor review.
 #include "Metadata/PCGMetadataAttribute.h"
-#include "Metadata/PCGMetadataAttributeTraits.h"
-
-#include "Containers/StaticArray.h"
-#include "Misc/TVariant.h"
-#include "StructUtils/PropertyBag.h"
-
-struct FPCGAttributeDefaultValue;
-class FPCGMetadataAttributeGeneric;
-
-namespace PCG::Private
-{
-	void CopyArray(const FArrayProperty* InProperty, void* DestPtr, const void* SrcPtr, int32 Count);
-
-	/** Magic number for AddValues to know if we need to write to the DefaultValue or not. */
-	constexpr int32 ExistingIndexForDefaultValue = -2;
-}
-
-struct FPCGAttributeProperty
-{
-public:
-	PCG_API FPCGAttributeProperty(const FPCGMetadataAttributeDesc& InAttributeDesc);
-	PCG_API ~FPCGAttributeProperty();
-
-	PCG_API void Copy(void* DestPtr, const void* SrcPtr, int32 Count = 1);
-	PCG_API bool IsValid() const;
-
-private:
-	friend FPCGMetadataAttributeGeneric;
-	friend FPCGAttributeDefaultValue;
-
-	/** Unsafe method that will apply pointer arithmetic to get to the element in the array at the specified index. */
-	const void* GetPtrInArray(const void* InArrayStart, int32 Index) const;
-
-	TUniquePtr<FArrayProperty> Property;
-	size_t InnerElementCachedSize = 0;
-	bool bCompressData = false;
-	bool bIsPlainOldData = false;
-};
-
-// Simple wrapper around a raw pointer to hold the default value to mimic how it was stored in the typed attributes
-struct FPCGAttributeDefaultValue
-{
-	FPCGAttributeDefaultValue() = default;
-	PCG_API ~FPCGAttributeDefaultValue();
-	
-	// Since this is an internal struct and attributes are never copied with constructor/assignment, just delete the copy/move
-	UE_NONCOPYABLE(FPCGAttributeDefaultValue);
-
-	PCG_API void Init(TSharedPtr<FPCGAttributeProperty> InUnderlyingProperty);
-	PCG_API void CopyFrom(const void* SrcData);
-	
-	void* GetRawPtr() { return Memory; }
-	const void* GetRawPtr() const { return Memory; };
-	bool IsValid() const { return Memory != nullptr; }
-
-private:
-	void Destroy();
-	
-	TSharedPtr<FPCGAttributeProperty> UnderlyingProperty;
-	void* Memory = nullptr;
-};
-
-class FPCGMetadataAttributeGeneric : public FPCGMetadataAttributeBase
-{
-public:
-	PCG_API ~FPCGMetadataAttributeGeneric();
-
-	PCG_API virtual void GetResourceSizeEx(FResourceSizeEx& CumulativeResourceSize) const override;
-	PCG_API void AddReferencedObjects(FReferenceCollector& Collector);
-
-	virtual bool IsGeneric() const override { return true; }
-	
-	/** Returns true if the attribute is not a container (Array/Set/Map). */
-	PCG_API bool IsSingleValue() const;
-
-	/** Returns true if the attribute is an array. */
-	PCG_API bool IsArray() const;
-
-	/** Returns true if the attribute is a set. */
-	PCG_API bool IsSet() const;
-
-	/** Returns true if the attribute is a map. */
-	PCG_API bool IsMap() const;
-
-	const FPCGMetadataAttributeDesc& GetAttributeDesc() const { return CachedDesc; }
-	bool DoesCompressData() const { return UnderlyingProperty && UnderlyingProperty->bCompressData; }
-
-	template <typename T>
-	bool IsOfType() const
-	{
-		return PCG::Private::IsEquivalentDesc<T>(CachedDesc);
-	}
-
-	PCG_API virtual PCGMetadataValueKey GetValueKeyOffsetForChild() const override;
-
-	PCG_API virtual void Serialize(FPCGMetadataDomain* InMetadata, FArchive& InArchive) override;
-
-	/**
-	 * GetValues section
-	 * Templated code to support all container types.
-	 * If the provided output container does not match the underlying type, it will fail.
-	 * It is also mandatory that the number of value keys and output container num match.
-	 * 
-	 * Supported types are:
-	 * - Single types that are valid for metadata
-	 * - Pointers to single types that are valid for metadata
-	 * - Typed ConstArrayViews for arrays
-	 * - PCG::TScriptSetWrapper<T> for set, wrapper that can be manipulated like a set
-	 * - PCG::TScriptMapWrapper<T> for map, wrapper that can be manipulated like a map
-	 */
-	template <typename T> requires (!std::is_pointer_v<T>)
-	void GetValues(TArrayView<const PCGMetadataValueKey> ValueKeys, TArrayView<T> OutValues) const
-	{
-		if (ValueKeys.IsEmpty())
-		{
-			return;
-		}
-
-		if (!IsOfType<T>())
-		{
-			// @todo_pcg Better logging
-			UE_LOG(LogPCG, Error, TEXT("Type mismatch on GetValues"));
-			return;
-		}
-		
-		if (!ensure(ValueKeys.Num() == OutValues.Num()))
-		{
-			return;
-		}
-		
-		GetValues_Internal(ValueKeys, PCG::Private::FOutValues{TInPlaceType<PCG::Private::FOutValuesByValue>{}, OutValues.GetData()});
-	}
-
-	template <typename T>
-	void GetValues(TArrayView<const PCGMetadataValueKey> ValueKeys, TArrayView<const T*> OutValues) const
-	{
-		if (ValueKeys.IsEmpty())
-		{
-			return;
-		}
-
-		if (!IsOfType<T>())
-		{
-			// @todo_pcg Better logging
-			UE_LOG(LogPCG, Error, TEXT("Type mismatch on GetValues"));
-			return;
-		}
-		
-		if (!ensure(ValueKeys.Num() == OutValues.Num()))
-		{
-			return;
-		}
-		
-		// We need to type erase the pointers, but to be fully C++ compliant about aliasing rules we can't reinterpret_cast. So make a copy of the pointers.
-		TArray<const void*, TInlineAllocator<256>> Temp;
-		Temp.SetNumUninitialized(OutValues.Num());
-		GetValues_Internal(ValueKeys, PCG::Private::FOutValues{ TInPlaceType<PCG::Private::FOutValuesByPtr>{}, Temp });
-		for (int32 i = 0; i < OutValues.Num(); ++i)
-		{
-			OutValues[i] = static_cast<const T*>(Temp[i]);
-		}
-	}
-	
-	template <typename T> 
-	void GetValues(TArrayView<const PCGMetadataValueKey> ValueKeys, TArrayView<TConstArrayView<T>> OutValues) const
-	{
-		if (ValueKeys.IsEmpty())
-		{
-			return;
-		}
-
-		if (!IsOfType<TArray<T>>())
-		{
-			// @todo_pcg Better logging
-			UE_LOG(LogPCG, Error, TEXT("Type mismatch on GetValues"));
-			return;
-		}
-		
-		if (!ensure(ValueKeys.Num() == OutValues.Num()))
-		{
-			return;
-		}
-		
-		TArray<TTuple<const void*, int32>, TInlineAllocator<256>> Temp;
-		Temp.SetNumUninitialized(OutValues.Num());
-		GetValues_Internal(ValueKeys, PCG::Private::FOutValues{TInPlaceType<PCG::Private::FOutValuesAsArray>{}, Temp});
-		for (int32 i = 0; i < OutValues.Num(); ++i)
-		{
-			OutValues[i] = MakeConstArrayView<T>(static_cast<const T*>(Temp[i].Get<0>()), Temp[i].Get<1>());
-		}
-	}
-	
-	// To support sets
-	template <typename T> 
-	void GetValues(TArrayView<const PCGMetadataValueKey> ValueKeys, TArrayView<PCG::TScriptSetWrapper<T>> OutValues) const
-	{
-		if (ValueKeys.IsEmpty())
-		{
-			return;
-		}
-
-		if (!IsOfType<TSet<T>>())
-		{
-			// @todo_pcg Better logging
-			UE_LOG(LogPCG, Error, TEXT("Type mismatch on GetValues"));
-			return;
-		}
-		
-		if (!ensure(ValueKeys.Num() == OutValues.Num()))
-		{
-			return;
-		}
-		
-		TArray<FScriptSetHelper*, TInlineAllocator<256>> Temp;
-		Temp.Reserve(OutValues.Num());
-		for (PCG::TScriptSetWrapper<T>& Value : OutValues)
-		{
-			Temp.Add(&Value.Helper);
-		}
-		GetValues_Internal(ValueKeys, PCG::Private::FOutValues{TInPlaceType<PCG::Private::FOutValuesAsSet>{}, Temp});
-	}
-	
-	// To support maps
-	template <typename KeyType, typename ValueType> 
-	void GetValues(TArrayView<const PCGMetadataValueKey> ValueKeys, TArrayView<PCG::TScriptMapWrapper<KeyType, ValueType>> OutValues) const
-	{
-		if (ValueKeys.IsEmpty())
-		{
-			return;
-		}
-
-		if (!IsOfType<TMap<KeyType, ValueType>>())
-		{
-			// @todo_pcg Better logging
-			UE_LOG(LogPCG, Error, TEXT("Type mismatch on GetValues"));
-			return;
-		}
-		
-		if (!ensure(ValueKeys.Num() == OutValues.Num()))
-		{
-			return;
-		}
-		
-		TArray<FScriptMapHelper*, TInlineAllocator<256>> Temp;
-		Temp.Reserve(OutValues.Num());
-		for (PCG::TScriptMapWrapper<KeyType, ValueType>& Value : OutValues)
-		{
-			Temp.Add(&Value.Helper);
-		}
-		GetValues_Internal(ValueKeys, PCG::Private::FOutValues{TInPlaceType<PCG::Private::FOutValuesAsMap>{}, Temp});
-	}
-	
-	////////////// End GetValues Section
-	
-	// Fill the OutValues using the array of entry keys.
-	// cf. GetValues for more info on supported types.
-	template <typename T>
-	void GetValuesFromItemKeys(TConstArrayView<PCGMetadataEntryKey> EntryKeys, TArrayView<T> OutValues) const
-	{
-		GetValuesFromItemKeys(PCGValueRangeHelpers::MakeConstValueRange(EntryKeys), OutValues);
-	}
-	
-	// Fill the OutValues using the range of entry keys.
-	// cf. GetValues for more info on supported types.
-	template <typename T>
-	void GetValuesFromItemKeys(TConstPCGValueRange<PCGMetadataEntryKey> EntryKeys, TArrayView<T> OutValues) const
-	{
-		if (EntryKeys.IsEmpty() || !ensure(EntryKeys.Num() == OutValues.Num()))
-		{
-			return;
-		}
-
-		TArray<PCGMetadataValueKey> ValueKeys;
-		GetValueKeys(EntryKeys, ValueKeys);
-		GetValues(ValueKeys, OutValues);
-	}
-
-	template <typename T>
-	T GetValueFromItemKey(const PCGMetadataEntryKey InEntryKey) const
-	{
-		T OutValue{};
-		GetValuesFromItemKeys(MakeArrayView(&InEntryKey, 1), MakeArrayView(&OutValue, 1));
-		return OutValue;
-	}
-
-	// Override the default value for this attribute.
-	// Will fail if the type mismatch.
-	template <typename T>
-	void SetDefaultValue(const T& Value)
-	{
-		if (TOptional<FScriptArrayHelper> Helper = MakeArrayHelper(); Helper.IsSet())
-		{
-			check(DefaultValue.IsValid());
-			AddValues(MakeArrayView(&Value, 1), /*ExistingIndex=*/PCG::Private::ExistingIndexForDefaultValue);
-		}
-	}
-
-	template <typename T>
-	void SetValue(const PCGMetadataEntryKey InEntryKey, const T& InValue)
-	{
-		SetValuesFromValueKeys(MakeArrayView(&InEntryKey, 1), AddValues<T>({InValue}));
-	}
-
-	template <typename T>
-	void SetValues(TArrayView<const PCGMetadataEntryKey * const> ItemKeys, TArrayView<const T> InValues)
-	{
-		SetValuesFromValueKeys(ItemKeys, AddValues(InValues));
-	}
-
-	template <typename T>
-	void SetValues(TArrayView<const PCGMetadataEntryKey> ItemKeys, TArrayView<const T> InValues)
-	{
-		SetValuesFromValueKeys(ItemKeys, AddValues(InValues));
-	}
-	
-	template <typename T>
-	void SetValues(TConstPCGValueRange<PCGMetadataEntryKey> ItemKeys, TArrayView<const T> InValues)
-	{
-		SetValuesFromValueKeys(ItemKeys, AddValues(InValues));
-	}
-	
-	template <typename T>
-	void SetValues(TPCGValueRange<PCGMetadataEntryKey> ItemKeys, TArrayView<const T> InValues)
-	{
-		SetValuesFromValueKeys(PCGValueRangeHelpers::MakeConstValueRange(ItemKeys), AddValues(InValues));
-	}
-
-	// Special version to write on pre-allocated data if we do not compress, in order to be able to write in concurrency.
-	template <typename T>
-	void SetValues_TryLockless(TArrayView<PCGMetadataEntryKey*> ItemKeys, const TArrayView<const T>& InValues, int32 StartIndex)
-	{
-		if (InValues.IsEmpty())
-		{
-			return;
-		}
-
-		if (UnderlyingProperty->bCompressData)
-		{
-			SetValues(ItemKeys, InValues);
-		}
-		else
-		{
-			check(StartIndex != INDEX_NONE && Values.IsValidIndex(StartIndex + InValues.Num() - 1));
-			
-			// Since values are already preallocated, we can take advantage of the AddValues with existing index.
-			AddValues(InValues, StartIndex, /*bLockless=*/true);
-		}
-	}
-
-	// Copy values from SrcPtr to the attribute. Property must match the attribute property, and the SrcPtr is expected to be
-	// reinterpreted as the property type. The ItemKey can be the invalid entry key to set the default value.
-	PCG_API void SetValueFromProperty(PCGMetadataEntryKey ItemKey, const void* SrcPtr, const FProperty* Property);
-	
-	PCG_API int32 PreallocateValues(TArrayView<PCGMetadataEntryKey*> EntryKeys, bool bLockless);
-	
-	PCG_API void Prepare(int32 Count);
-
-	// @todo_pcg: Temporary for debug. Build a struct that will contain the value at the entry key.
-	PCG_API FInstancedPropertyBag BuildStructForDebug(PCGMetadataEntryKey EntryKey) const;
-
-	// Utility function to get the pointer to a value using a value key. Return nullptr if invalid. Advanced function, use at your own risk.
-	PCG_API const void* GetReadAddressForValueKey_Unsafe(const PCGMetadataValueKey InValueKey) const;
-
-	// Utility function to get the pointer to a value using an entry key. Return nullptr if invalid. Advanced function, use at your own risk.
-	PCG_API const void* GetReadAddressFromEntryKey_Unsafe(const PCGMetadataEntryKey InEntryKey) const;
-
-protected:
-	/** Attributes can't be created individually, they need to be created through a Metadata Domain. */
-	friend FPCGMetadataDomain;
-	
-	friend class FPCGAttributeGenericAccessor;
-
-	// New attribute with no parent, name is contained in the InAttributeDesc.
-	PCG_API FPCGMetadataAttributeGeneric(const FPCGMetadataAttributeDesc& InAttributeDesc, FPCGMetadataDomain* InMetadata, bool bInAllowsInterpolation);
-
-	// New attribute with an existing parent, can override the name.
-	PCG_API FPCGMetadataAttributeGeneric(TNotNull<const FPCGMetadataAttributeGeneric*> InParent, FPCGMetadataDomain* InMetadata, bool bInAllowsInterpolation, TOptional<FName> InName = {});
-
-	// For deprecation purposes of the typed attributes
-	PCG_API FPCGMetadataAttributeGeneric(const FPCGMetadataAttributeDesc& InAttributeDesc, const FPCGMetadataAttributeGeneric* InParent, FPCGMetadataDomain* InMetadata, bool bInAllowsInterpolation);
-
-	// For deserialization purposes
-	FPCGMetadataAttributeGeneric() = default;
-
-	// Allocate the values (without initialization if it can) and return the first index of that new memory space.
-	// Return INDEX_NONE if it didn't work.
-	PCG_API int32 AllocateValues_Unsafe(int32 Count);
-
-	// Used to make an array helper of the local FScriptArray.
-	// Optional, since the attribute can be ill-formed, so the user should always check if it is set before
-	// using it.
-	PCG_API TOptional<FScriptArrayHelper> MakeArrayHelper() const;
-	
-	/**
-	 * GetValues implementation section
-	 * GetValues_Internal is working on a Variant that will type erase the underlying output container
-	 * in order to reduce template code size.
-	 */
-	void GetValues_Internal(const TArrayView<const PCGMetadataValueKey> ValueKeys, PCG::Private::FOutValues OutValues) const
-	{
-		// Bitset with all unretrieved values. If we have any unretrieved value, we will ask the parent for those.
-		TBitArray<> UnretrievedValues(true, ValueKeys.Num());
-		
-		return GetValues_Internal(ValueKeys, MoveTemp(OutValues), UnretrievedValues);
-	}
-
-	/**
-	 * AddValues section
-	 * Values will be dispatched to a type erased variant to reduce template code size
-	 * Variant contains all supported containers: Values/Array/Set/Map
-	 * 
-	 * Note that to support setting the default value (and potentially write on existing values if it
-	 * was preallocated), all AddValues have an extra parameter `ExistingIndex` that will indicate if allocation
-	 * is necessary or not. This is a very niche feature, and it is expected that everything is already allocated
-	 * or it will crash. There is also an option for lockless, if we expect to be able to write without any lock.
-	 */
-
-	/**
-	 * Basic version that will accept all single values (as an array view).
-	 */
-	template <typename T>
-	TArray<PCGMetadataValueKey> AddValues(TArrayView<const T> InValues, int32 ExistingIndex = INDEX_NONE, bool bLockless = false)
-	{
-		if (InValues.IsEmpty())
-		{
-			return {};
-		}
-
-		if (!IsOfType<T>())
-		{
-			// @todo_pcg Better logging
-			UE_LOG(LogPCG, Error, TEXT("Type mismatch on AddValues"));
-			return {};
-		}
-
-		// For compressed data, we need to analyze which values actually need to be added. If existing index is not none, ignore it.
-		// Need the concept for Equals, since `AddUnique` require the operator '==` to be defined.
-		if constexpr (PCG::Concepts::CIsComparable<T, T>)
-		{
-			if (UnderlyingProperty->bCompressData && ExistingIndex == INDEX_NONE)
-			{
-				// Since we're getting raw values here, we might have duplicates
-				// so we should aim to remove duplicates here so we preserve our 'compress data' idea, otherwise it will break other foundational blocks (e.g. partition)
-				TArray<T, TInlineAllocator<256>> UniqueValues;
-
-				// Initially, fill with mapping to unique values so we can remap them at the end if needed
-				TArray<PCGMetadataValueKey> FoundValueKeys;
-				FoundValueKeys.Reserve(InValues.Num());
-
-				for (int ValueIndex = 0; ValueIndex < InValues.Num(); ++ValueIndex)
-				{
-					FoundValueKeys.Emplace(UniqueValues.AddUnique(InValues[ValueIndex]));
-				}
-
-				return AddCompressedValues(PCG::Private::FInValues{TInPlaceType<PCG::Private::FInValuesByValue>{}, UniqueValues.GetData(), UniqueValues.Num()}, UniqueValues.Num(), FoundValueKeys);
-			}
-		}
-
-		return AddValues_Internal(PCG::Private::FInValues{TInPlaceType<PCG::Private::FInValuesByValue>{}, InValues.GetData(), InValues.Num()}, InValues.Num(), ExistingIndex, bLockless);
-	}
-	
-	/**
-	 * Version that accept anything that is array like, like `TArray<T>`, `TArrayView<T>`, `TStaticArray<T, 5>`, `TArray<T, TInlineAllocator<5>>`, etc...
-	 */
-	template <typename ArrayContainer> requires (PCG::Concepts::CIsArrayLike<ArrayContainer>)
-	TArray<PCGMetadataValueKey> AddValues(TConstArrayView<ArrayContainer> InValues, int32 ExistingIndex = INDEX_NONE, bool bLockless = false)
-	{
-		if (InValues.IsEmpty())
-		{
-			return {};
-		}
-
-		using T = typename ArrayContainer::ElementType;
-		if (!IsOfType<TArray<T>>())
-		{
-			// @todo_pcg Better logging
-			UE_LOG(LogPCG, Error, TEXT("Type mismatch on AddValues"));
-			return {};
-		}
-		
-		TArray<TTuple<const void*, int32>, TInlineAllocator<256>> InValuesTuple;
-		InValuesTuple.Reserve(InValues.Num());
-		Algo::Transform(InValues, InValuesTuple, [](const auto& Item) { return TTuple<const void*, int32>{Item.GetData(), Item.Num()}; });
-		return AddValues_Internal(PCG::Private::FInValues{TInPlaceType<PCG::Private::FInValuesAsArray>{}, InValuesTuple}, InValues.Num(), ExistingIndex, bLockless);
-	}
-
-	/**
-	 * Version that accepts Sets
-	 */
-	template <typename T, typename ...Extra>
-	TArray<PCGMetadataValueKey> AddValues(TConstArrayView<TSet<T, Extra...>> InValues, int32 ExistingIndex = INDEX_NONE, bool bLockless = false)
-	{
-		if (InValues.IsEmpty())
-		{
-			return {};
-		}
-
-		if (!IsOfType<TSet<T>>())
-		{
-			// @todo_pcg Better logging
-			UE_LOG(LogPCG, Error, TEXT("Type mismatch on AddValues"));
-			return {};
-		}
-		
-		TArray<TArray<const void*>, TInlineAllocator<256>> InSetValues;
-		InSetValues.Reserve(InValues.Num());
-		Algo::Transform(InValues, InSetValues, [](const auto& Item) -> TArray<const void*>
-			{
-				TArray<const void*> SetValues;
-				SetValues.Reserve(Item.Num());
-				Algo::Transform(Item, SetValues, [](const T& Value) -> const T* { return &Value;});
-				return SetValues;
-			});
-		return AddValues_Internal(PCG::Private::FInValues{TInPlaceType<PCG::Private::FInValuesAsSet>{}, InSetValues}, InValues.Num(), ExistingIndex, bLockless);
-	}
-	
-	/**
-	 * Version that accepts Maps
-	 */
-	template <typename KeyType, typename ValueType, typename ...Extra>
-	TArray<PCGMetadataValueKey> AddValues(TConstArrayView<TMap<KeyType, ValueType, Extra...>> InValues, int32 ExistingIndex = INDEX_NONE, bool bLockless = false)
-	{
-		if (InValues.IsEmpty())
-		{
-			return {};
-		}
-
-		if (!IsOfType<TMap<KeyType, ValueType>>())
-		{
-			// @todo_pcg Better logging
-			UE_LOG(LogPCG, Error, TEXT("Type mismatch on AddValues"));
-			return {};
-		}
-		
-		TArray<TArray<TPair<const void*, const void*>>, TInlineAllocator<256>> InMapValues;
-		InMapValues.Reserve(InValues.Num());
-		Algo::Transform(InValues, InMapValues, [](const auto& Item) -> TArray<TPair<const void*, const void*>>
-			{
-				TArray<TPair<const void*, const void*>> MapValues;
-				MapValues.Reserve(Item.Num());
-				Algo::Transform(Item, MapValues, [](const auto& It) -> TPair<const void*, const void*> { return {&It.Key, &It.Value};});
-				return MapValues;
-			});
-		return AddValues_Internal(PCG::Private::FInValues{TInPlaceType<PCG::Private::FInValuesAsMap>{}, InMapValues}, InValues.Num(), ExistingIndex, bLockless);
-	}
-	
-	// Internal implementation that will do the dispatch using the typed erased variant.
-	PCG_API TArray<PCGMetadataValueKey> AddValues_Internal(const PCG::Private::FInValues& InValues, int32 Count, int32 ExistingIndex = INDEX_NONE, bool bLockless = false, TOptional<TConstArrayView<int32>> OptionalIndices = {});
-
-	// Internal implementation for compressed data
-	PCG_API TArray<PCGMetadataValueKey> AddCompressedValues(PCG::Private::FInValues InUniqueValues, int32 Count, TArray<PCGMetadataValueKey>& FoundValueKeys);
-
-	// Internal implementation for typed-erased compressed data. Will create a temporary storage to find all the unique values then will call AddCompressedValues with those.
-	PCG_API TArray<PCGMetadataValueKey> AddCompressedValues(PCG::Private::FInValuesByValue InValues, int32 Count);
-
-	/** 
-	 * Find all the value keys associated to InValues for compressed data. Return true if all the values were found.
-	 */
-	PCG_API bool FindValues(TArray<PCGMetadataValueKey>& OutValueKeys, const PCG::Private::FInValues& InValues, int32 Count) const;
-
-	/** 
-	 * Find all the value keys associated to InValues for compressed data. Return true if all the values were found. Preallocated version.
-	 */
-	PCG_API bool FindValues(TArrayView<PCGMetadataValueKey> OutValueKeys, const PCG::Private::FInValues& InValues, int32 Count) const;
-
-	template <typename T>
-	void InitForDeprecation(TArray<T> InValues, const T& InDefaultValue)
-	{
-		CachedDesc.Name = Name;
-
-		if (TOptional<FScriptArrayHelper> Helper = MakeArrayHelper(); ensure(Helper))
-		{
-			Helper->EmptyValues(InValues.Num());
-			SetDefaultValue<T>(InDefaultValue);
-
-			// FScriptArray and TArray are interchangeable, if and only if they are the exact same type
-			Helper->MoveAssign(&InValues);
-			
-			ValueKeyOffset = GetParent() ? GetParent()->GetValueKeyOffsetForChild() : 0;
-		}
-	}
-
-	using FCopyFunctor = TFunction<void(void* /*DestData*/, const void* /*SrcData*/, const int32 /*Count*/)>;
-	PCG_API void CopyInternal(FPCGMetadataAttributeGeneric* NewAttribute, bool bKeepParent, bool bCopyEntries = true, bool bCopyValues = true, FCopyFunctor CopyFunctor = {}) const;
-
-private:
-	void Init();
-	
-	PCG_API void GetValues_Internal(const TArrayView<const PCGMetadataValueKey> ValueKeys, PCG::Private::FOutValues OutValues, TBitArray<>& UnretrievedValues) const;
-	void FindValuesInternal(const PCG::Private::FInValues& InValues, int32 Count, TArrayView<PCGMetadataValueKey> ValueKeys, int& ValueKeysSet, bool bIsRoot) const;
-
-public:
-	PCG_API virtual void Flatten() override;
-	PCG_API virtual void FlattenAndCompress(const TArrayView<const PCGMetadataEntryKey>& InEntryKeysToKeep) override;
-	PCG_API virtual void Reset() override;
-	PCG_API virtual FPCGMetadataAttributeBase* Copy(FName NewName, FPCGMetadataDomain* InMetadata, bool bKeepParent, bool bCopyEntries = true, bool bCopyValues = true) const override;
-
-	UE_DEPRECATED(5.6, "Use the version with FPCGMetadataDomain")
-	PCG_API virtual FPCGMetadataAttributeBase* Copy(FName NewName, UPCGMetadata* InMetadata, bool bKeepParent, bool bCopyEntries = true, bool bCopyValues = true) const override;
-
-	PCG_API virtual FPCGMetadataAttributeBase* CopyToAnotherType(int16 Type) const override;
-	PCG_API virtual void SetValue(PCGMetadataEntryKey ItemKey, const FPCGMetadataAttributeBase* InAttribute, PCGMetadataEntryKey InEntryKey) override;
-	PCG_API virtual void SetZeroValue(PCGMetadataEntryKey ItemKey) override;
-	PCG_API virtual void AccumulateValue(PCGMetadataEntryKey ItemKey, const FPCGMetadataAttributeBase* InAttributeA, PCGMetadataEntryKey InEntryKeyA, float Weight) override;
-	PCG_API virtual void SetWeightedValue(PCGMetadataEntryKey ItemKey, const FPCGMetadataAttributeBase* InAttribute, const TArrayView<const TPair<PCGMetadataEntryKey, float>>& InWeightedKeys) override;
-	PCG_API virtual void SetValue(PCGMetadataEntryKey ItemKey, const FPCGMetadataAttributeBase* InAttributeA, PCGMetadataEntryKey InEntryKeyA, const FPCGMetadataAttributeBase* InAttributeB, PCGMetadataEntryKey InEntryKeyB, EPCGMetadataOp Op) override;
-	PCG_API virtual bool IsEqualToDefaultValue(PCGMetadataValueKey ValueKey) const override;
-	PCG_API virtual void SetDefaultValueToFirstEntry() override;
-	PCG_API virtual bool UsesValueKeys() const override;
-	PCG_API virtual bool AreValuesEqualForEntryKeys(PCGMetadataEntryKey EntryKey1, PCGMetadataEntryKey EntryKey2) const override;
-	PCG_API virtual bool AreValuesEqual(PCGMetadataValueKey ValueKey1, PCGMetadataValueKey ValueKey2) const override;
-
-private:
-	FPCGMetadataAttributeDesc CachedDesc;
-	TSharedPtr<FPCGAttributeProperty> UnderlyingProperty;
-	FScriptArray Values;
-	FPCGAttributeDefaultValue DefaultValue;
 
-	mutable PCG::FSharedLock ValueLock;
-	PCGMetadataValueKey ValueKeyOffset = 0;
-};
+using FPCGMetadataAttributeGeneric = FPCGMetadataAttributeBase;
\ No newline at end of file
diff --git a/Engine/Plugins/PCG/Source/PCG/Public/Metadata/PCGMetadataAttributeTpl.h b/Engine/Plugins/PCG/Source/PCG/Public/Metadata/PCGMetadataAttributeTpl.h
--- a/Engine/Plugins/PCG/Source/PCG/Public/Metadata/PCGMetadataAttributeTpl.h
+++ b/Engine/Plugins/PCG/Source/PCG/Public/Metadata/PCGMetadataAttributeTpl.h
@@ -2,11 +2,11 @@
 
 #pragma once
 
-#include "PCGMetadataAttributeGeneric.h"
+
 
 #include "PCGModule.h"
 #include "Helpers/PCGMetadataHelpers.h"
-#include "Metadata/PCGMetadataAttribute.h"
+#include "Metadata/PCGMetadataAttributeGeneric.h"
 #include "Metadata/PCGMetadataAttributeTraits.h"
 #include "Metadata/PCGMetadataCommon.h"
 
@@ -21,17 +21,23 @@ class UPCGMetadata;
  * It does not support arrays or complex structs.
  */
 template<typename T>
-class FPCGMetadataAttribute : public FPCGMetadataAttributeGeneric
+class FPCGMetadataAttribute : public FPCGMetadataAttributeBase
 {
 	template<typename U>
 	friend class FPCGMetadataAttribute;
+	
+	friend FPCGMetadataDomain;
 
-public:
-	/** TODO: All deprecated versions taking a UPCGMetadata will be added at the end of this declaration. For the review, it will be in place for "clarity" */
-	FPCGMetadataAttribute(FPCGMetadataDomain* InMetadata, FName InName, const FPCGMetadataAttributeBase* InParent, const T& InDefaultValue, bool bInAllowsInterpolation)
-		: FPCGMetadataAttributeGeneric(PCG::Private::MakeAttributeDesc<T>(InName), InParent && InParent->IsGeneric() ? static_cast<const FPCGMetadataAttributeGeneric*>(InParent) : nullptr, InMetadata, bInAllowsInterpolation)
+protected:
+	// Constructors were public before 5.8, but are now protected. Use this sentinel token to guard the new protected constructors.
+	struct FProtectedToken{};
+	
+	FPCGMetadataAttribute(FProtectedToken, FPCGMetadataDomain* InMetadata, FName InName, const FPCGMetadataAttributeBase* InParent, const T& InDefaultValue, bool bInAllowsInterpolation)
+		: FPCGMetadataAttributeBase(PCG::Private::MakeAttributeDesc<T>(InName), InParent, InMetadata, bInAllowsInterpolation)
 	{
-		static_assert(PCG::Private::IsPCGType<T>(), "Only supports legacy types, such as int32 or float");
+		static_assert(PCG::Private::TIsBasicType<T>::Value, "Only supports legacy types, such as int32 or float");
+		TypeId = PCG::Private::MetadataTypes<T>::Id;
+		
 		if (!InParent)
 		{
 			SetDefaultValue(InDefaultValue);
@@ -39,11 +45,24 @@ public:
 	}
 
 	// This constructor is used only during serialization
-	FPCGMetadataAttribute()
-		: FPCGMetadataAttributeGeneric(PCG::Private::MakeAttributeDesc<T>(NAME_None), nullptr, false)
+	FPCGMetadataAttribute(FProtectedToken)
+		: FPCGMetadataAttributeBase(PCG::Private::MakeAttributeDesc<T>(NAME_None), nullptr, false)
 	{
-		static_assert(PCG::Private::IsPCGType<T>(), "Only supports legacy types, such as int32 or float");
+		static_assert(PCG::Private::TIsBasicType<T>::Value, "Only supports legacy types, such as int32 or float");
+		TypeId = PCG::Private::MetadataTypes<T>::Id;
 	}
+	
+public:
+	UE_DEPRECATED(5.8, "Metadata attributes should never be constructed individually and should be created from the Metadata domain.")
+	FPCGMetadataAttribute(FPCGMetadataDomain* InMetadata, FName InName, const FPCGMetadataAttributeBase* InParent, const T& InDefaultValue, bool bInAllowsInterpolation)
+		: FPCGMetadataAttribute(FProtectedToken{}, InMetadata, InName, InParent, InDefaultValue, bInAllowsInterpolation)
+	{}
+
+	// This constructor is used only during serialization
+	UE_DEPRECATED(5.8, "Metadata attributes should never be constructed individually and should be created from the Metadata domain.")
+	FPCGMetadataAttribute()
+		: FPCGMetadataAttribute(FProtectedToken{})
+	{}
 
 	virtual void Serialize(FPCGMetadataDomain* InMetadata, FArchive& InArchive) override
 	{
@@ -52,7 +71,7 @@ public:
 		if (InArchive.IsLoading() && InArchive.CustomVer(FFortniteMainBranchObjectVersion::GUID) < FFortniteMainBranchObjectVersion::ConvertFPCGMetadataAttributeToGenericAttributes)
 		{
 			// Need to forward all the values coming from the previous serialization, and follow the same order of initialization
-			FPCGMetadataAttributeBase::Serialize(InMetadata, InArchive);
+			FPCGMetadataAttributeBase::BaseSerialize(InMetadata, InArchive);
 
 			TArray<T> PreviousValues;
 			T PreviousDefaultValue{};
@@ -63,7 +82,7 @@ public:
 		}
 		else
 		{
-			FPCGMetadataAttributeGeneric::Serialize(InMetadata, InArchive);
+			FPCGMetadataAttributeBase::Serialize(InMetadata, InArchive);
 		}
 	}
 
@@ -99,7 +118,7 @@ protected:
 		}
 		
 		// This copies to a new attribute.
-		FPCGMetadataAttribute<U>* AttributeCopy = new FPCGMetadataAttribute<U>(InMetadata, NewName, bKeepParent ? this : nullptr, U{}, bAllowsInterpolation);
+		FPCGMetadataAttribute<U>* AttributeCopy = new FPCGMetadataAttribute<U>(typename FPCGMetadataAttribute<U>::FProtectedToken(), InMetadata, NewName, bKeepParent ? this : nullptr, U{}, bAllowsInterpolation);
 		
 		auto CopyFunctor = [](void* DestData, const void* SrcData, const int32 Count)
 		{
@@ -140,12 +159,12 @@ public:
 		if constexpr (PCG::Private::MetadataTraits<T>::CanInterpolate)
 		{
 			check(ItemKey != PCGInvalidEntryKey);
-			
+
 			if (InAttribute && !InAttribute->IsGeneric())
 			{
 				return;
 			}
-		
+
 			const FPCGMetadataAttributeGeneric* Attribute = static_cast<const FPCGMetadataAttributeGeneric*>(InAttribute);
 		
 			if (Attribute && !Attribute->IsOfType<T>())
@@ -162,12 +181,12 @@ public:
 		if constexpr (PCG::Private::MetadataTraits<T>::CanInterpolate)
 		{
 			check(ItemKey != PCGInvalidEntryKey);
-			
+
 			if (InAttribute && !InAttribute->IsGeneric())
 			{
 				return;
 			}
-		
+
 			const FPCGMetadataAttributeGeneric* Attribute = static_cast<const FPCGMetadataAttributeGeneric*>(InAttribute);
 		
 			if (Attribute && !Attribute->IsOfType<T>())
@@ -183,19 +202,19 @@ public:
 	//~ Note that this override is purely to disambiguate between all the different versions of SetValue. It just forward the call to the super method.
 	virtual void SetValue(PCGMetadataEntryKey ItemKey, const FPCGMetadataAttributeBase* InAttribute, PCGMetadataEntryKey InEntryKey) override
 	{
-		return FPCGMetadataAttributeGeneric::SetValue(ItemKey, InAttribute, InEntryKey);
+		return FPCGMetadataAttributeBase::SetValue(ItemKey, InAttribute, InEntryKey);
 	}
 
 	virtual void SetValue(PCGMetadataEntryKey ItemKey, const FPCGMetadataAttributeBase* InAttributeA, PCGMetadataEntryKey InEntryKeyA, const FPCGMetadataAttributeBase* InAttributeB, PCGMetadataEntryKey InEntryKeyB, EPCGMetadataOp Op) override
 	{
 		check(ItemKey != PCGInvalidEntryKey);
 		bool bAppliedValue = false;
-		
+
 		if ((InAttributeA && !InAttributeA->IsGeneric()) || (InAttributeB && !InAttributeB->IsGeneric()))
 		{
 			return;
 		}
-		
+
 		const FPCGMetadataAttributeGeneric* AttributeA = static_cast<const FPCGMetadataAttributeGeneric*>(InAttributeA);
 		const FPCGMetadataAttributeGeneric* AttributeB = static_cast<const FPCGMetadataAttributeGeneric*>(InAttributeB);
 		
@@ -582,9 +601,10 @@ protected:
 
 namespace PCGMetadataAttribute
 {
+	UE_DEPRECATED(5.8, "Should never be used anymore")
 	inline FPCGMetadataAttributeBase* AllocateEmptyAttributeFromType(int16 TypeId)
 	{
-
+PRAGMA_DISABLE_DEPRECATION_WARNINGS
 		switch (TypeId)
 		{
 
@@ -595,6 +615,7 @@ namespace PCGMetadataAttribute
 		default:
 			return nullptr;
 		}
+PRAGMA_ENABLE_DEPRECATION_WARNINGS
 	}
 
 	template <typename Func, typename... Args>
diff --git a/Engine/Plugins/PCG/Source/PCG/Public/Metadata/PCGMetadataAttributeTraits.h b/Engine/Plugins/PCG/Source/PCG/Public/Metadata/PCGMetadataAttributeTraits.h
--- a/Engine/Plugins/PCG/Source/PCG/Public/Metadata/PCGMetadataAttributeTraits.h
+++ b/Engine/Plugins/PCG/Source/PCG/Public/Metadata/PCGMetadataAttributeTraits.h
@@ -161,8 +161,6 @@ namespace PCG
 		PCGMetadataGenerateDataTypes(FName, Name);
 		PCGMetadataGenerateDataTypes(FSoftObjectPath, SoftObjectPath);
 		PCGMetadataGenerateDataTypes(FSoftClassPath, SoftClassPath);
-		PCGMetadataGenerateDataTypes(FText, Text);
-		PCGMetadataGenerateDataTypes(uint8, Byte);
 
 #undef PCGMetadataGenerateDataTypes
 
diff --git a/Engine/Plugins/PCG/Source/PCG/Public/Metadata/PCGMetadataContainerTypes.h b/Engine/Plugins/PCG/Source/PCG/Public/Metadata/PCGMetadataContainerTypes.h
--- a/Engine/Plugins/PCG/Source/PCG/Public/Metadata/PCGMetadataContainerTypes.h
+++ b/Engine/Plugins/PCG/Source/PCG/Public/Metadata/PCGMetadataContainerTypes.h
@@ -4,7 +4,7 @@
 
 #include "UObject/UnrealType.h"
 
-class FPCGMetadataAttributeGeneric;
+class FPCGMetadataAttributeBase;
 class IPCGAttributeAccessor;
 
 namespace PCG
@@ -93,7 +93,7 @@ public:
 	}
 	
 private:
-	friend FPCGMetadataAttributeGeneric;
+	friend FPCGMetadataAttributeBase;
 	friend IPCGAttributeAccessor;
 	FScriptSetHelper Helper;
 };
@@ -181,7 +181,7 @@ public:
 	}
 	
 private:
-	friend FPCGMetadataAttributeGeneric;
+	friend FPCGMetadataAttributeBase;
 	friend IPCGAttributeAccessor;
 	FScriptMapHelper Helper;
 };
diff --git a/Engine/Plugins/PCG/Source/PCG/Public/Metadata/PCGMetadataDomain.h b/Engine/Plugins/PCG/Source/PCG/Public/Metadata/PCGMetadataDomain.h
--- a/Engine/Plugins/PCG/Source/PCG/Public/Metadata/PCGMetadataDomain.h
+++ b/Engine/Plugins/PCG/Source/PCG/Public/Metadata/PCGMetadataDomain.h
@@ -140,13 +140,14 @@ public:
 	/** Return the number of attributes in this metadata. */
 	UE_API int32 GetAttributeCount() const;
 
-	template <typename T>
+	template <typename T> requires (PCG::Private::TIsBasicType<T>::Value == true)
 	FPCGMetadataAttribute<T>* GetMutableTypedAttribute(FName AttributeName);
 
-	template <typename T>
+	template <typename T> requires (PCG::Private::TIsBasicType<T>::Value == true)
+	UE_DEPRECATED(5.8, "Should not be used anymore.")
 	FPCGMetadataAttribute<T>* GetMutableTypedAttribute_Unsafe(FName AttributeName);
 
-	template <typename T>
+	template <typename T> requires (PCG::Private::TIsBasicType<T>::Value == true)
 	const FPCGMetadataAttribute<T>* GetConstTypedAttribute(FName AttributeName) const;
 	
 	UE_API void GetAttributes(TArray<FName>& AttributeNames, TArray<EPCGMetadataTypes>& AttributeTypes) const;
@@ -168,6 +169,9 @@ public:
 	UE_API void ClearAttribute(FName AttributeToClear);
 
 	/** Change type of an attribute */
+	UE_API bool ChangeAttributeType(FName AttributeName, const FPCGMetadataAttributeDesc& AttributeNewDesc);
+	
+	UE_DEPRECATED(5.8, "Use the version with the Attribute Desc")
 	UE_API bool ChangeAttributeType(FName AttributeName, int16 AttributeNewType);
 
 	/** Adds a unique entry key to the metadata */
@@ -229,23 +233,46 @@ public:
 	* If the attribute already exists but is of the wrong type, it will fail and return nullptr. Same if the name is invalid.
 	* Return a typed attribute pointer, of the requested type T.
 	*/
-	template<typename T>
-	FPCGMetadataAttribute<T>* CreateAttribute(FName AttributeName, const T& DefaultValue, bool bAllowsInterpolation, bool bOverrideParent);
+	template<typename T> requires (PCG::Private::TIsBasicType<T>::Value == true)
+	FPCGMetadataAttribute<T>* CreateAttribute(FName AttributeName, const T& DefaultValue, bool bAllowsInterpolation = true, bool bOverrideParent = true);
 	
-	// @todo_pcg TO BE REMOVED
 	/**
-	* Create a new generic attribute.
+	* Create a new attribute, that will be generic and not typed. This should never be static cast to a FPCGMetadataAttribute<T>* as this would
+	* be undefined behavior.
+	* If the attribute already exists, it will raise a warning (use FindOrCreateAttribute if this usecase can arise)
+	* If the attribute already exists but is of the wrong type, it will fail and return nullptr. Same if the name is invalid.
 	*/
-	template<typename T>
-	FPCGMetadataAttributeGeneric* CreateGenericAttribute(FName AttributeName, const T& DefaultValue, bool bAllowsInterpolation = true, bool bOverrideParent = false);
+	template<typename T> requires (PCG::Private::TIsBasicType<T>::Value == false)
+	FPCGMetadataAttributeBase* CreateAttribute(FName AttributeName, const T& DefaultValue, bool bAllowsInterpolation = true, bool bOverrideParent = true);
 	
 	/**
-	* Create a new generic attribute from a desc.
+	* Create a new attribute, that will be either generic or typed, depending the type given by the desc. The default value will not be set.
+	* If the attribute already exists, it will raise a warning (use FindOrCreateAttribute if this usecase can arise)
+	* If the attribute already exists but is of the wrong type, it will fail and return nullptr. Same if the name is invalid.
 	*/
-	UE_API FPCGMetadataAttributeGeneric* CreateGenericAttribute(const FPCGMetadataAttributeDesc& Desc, bool bAllowsInterpolation = true, bool bOverrideParent = false);
+	UE_API FPCGMetadataAttributeBase* CreateAttribute(const FPCGMetadataAttributeDesc& AttributeDesc, bool bAllowsInterpolation = true, bool bOverrideParent = true);
 
-	UE_API FPCGMetadataAttributeGeneric* GetMutableGenericAttribute(FName AttributeName);
-	UE_API const FPCGMetadataAttributeGeneric* GetConstGenericAttribute(FName AttributeName) const;
+	// @todo_pcg All generic methods to be removed in a follow-up CL. Will ease the refactor review.
+	template<typename T>
+	FPCGMetadataAttributeGeneric* CreateGenericAttribute(FName AttributeName, const T& DefaultValue, bool bAllowsInterpolation = true, bool bOverrideParent = false)
+	{
+		return CreateAttribute<T>(AttributeName, DefaultValue, bAllowsInterpolation, bOverrideParent);
+	}
+	
+	FPCGMetadataAttributeGeneric* CreateGenericAttribute(const FPCGMetadataAttributeDesc& Desc, bool bAllowsInterpolation = true, bool bOverrideParent = false)
+	{
+		return CreateAttribute(Desc, bAllowsInterpolation, bOverrideParent);
+	}
+	
+	FPCGMetadataAttributeGeneric* GetMutableGenericAttribute(FName AttributeName)
+	{
+		return GetMutableAttribute(AttributeName);
+	}
+	
+	const FPCGMetadataAttributeGeneric* GetConstGenericAttribute(FName AttributeName) const
+	{
+		return GetConstAttribute(AttributeName);
+	}
 
 	/**
 	* Find or create an attribute. Follows CreateAttribute signature.
@@ -253,8 +280,17 @@ public:
 	* Same as CreateAttribute, it will return nullptr if the attribute name is invalid.
 	* Return a typed attribute pointer, of the requested type T.
 	*/
-	template<typename T>
+	template<typename T> requires (PCG::Private::TIsBasicType<T>::Value == true)
 	FPCGMetadataAttribute<T>* FindOrCreateAttribute(FName AttributeName, const T& DefaultValue = T{}, bool bAllowsInterpolation = true, bool bOverrideParent = true, bool bOverwriteIfTypeMismatch = true);
+	
+	/**
+	* Find or create an attribute. Follows CreateAttribute signature.
+	* Extra boolean bOverwriteIfTypeMismatch allows to overwrite an existing attribute if the type mismatch.
+	* Same as CreateAttribute, it will return nullptr if the attribute name is invalid.
+	* Return a generic attribute pointer.
+	*/
+	template<typename T> requires (PCG::Private::TIsBasicType<T>::Value == false)
+	FPCGMetadataAttributeBase* FindOrCreateAttribute(FName AttributeName, const T& DefaultValue = T{}, bool bAllowsInterpolation = true, bool bOverrideParent = true, bool bOverwriteIfTypeMismatch = true);
 
 	/** Computes Crc from all attributes & keys from outer's data. */
 	UE_API void AddToCrc(FArchiveCrc32& Ar, const UPCGData* Data, bool bFullDataCrc) const;
@@ -262,6 +298,13 @@ public:
 	FPCGMetadataDomainID GetDomainID() const { return DomainID; }
 
 	bool SupportsMultiEntries() const { return bSupportMultiEntries; }
+	
+private:
+	UE_API FPCGMetadataAttributeBase* CreateAttributeInternal(const FPCGMetadataAttributeDesc& AttributeDesc, bool bAllowsInterpolation, bool bOverrideParent, TFunctionRef<void(FPCGMetadataAttributeBase*)> SetupDefaultValue, bool bAlreadyLocked = false);
+	UE_API FPCGMetadataAttributeBase* FindOrCreateAttributeInternal(const FPCGMetadataAttributeDesc& AttributeDesc, bool bAllowsInterpolation, bool bOverrideParent, bool bOverwriteIfTypeMismatch, TFunctionRef<void(FPCGMetadataAttributeBase*)> SetupDefaultValue);
+	
+	UE_API FPCGMetadataAttributeBase* GetMutableAttribute_Unsafe(FName AttributeName);
+	UE_API const FPCGMetadataAttributeBase* GetConstAttribute_Unsafe(FName AttributeName) const;
 
 protected:
 	UE_API FPCGMetadataAttributeBase* CopyAttribute(FName AttributeToCopy, FName NewAttributeName, bool bKeepParent, bool bCopyEntries, bool bCopyValues);
@@ -272,6 +315,8 @@ protected:
 	UE_API void RemoveAttributeInternal(FName AttributeName);
 
 	UE_API void SetLastCachedSelectorOnOwner(FName AttributeName);
+	
+	UE_API FPCGMetadataAttributeBase* AllocateAttributeFromDesc(const FPCGMetadataAttributeDesc& InDesc, const FPCGMetadataAttributeBase* InAttributeParent = nullptr);
 
 	UPCGMetadata* TopMetadata = nullptr;
 	FPCGMetadataDomainID DomainID;
@@ -293,153 +338,68 @@ protected:
 };
 
 
-template<typename T>
+template <typename T> requires (PCG::Private::TIsBasicType<T>::Value == true)
 FPCGMetadataAttribute<T>* FPCGMetadataDomain::CreateAttribute(FName AttributeName, const T& DefaultValue, bool bAllowsInterpolation, bool bOverrideParent)
 {
-	if (!FPCGMetadataAttributeBase::IsValidName(AttributeName))
-	{
-		UE_LOG(LogPCG, Error, TEXT("Attribute name '%s' is invalid"), *AttributeName.ToString());
-		return nullptr;
-	}
-
-	const FPCGMetadataAttributeBase* ParentAttribute = nullptr;
-
-	if (bOverrideParent && Parent)
-	{
-		ParentAttribute = Parent->GetConstAttribute(AttributeName);
-	}
-
-	if (ParentAttribute && (ParentAttribute->GetTypeId() != PCG::Private::MetadataTypes<T>::Id))
-	{
-		// Can't parent if the types do not match
-		ParentAttribute = nullptr;
-	}
-
-	FPCGMetadataAttribute<T>* NewAttribute = new FPCGMetadataAttribute<T>(this, AttributeName, ParentAttribute, DefaultValue, bAllowsInterpolation);
-
+	FPCGMetadataAttributeBase* Attribute = CreateAttributeInternal(PCG::Private::MakeAttributeDesc<T>(AttributeName), bAllowsInterpolation, bOverrideParent, [&DefaultValue](FPCGMetadataAttributeBase* Attribute)
 	{
-		PCG::TUniqueScopeLock WriteLock(AttributeLock);
-
-		if (FPCGMetadataAttributeBase** ExistingAttribute = Attributes.Find(AttributeName))
-		{
-			delete NewAttribute;
-			if ((*ExistingAttribute)->GetTypeId() != PCG::Private::MetadataTypes<T>::Id)
-			{
-				UE_LOG(LogPCG, Error, TEXT("Attribute %s already exists but is not the right type. Abort."), *AttributeName.ToString());
-				return nullptr;
-			}
-			else
-			{
-				UE_LOG(LogPCG, Warning, TEXT("Attribute %s already exists"), *AttributeName.ToString());
-				NewAttribute = static_cast<FPCGMetadataAttribute<T>*>(*ExistingAttribute);
-			}
-		}
-		else
-		{
-			NewAttribute->AttributeId = NextAttributeId++;
-			AddAttributeInternal(AttributeName, NewAttribute);
-
-			// Also when creating an attribute, notify the PCG Data owner that the latest attribute manipulated is this one.
-			SetLastCachedSelectorOnOwner(AttributeName);
-		}
-	}
-
-	return NewAttribute;
+		Attribute->SetDefaultValue(DefaultValue);
+	});
+	return static_cast<FPCGMetadataAttribute<T>*>(Attribute);
 }
 
-template<typename T>
-FPCGMetadataAttributeGeneric* FPCGMetadataDomain::CreateGenericAttribute(FName AttributeName, const T& DefaultValue, bool bAllowsInterpolation, bool bOverrideParent)
+template <typename T> requires (PCG::Private::TIsBasicType<T>::Value == false)
+FPCGMetadataAttributeBase* FPCGMetadataDomain::CreateAttribute(FName AttributeName, const T& DefaultValue, bool bAllowsInterpolation, bool bOverrideParent)
 {
-	FPCGMetadataAttributeGeneric* Attribute = CreateGenericAttribute(PCG::Private::MakeAttributeDesc<T>(AttributeName), bAllowsInterpolation, bOverrideParent);
-	if (Attribute)
+	return CreateAttributeInternal(PCG::Private::MakeAttributeDesc<T>(AttributeName), bAllowsInterpolation, bOverrideParent, [&DefaultValue](FPCGMetadataAttributeBase* Attribute)
 	{
-		Attribute->SetDefaultValue<T>(DefaultValue);
-	}
-
-	return Attribute;
+		Attribute->SetDefaultValue(DefaultValue);
+	});
 }
 
-template<typename T>
-FPCGMetadataAttribute<T>* FPCGMetadataDomain::GetMutableTypedAttribute_Unsafe(FName AttributeName)
+template<typename T> requires (PCG::Private::TIsBasicType<T>::Value == true)
+FPCGMetadataAttribute<T>* FPCGMetadataDomain::FindOrCreateAttribute(FName AttributeName, const T& DefaultValue, bool bAllowsInterpolation, bool bOverrideParent, bool bOverwriteIfTypeMismatch)
 {
-	FPCGMetadataAttribute<T>* Attribute = nullptr;
-
-	if (FPCGMetadataAttributeBase** FoundAttribute = Attributes.Find(AttributeName))
+	FPCGMetadataAttributeBase* Attribute = FindOrCreateAttributeInternal(PCG::Private::MakeAttributeDesc<T>(AttributeName), bAllowsInterpolation, bOverrideParent, bOverwriteIfTypeMismatch, [&DefaultValue](FPCGMetadataAttributeBase* Attribute)
 	{
-		FPCGMetadataAttributeBase* BaseAttribute = *FoundAttribute;
-		if (BaseAttribute && BaseAttribute->GetTypeId() == PCG::Private::MetadataTypes<T>::Id)
-		{
-			Attribute = static_cast<FPCGMetadataAttribute<T>*>(BaseAttribute);
-		}
-	}
-
-	return Attribute;
+		Attribute->SetDefaultValue(DefaultValue);
+	});
+	
+	return static_cast<FPCGMetadataAttribute<T>*>(Attribute);
 }
 
-template<typename T>
-FPCGMetadataAttribute<T>* FPCGMetadataDomain::FindOrCreateAttribute(FName AttributeName, const T& DefaultValue, bool bAllowsInterpolation, bool bOverrideParent, bool bOverwriteIfTypeMismatch)
+template<typename T> requires (PCG::Private::TIsBasicType<T>::Value == false)
+FPCGMetadataAttributeBase* FPCGMetadataDomain::FindOrCreateAttribute(FName AttributeName, const T& DefaultValue, bool bAllowsInterpolation, bool bOverrideParent, bool bOverwriteIfTypeMismatch)
 {
+	return FindOrCreateAttributeInternal(PCG::Private::MakeAttributeDesc<T>(AttributeName), bAllowsInterpolation, bOverrideParent, bOverwriteIfTypeMismatch, [&DefaultValue](FPCGMetadataAttributeBase* Attribute)
 	{
-		PCG::TSharedScopeLock ReadLock(AttributeLock);
-		if (FPCGMetadataAttribute<T>* Attribute = GetMutableTypedAttribute_Unsafe<T>(AttributeName))
-		{
-			return Attribute;
-		}
-	}
-
-	PCG::TUniqueScopeLock WriteLock(AttributeLock);
-	if (FPCGMetadataAttribute<T>* Attribute = GetMutableTypedAttribute_Unsafe<T>(AttributeName))
-	{
-		return Attribute;
-	}
-
-	// If an attribute with this name exists here, there is a type mismatch.
-	if (FPCGMetadataAttributeBase** FoundAttribute = Attributes.Find(AttributeName))
-	{
-		if (bOverwriteIfTypeMismatch)
-		{
-			delete *FoundAttribute;
-			RemoveAttributeInternal(AttributeName);
-		}
-		else
-		{
-			return nullptr;
-		}
-	}
-
-	// A new attribute will be created.
-	if (!FPCGMetadataAttributeBase::IsValidName(AttributeName))
-	{
-		UE_LOG(LogPCG, Error, TEXT("Attribute name '%s' is invalid"), *AttributeName.ToString());
-		return nullptr;
-	}
-
-	// Parent is const and therefore should be safe
-	const FPCGMetadataAttributeBase* ParentAttribute = (bOverrideParent && Parent) ? Parent->GetConstTypedAttribute<T>(AttributeName) : nullptr;
-	FPCGMetadataAttribute<T>* NewAttribute = new FPCGMetadataAttribute<T>(this, AttributeName, ParentAttribute, DefaultValue, bAllowsInterpolation);
-	NewAttribute->AttributeId = NextAttributeId++;
-	AddAttributeInternal(AttributeName, NewAttribute);
-
-	// Also when creating an attribute, notify the PCG Data owner that the latest attribute manipulated is this one.
-	SetLastCachedSelectorOnOwner(AttributeName);
+		Attribute->SetDefaultValue(DefaultValue);
+	});
+}
 
-	return NewAttribute;
+template<typename T> requires (PCG::Private::TIsBasicType<T>::Value == true)
+FPCGMetadataAttribute<T>* FPCGMetadataDomain::GetMutableTypedAttribute_Unsafe(FName AttributeName)
+{
+	FPCGMetadataAttributeBase* BaseAttribute = GetMutableAttribute_Unsafe(AttributeName);
+	return (BaseAttribute && (BaseAttribute->GetAttributeDesc().IsSameType(PCG::Private::GetDefaultAttributeDesc<T>())))
+		? static_cast<FPCGMetadataAttribute<T>*>(BaseAttribute)
+		: nullptr;
 }
 
-template <typename T>
+template <typename T> requires (PCG::Private::TIsBasicType<T>::Value == true)
 FPCGMetadataAttribute<T>* FPCGMetadataDomain::GetMutableTypedAttribute(FName AttributeName)
 {
 	FPCGMetadataAttributeBase* BaseAttribute = GetMutableAttribute(AttributeName);
-	return (BaseAttribute && (BaseAttribute->GetTypeId() == PCG::Private::MetadataTypes<T>::Id))
+	return (BaseAttribute && (BaseAttribute->GetAttributeDesc().IsSameType(PCG::Private::GetDefaultAttributeDesc<T>())))
 		? static_cast<FPCGMetadataAttribute<T>*>(BaseAttribute)
 		: nullptr;
 }
 
-template <typename T>
+template <typename T> requires (PCG::Private::TIsBasicType<T>::Value == true)
 const FPCGMetadataAttribute<T>* FPCGMetadataDomain::GetConstTypedAttribute(FName AttributeName) const
 {
 	const FPCGMetadataAttributeBase* BaseAttribute = GetConstAttribute(AttributeName);
-	return (BaseAttribute && (BaseAttribute->GetTypeId() == PCG::Private::MetadataTypes<T>::Id))
+	return (BaseAttribute && (BaseAttribute->GetAttributeDesc().IsSameType(PCG::Private::GetDefaultAttributeDesc<T>())))
 		? static_cast<const FPCGMetadataAttribute<T>*>(BaseAttribute)
 		: nullptr;
 }
diff --git a/Engine/Plugins/PCG/Source/PCG/Public/PCGDataVisualization.h b/Engine/Plugins/PCG/Source/PCG/Public/PCGDataVisualization.h
--- a/Engine/Plugins/PCG/Source/PCG/Public/PCGDataVisualization.h
+++ b/Engine/Plugins/PCG/Source/PCG/Public/PCGDataVisualization.h
@@ -11,7 +11,7 @@
 
 #if WITH_EDITOR
 class AActor;
-class FPCGMetadataAttributeGeneric;
+class FPCGMetadataAttributeBase;
 class UPCGData;
 class UPCGSettingsInterface;
 struct FPCGContext;
@@ -64,7 +64,7 @@ struct FPCGTableVisualizerColumnInfo
 	EPCGTableVisualizerCellAlignment CellAlignment = EPCGTableVisualizerCellAlignment::Right;
 	TSharedPtr<const IPCGAttributeAccessor> Accessor;
 	TSharedPtr<const IPCGAttributeAccessorKeys> AccessorKeys;
-	const FPCGMetadataAttributeGeneric* GenericAttribute = nullptr;
+	const FPCGMetadataAttributeBase* GenericAttribute = nullptr;
 };
 
 struct FPCGTableVisualizerInfo
diff --git a/Engine/Plugins/PCG/Source/PCGEditor/Private/DataVisualizations/PCGDataVisualizationHelpers.cpp b/Engine/Plugins/PCG/Source/PCGEditor/Private/DataVisualizations/PCGDataVisualizationHelpers.cpp
--- a/Engine/Plugins/PCG/Source/PCGEditor/Private/DataVisualizations/PCGDataVisualizationHelpers.cpp
+++ b/Engine/Plugins/PCG/Source/PCGEditor/Private/DataVisualizations/PCGDataVisualizationHelpers.cpp
@@ -166,11 +166,11 @@ namespace PCGDataVisualizationHelpers
 
 		for (int32 I = 0; I < AttributeNames.Num(); I++)
 		{
-			const FPCGMetadataAttributeGeneric* GenericAttr = MetadataDomain->GetConstGenericAttribute(AttributeNames[I]);
+			const FPCGMetadataAttributeBase* GenericAttr = MetadataDomain->GetConstAttribute(AttributeNames[I]);
 			const int32 ColumnInfosNum = OutInfo.ColumnInfos.Num();
 
 			// If we don't have a generic attribute or the generic attribute is a single value (not in a container), try to add the columns for that type. 
-			if (!GenericAttr || GenericAttr->IsSingleValue())
+			if (!GenericAttr || GenericAttr->GetAttributeDesc().IsSingleValue())
 			{
 				PCGMetadataAttribute::CallbackWithRightType(static_cast<uint16>(AttributeTypes[I]), AddInfo, AttributeNames[I]);
 			}
diff --git a/Engine/Plugins/PCG/Source/PCGEditor/Private/Widgets/SPCGMetadataAttributeGenericViewer.h b/Engine/Plugins/PCG/Source/PCGEditor/Private/Widgets/SPCGMetadataAttributeGenericViewer.h
--- a/Engine/Plugins/PCG/Source/PCGEditor/Private/Widgets/SPCGMetadataAttributeGenericViewer.h
+++ b/Engine/Plugins/PCG/Source/PCGEditor/Private/Widgets/SPCGMetadataAttributeGenericViewer.h
@@ -16,7 +16,7 @@
 #include "Widgets/Views/STableRow.h"
 
 class UPCGData;
-class FPCGMetadataAttributeGeneric;
+class FPCGMetadataAttributeBase;
 class FPCGEditor;
 class IDetailsView;
 class SPCGMetadataAttributeGenericViewer;
@@ -54,7 +54,7 @@ public:
 	void Refresh();
 
 	/** Controls lock from the editor */
-	void Setup(TWeakObjectPtr<const UPCGData> InData, const FPCGMetadataAttributeGeneric* InAttr, PCGMetadataEntryKey InEntryKey);
+	void Setup(TWeakObjectPtr<const UPCGData> InData, const FPCGMetadataAttributeBase* InAttr, PCGMetadataEntryKey InEntryKey);
 
 protected:
 	friend FPCGMetadataAttributeGenericViewerDetailCustomization;
@@ -63,7 +63,7 @@ protected:
 	FInstancedPropertyBag TempStruct;
 	
 	TWeakObjectPtr<const UPCGData> Data;
-	const FPCGMetadataAttributeGeneric* Attr;
+	const FPCGMetadataAttributeBase* Attr;
 	PCGMetadataEntryKey EntryKey;
 	TWeakPtr<FPCGEditor> EditorPtr = nullptr;
 };
\ No newline at end of file
diff --git a/Engine/Source/Runtime/Core/Public/UObject/FortniteMainBranchObjectVersions.inl b/Engine/Source/Runtime/Core/Public/UObject/FortniteMainBranchObjectVersions.inl
--- a/Engine/Source/Runtime/Core/Public/UObject/FortniteMainBranchObjectVersions.inl
+++ b/Engine/Source/Runtime/Core/Public/UObject/FortniteMainBranchObjectVersions.inl
@@ -811,5 +811,8 @@ FORTNITE_MAIN_VERSION(PCGCustomHLSLAttributeCreationSettingsToggle, 265)
 // Converting FPCGMetadataAttribute<T> to Generic attributes. 
 FORTNITE_MAIN_VERSION(ConvertFPCGMetadataAttributeToGenericAttributes, 266)
 
+// Merging PCG Metadata Attribute base and Generic ones. 
+FORTNITE_MAIN_VERSION(MergePCGMetadataAttributeBaseAndGeneric, 267)
+
 // -----<new versions can be added above this line>-------------------------------------------------
 #undef FORTNITE_MAIN_VERSION
