[PCG] Avoid compiling inline RT shader when bindless support is not available on the target platform

#rnx
#Synced-CL 52429529

diff --git a/Engine/Plugins/PCG/Source/PCGCompute/Private/PCGRayTrace.cpp b/Engine/Plugins/PCG/Source/PCGCompute/Private/PCGRayTrace.cpp
--- a/Engine/Plugins/PCG/Source/PCGCompute/Private/PCGRayTrace.cpp
+++ b/Engine/Plugins/PCG/Source/PCGCompute/Private/PCGRayTrace.cpp
@@ -49,7 +49,8 @@ class FPCGRayTraceCS : public FGlobalShader
 
 	static bool ShouldCompilePermutation(const FGlobalShaderPermutationParameters& Parameters)
 	{
-		return IsRayTracingEnabledForProject(Parameters.Platform);
+		return IsRayTracingEnabledForProject(Parameters.Platform)
+			&& IsBindlessEnabledForAnyGraphics(FShaderPlatformConfig::GetBindlessConfiguration(Parameters.Platform));
 	}
 
 	static void ModifyCompilationEnvironment(const FGlobalShaderPermutationParameters& Parameters, FShaderCompilerEnvironment& OutEnvironment)
@@ -85,6 +86,12 @@ void PCGRayTrace::RenderPCGRayTraceInline(FRDGBuilder& GraphBuilder, const FPCGR
 		return;
 	}
 
+	if (!IsBindlessEnabledForAnyGraphics(FShaderPlatformConfig::GetBindlessConfiguration(GMaxRHIShaderPlatform)))
+	{
+		UE_LOG(LogPCGCompute, Error, TEXT("Failed to submit inline raytrace shader. Bindless rendering is not enabled for platform '%s'."), *LexToString(GMaxRHIShaderPlatform));
+		return;
+	}
+
 	if (!InParams.Scene)
 	{
 		UE_LOG(LogPCGCompute, Error, TEXT("Failed to submit inline raytrace shader. Invalid scene."));
