Added TIsTMap_V and TIsTSet_V traits.

#Synced-CL 52474350

diff --git a/Engine/Plugins/PCG/Source/PCG/Public/Metadata/PCGMetadataAttributeTraits.h b/Engine/Plugins/PCG/Source/PCG/Public/Metadata/PCGMetadataAttributeTraits.h
--- a/Engine/Plugins/PCG/Source/PCG/Public/Metadata/PCGMetadataAttributeTraits.h
+++ b/Engine/Plugins/PCG/Source/PCG/Public/Metadata/PCGMetadataAttributeTraits.h
@@ -210,7 +210,7 @@ namespace PCG
 			return TypedDesc;
 		}
 
-		template <typename Container> requires (TIsTSet<Container>::Value == true)
+		template <typename Container> requires (TIsTSet_V<Container>)
 		const FPCGMetadataAttributeDesc& GetDefaultAttributeDesc()
 		{
 			using T = typename Container::ElementType;
@@ -225,7 +225,7 @@ namespace PCG
 			return TypedDesc;
 		}
 
-		template<typename Container> requires (TIsTMap<Container>::Value == true)
+		template<typename Container> requires (TIsTMap_V<Container>)
 		const FPCGMetadataAttributeDesc& GetDefaultAttributeDesc()
 		{
 			using KeyType = typename TMapTypeExtractor<Container>::KeyType;
@@ -293,7 +293,7 @@ namespace PCG
 				&& IsSameType<T>(Desc.ValueType, Desc.ValueTypeObject);
 		}
 
-		template <typename Container> requires (TIsTSet<Container>::Value == true)
+		template <typename Container> requires (TIsTSet_V<Container>)
 		bool IsEquivalentDesc(const FPCGMetadataAttributeDesc& Desc)
 		{
 			using T = typename Container::ElementType;
@@ -302,7 +302,7 @@ namespace PCG
 				&& IsSameType<T>(Desc.ValueType, Desc.ValueTypeObject);
 		}
 
-		template<typename Container> requires (TIsTMap<Container>::Value == true)
+		template<typename Container> requires (TIsTMap_V<Container>)
 		bool IsEquivalentDesc(const FPCGMetadataAttributeDesc& Desc)
 		{
 			using KeyType = typename TMapTypeExtractor<Container>::KeyType;
diff --git a/Engine/Source/Developer/CQTest/Public/ObjectBuilder.h b/Engine/Source/Developer/CQTest/Public/ObjectBuilder.h
--- a/Engine/Source/Developer/CQTest/Public/ObjectBuilder.h
+++ b/Engine/Source/Developer/CQTest/Public/ObjectBuilder.h
@@ -73,11 +73,11 @@ public:
 		{
 			TrySet(FindFProperty<FArrayProperty>(Object->GetClass(), InPropertyName));
 		}
-		else if constexpr (TIsTMap<Raw>::Value)
+		else if constexpr (TIsTMap_V<Raw>)
 		{
 			TrySet(FindFProperty<FMapProperty>(Object->GetClass(), InPropertyName));
 		}
-		else if constexpr (TIsTSet<Raw>::Value)
+		else if constexpr (TIsTSet_V<Raw>)
 		{
 			TrySet(FindFProperty<FSetProperty>(Object->GetClass(), InPropertyName));
 		}
@@ -194,14 +194,14 @@ private:
 				return IsCompatibleWith<typename Raw::ElementType>(*AsProperty->Inner);
 			}
 		}
-		else if constexpr (TIsTSet<Raw>::Value)
+		else if constexpr (TIsTSet_V<Raw>)
 		{
 			if (auto AsProperty = CastField<FSetProperty>(&Prop))
 			{
 				return IsCompatibleWith<typename Raw::ElementType>(*AsProperty->ElementProp);
 			}
 		}
-		else if constexpr (TIsTMap<Raw>::Value)
+		else if constexpr (TIsTMap_V<Raw>)
 		{
 			if (auto AsProperty = CastField<FMapProperty>(&Prop))
 			{
diff --git a/Engine/Source/Runtime/Core/Public/Containers/CompactSet.h.inl b/Engine/Source/Runtime/Core/Public/Containers/CompactSet.h.inl
--- a/Engine/Source/Runtime/Core/Public/Containers/CompactSet.h.inl
+++ b/Engine/Source/Runtime/Core/Public/Containers/CompactSet.h.inl
@@ -1812,10 +1812,8 @@ template<typename ElementType, typename KeyFuncs, typename Allocator>
 	return !TSETPRIVATEFRIEND::LegacyCompareEqual(A, B);
 }
 
-template <typename ElementType, typename KeyFuncs, typename Allocator> struct TIsTSet<               UE_TCOMPACT_SET<ElementType, KeyFuncs, Allocator>> { enum { Value = true }; };
-template <typename ElementType, typename KeyFuncs, typename Allocator> struct TIsTSet<const          UE_TCOMPACT_SET<ElementType, KeyFuncs, Allocator>> { enum { Value = true }; };
-template <typename ElementType, typename KeyFuncs, typename Allocator> struct TIsTSet<      volatile UE_TCOMPACT_SET<ElementType, KeyFuncs, Allocator>> { enum { Value = true }; };
-template <typename ElementType, typename KeyFuncs, typename Allocator> struct TIsTSet<const volatile UE_TCOMPACT_SET<ElementType, KeyFuncs, Allocator>> { enum { Value = true }; };
+template <typename ElementType, typename KeyFuncs, typename Allocator>
+inline constexpr bool TIsTSet_V<UE_TCOMPACT_SET<ElementType, KeyFuncs, Allocator>> = true;
 
 #undef TSETPRIVATEFRIEND
 #undef UE_TCOMPACT_SET_ARRAY_VIEW_ENABLED
\ No newline at end of file
diff --git a/Engine/Source/Runtime/Core/Public/Containers/Map.h b/Engine/Source/Runtime/Core/Public/Containers/Map.h
--- a/Engine/Source/Runtime/Core/Public/Containers/Map.h
+++ b/Engine/Source/Runtime/Core/Public/Containers/Map.h
@@ -121,10 +121,21 @@ struct TDefaultMapHashableKeyFuncs : TDefaultMapKeyFuncs<KeyType, ValueType, bIn
 template <typename AllocatorType, typename InDerivedType = void>
 class TScriptMap;
 
-	/**
-* Traits class which determines whether or not a type is a TMAP.
-	 */
-template <typename T> struct TIsTMap { enum { Value = false }; };
+/**
+ * Traits which determine whether or not a type is a TMap.
+ */
+template <typename T>
+inline constexpr bool TIsTMap_V = false;
+
+template <typename T> inline constexpr bool TIsTMap_V<const          T> = TIsTMap_V<T>;
+template <typename T> inline constexpr bool TIsTMap_V<      volatile T> = TIsTMap_V<T>;
+template <typename T> inline constexpr bool TIsTMap_V<const volatile T> = TIsTMap_V<T>;
+
+template <typename T>
+struct TIsTMap
+{
+	static inline constexpr bool Value = TIsTMap_V<T>;
+};
 
 #define UE_TMAP_PREFIX T
 #include "Map.h.inl"
diff --git a/Engine/Source/Runtime/Core/Public/Containers/Map.h.inl b/Engine/Source/Runtime/Core/Public/Containers/Map.h.inl
--- a/Engine/Source/Runtime/Core/Public/Containers/Map.h.inl
+++ b/Engine/Source/Runtime/Core/Public/Containers/Map.h.inl
@@ -1777,10 +1777,8 @@ namespace Freeze
 
 DECLARE_TEMPLATE_INTRINSIC_TYPE_LAYOUT((template <typename KeyType, typename ValueType, typename SetAllocator, typename KeyFuncs>), (TMULTIMAP<KeyType, ValueType, SetAllocator, KeyFuncs>));
 
-template <typename KeyType, typename ValueType, typename SetAllocator, typename KeyFuncs> struct TIsTMap<               TMAP<KeyType, ValueType, SetAllocator, KeyFuncs>> { enum { Value = true }; };
-template <typename KeyType, typename ValueType, typename SetAllocator, typename KeyFuncs> struct TIsTMap<const          TMAP<KeyType, ValueType, SetAllocator, KeyFuncs>> { enum { Value = true }; };
-template <typename KeyType, typename ValueType, typename SetAllocator, typename KeyFuncs> struct TIsTMap<      volatile TMAP<KeyType, ValueType, SetAllocator, KeyFuncs>> { enum { Value = true }; };
-template <typename KeyType, typename ValueType, typename SetAllocator, typename KeyFuncs> struct TIsTMap<const volatile TMAP<KeyType, ValueType, SetAllocator, KeyFuncs>> { enum { Value = true }; };
+template <typename KeyType, typename ValueType, typename SetAllocator, typename KeyFuncs>
+inline constexpr bool TIsTMap_V<TMAP<KeyType, ValueType, SetAllocator, KeyFuncs>> = true;
 
 struct TMAPPRIVATEFRIEND
 {
diff --git a/Engine/Source/Runtime/Core/Public/Containers/SetUtilities.h b/Engine/Source/Runtime/Core/Public/Containers/SetUtilities.h
--- a/Engine/Source/Runtime/Core/Public/Containers/SetUtilities.h
+++ b/Engine/Source/Runtime/Core/Public/Containers/SetUtilities.h
@@ -14,9 +14,20 @@
 #define UE_SET_UNHASH_THRESHOLD_CALCULATE_KEY 32
 
 /**
- * Traits class which determines whether or not a type is a TSet.
+ * Traits which determine whether or not a type is a TSet.
  */
-template <typename T> struct TIsTSet { enum { Value = false }; };
+template <typename T>
+inline constexpr bool TIsTSet_V = false;
+
+template <typename T> inline constexpr bool TIsTSet_V<const          T> = TIsTSet_V<T>;
+template <typename T> inline constexpr bool TIsTSet_V<      volatile T> = TIsTSet_V<T>;
+template <typename T> inline constexpr bool TIsTSet_V<const volatile T> = TIsTSet_V<T>;
+
+template <typename T>
+struct TIsTSet
+{
+	static inline constexpr bool Value = TIsTSet_V<T>;
+};
 
 /**
  * The base KeyFuncs type with some useful definitions for all KeyFuncs; meant to be derived from instead of used directly.
diff --git a/Engine/Source/Runtime/Core/Public/Containers/SparseSet.h.inl b/Engine/Source/Runtime/Core/Public/Containers/SparseSet.h.inl
--- a/Engine/Source/Runtime/Core/Public/Containers/SparseSet.h.inl
+++ b/Engine/Source/Runtime/Core/Public/Containers/SparseSet.h.inl
@@ -1849,9 +1849,7 @@ template<typename ElementType, typename KeyFuncs,typename Allocator>
 	return !TSETPRIVATEFRIEND::LegacyCompareEqual(A, B);
 }
 
-template <typename ElementType, typename KeyFuncs, typename Allocator> struct TIsTSet<               UE_TSPARSE_SET<ElementType, KeyFuncs, Allocator>> { enum { Value = true }; };
-template <typename ElementType, typename KeyFuncs, typename Allocator> struct TIsTSet<const          UE_TSPARSE_SET<ElementType, KeyFuncs, Allocator>> { enum { Value = true }; };
-template <typename ElementType, typename KeyFuncs, typename Allocator> struct TIsTSet<      volatile UE_TSPARSE_SET<ElementType, KeyFuncs, Allocator>> { enum { Value = true }; };
-template <typename ElementType, typename KeyFuncs, typename Allocator> struct TIsTSet<const volatile UE_TSPARSE_SET<ElementType, KeyFuncs, Allocator>> { enum { Value = true }; };
+template <typename ElementType, typename KeyFuncs, typename Allocator>
+inline constexpr bool TIsTSet_V<UE_TSPARSE_SET<ElementType, KeyFuncs, Allocator>> = true;
 
 #undef TSETPRIVATEFRIEND
