[ContentBrowser] New Add Menu Procedural Menu

#Synced-CL 52439714

diff --git a/Engine/Plugins/PCG/Source/PCGEditor/Private/AssetDefinitions/AssetDefinition_PCGGraphInterface.cpp b/Engine/Plugins/PCG/Source/PCGEditor/Private/AssetDefinitions/AssetDefinition_PCGGraphInterface.cpp
--- a/Engine/Plugins/PCG/Source/PCGEditor/Private/AssetDefinitions/AssetDefinition_PCGGraphInterface.cpp
+++ b/Engine/Plugins/PCG/Source/PCGEditor/Private/AssetDefinitions/AssetDefinition_PCGGraphInterface.cpp
@@ -49,7 +49,7 @@ TSoftClassPtr<UObject> UAssetDefinition_PCGGraphInterface::GetAssetClass() const
 
 TConstArrayView<FAssetCategoryPath> UAssetDefinition_PCGGraphInterface::GetAssetCategories() const
 {
-	static const FAssetCategoryPath Categories[] = { FPCGEditorCommon::PCGAssetCategoryPath };
+	static const FAssetCategoryPath Categories[] = { FPCGEditorCommon::PCGAssetSubCategoryPath };
 	return Categories;
 }
 
@@ -220,7 +220,7 @@ namespace MenuExtension_PCGGraphInterface
 				LevelGraphs = GetGraphsByType(EPCGGraphUsage::Level);
 
 				// When exactly one asset is selected, check whether it has parameters to offer the configure option.
-				// TODO: Loading the asset at menu-build time to check for parameters — revisit if this proves too slow for large unloaded assets.
+				// TODO: Loading the asset at menu-build time to check for parameters - revisit if this proves too slow for large unloaded assets.
 				const bool bShowConfigureInstance = CBContext->SelectedAssets.Num() == 1 && (!AssetGraphs.IsEmpty() || !LevelGraphs.IsEmpty()) && HasGraphParameters(Cast<UPCGGraphInterface>(CBContext->SelectedAssets[0].GetAsset()));
 				const FSlateIcon Icon = FSlateIcon(FPCGEditorStyle::Get().GetStyleSetName(), "ClassIcon.PCGGraph");
 
diff --git a/Engine/Plugins/PCG/Source/PCGEditor/Public/PCGEditorCommon.h b/Engine/Plugins/PCG/Source/PCGEditor/Public/PCGEditorCommon.h
--- a/Engine/Plugins/PCG/Source/PCGEditor/Public/PCGEditorCommon.h
+++ b/Engine/Plugins/PCG/Source/PCGEditor/Public/PCGEditorCommon.h
@@ -54,7 +54,8 @@ namespace FPCGEditorCommon
 
 	extern TAutoConsoleVariable<bool> CVarShowAdvancedAttributesFields;
 
-	const FAssetCategoryPath PCGAssetCategoryPath(LOCTEXT("PCGCategory", "PCG"));
+	const FAssetCategoryPath PCGAssetCategoryPath(LOCTEXT("PCGCategory", "Procedural"));
+	const FAssetCategoryPath PCGAssetSubCategoryPath(FPCGEditorCommon::PCGAssetCategoryPath, LOCTEXT("PCGSubCategory", "PCG"), ECategoryMenuType::Section);
 	const FAssetCategoryPath PCGAdvancedAssetCategoryPath(FPCGEditorCommon::PCGAssetCategoryPath, LOCTEXT("PCGAvancedCategory", "Advanced"));
 
 	namespace Helpers
