[PCG] Cleanup RuntimeGen state on network replay scrub

#rnx
#Synced-CL 52438074

diff --git a/Engine/Plugins/PCG/Source/PCG/Private/RuntimeGen/PCGRuntimeGenScheduler.cpp b/Engine/Plugins/PCG/Source/PCG/Private/RuntimeGen/PCGRuntimeGenScheduler.cpp
--- a/Engine/Plugins/PCG/Source/PCG/Private/RuntimeGen/PCGRuntimeGenScheduler.cpp
+++ b/Engine/Plugins/PCG/Source/PCG/Private/RuntimeGen/PCGRuntimeGenScheduler.cpp
@@ -26,6 +26,7 @@
 #include "Engine/Engine.h"
 #include "Engine/LevelStreaming.h"
 #include "Components/RuntimeVirtualTextureComponent.h"
+#include "Net/UnrealNetwork.h"
 #include "Streaming/LevelStreamingDelegates.h"
 #include "UObject/UObjectGlobals.h"
 #include "VT/RuntimeVirtualTexture.h"
@@ -269,6 +270,8 @@ FPCGRuntimeGenScheduler::FPCGRuntimeGenScheduler(UWorld* InWorld, FPCGTrackingMa
 	RefreshContext();
 
 	FLevelStreamingDelegates::OnLevelStreamingStateChanged.AddRaw(this, &FPCGRuntimeGenScheduler::OnLevelStreamingStateChanged);
+	FNetworkReplayDelegates::OnPreScrub.AddRaw(this, &FPCGRuntimeGenScheduler::OnNetworkReplayScrub);
+	FNetworkReplayDelegates::OnReplayScrubComplete.AddRaw(this, &FPCGRuntimeGenScheduler::OnNetworkReplayScrubComplete);
 
 #if WITH_EDITOR
 	// Handle UWorld::ReInitWorld
@@ -310,6 +313,8 @@ FPCGRuntimeGenScheduler::~FPCGRuntimeGenScheduler()
 	GenSourceManager = nullptr;
 
 	FLevelStreamingDelegates::OnLevelStreamingStateChanged.RemoveAll(this);
+	FNetworkReplayDelegates::OnPreScrub.RemoveAll(this);
+	FNetworkReplayDelegates::OnReplayScrubComplete.RemoveAll(this);
 }
 
 void FPCGRuntimeGenScheduler::Tick(APCGWorldActor* InPCGWorldActor, double InDeltaTime, double InEndTime)
@@ -421,6 +426,11 @@ bool FPCGRuntimeGenScheduler::ShouldTick(double InDeltaTime)
 		return false;
 	}
 
+	if (bSuspendedForReplayScrub)
+	{
+		return false;
+	}
+
 	TimeToNextTick -= InDeltaTime;
 	if (TimeToNextTick > 0.0)
 	{
@@ -1781,6 +1791,42 @@ void FPCGRuntimeGenScheduler::OnLevelStreamingStateChanged(UWorld* InWorld, cons
 	}
 }
 
+void FPCGRuntimeGenScheduler::OnNetworkReplayScrub(UWorld* InWorld)
+{
+	if (Context.World != InWorld || !Subsystem)
+	{
+		return;
+	}
+
+	bSuspendedForReplayScrub = true;
+
+	UE_LOG(LogPCG, Verbose, TEXT("RuntimeGen scheduler suspended for network replay scrub."));
+
+	// During a network replay scrub, all partition actors will be destroyed. Walk the proper unregistration
+	// path for every runtime-gen execution source so that the scheduler's internal state is fully cleaned up.
+	for (IPCGGraphExecutionSource* ExecutionSource : Subsystem->GetAllRegisteredExecutionSources())
+	{
+		if (ExecutionSource && ExecutionSource->GetExecutionState().IsManagedByRuntimeGenSystem())
+		{
+			Subsystem->OnOriginalExecutionSourceUnregistered(ExecutionSource);
+		}
+	}
+
+	// The scrub will destroy all transient actors including pooled PAs. Clear our references so we don't hold stale pointers.
+	// The pool will be re-created on the next tick after the scrub completes. 
+	ResetPartitionActorPoolToSize(0);
+}
+
+void FPCGRuntimeGenScheduler::OnNetworkReplayScrubComplete(UWorld* InWorld)
+{
+	if (Context.World == InWorld)
+	{
+		bSuspendedForReplayScrub = false;
+	
+		UE_LOG(LogPCG, Verbose, TEXT("RuntimeGen scheduler resumed after network replay scrub."));
+	}
+}
+
 void FPCGRuntimeGenScheduler::CleanupLocalSource(APCGPartitionActor* InPartitionActor, IPCGGraphExecutionSource* InLocalSource)
 {
 	if (InLocalSource && InLocalSource->GetExecutionState().IsLocalSource())
diff --git a/Engine/Plugins/PCG/Source/PCG/Public/RuntimeGen/PCGRuntimeGenScheduler.h b/Engine/Plugins/PCG/Source/PCG/Public/RuntimeGen/PCGRuntimeGenScheduler.h
--- a/Engine/Plugins/PCG/Source/PCG/Public/RuntimeGen/PCGRuntimeGenScheduler.h
+++ b/Engine/Plugins/PCG/Source/PCG/Public/RuntimeGen/PCGRuntimeGenScheduler.h
@@ -224,6 +224,12 @@ private:
 	/** Called on world streaming events. */
 	void OnLevelStreamingStateChanged(UWorld* InWorld, const ULevelStreaming* InLevelStreaming, ULevel* InLevelIfLoaded, ELevelStreamingState InPreviousState, ELevelStreamingState InNewState);
 
+	/** Called before network replay scrub. Unregisters all runtime-gen execution sources and suspends ticking. */
+	void OnNetworkReplayScrub(UWorld* InWorld);
+
+	/** Called after network replay scrub completes. Resumes ticking so sources can re-register naturally. */
+	void OnNetworkReplayScrubComplete(UWorld* InWorld);
+
 	/** Cleanup a local execution source along with its partition actor if it exists. */
 	void CleanupLocalSource(APCGPartitionActor* InPartitionActor, IPCGGraphExecutionSource* InLocalSource);
 
@@ -287,6 +293,9 @@ private:
 	/** Requests to flush all actors are deferred so they can be handled at a known time during tick. */
 	bool bActorFlushRequested = false;
 
+	/** True while a network replay scrub is in progress. Prevents ticking until scrub completes. */
+	bool bSuspendedForReplayScrub = false;
+
 	/** Setting up a PA calls APCGPartitionActor::AddGraphInstance which later calls RefreshComponent, which can create
 	* an infinite refresh loop. To break this loop we write the OC pointer to this variable, and if Refresh gets called for
 	* this OC we early out. Basically we don't respond to refresh calls for a execution source we are midway through setting up.
