[PCG] Fix Custom HLSL overridden properties not appearing as overridden

Also mark NumElements pins as unused when kernel is a processor type

#Synced-CL 52429529
#rnx
#Synced-CL 52429529

diff --git a/Engine/Plugins/PCG/Source/PCG/Private/Compute/Elements/PCGCustomHLSL.cpp b/Engine/Plugins/PCG/Source/PCG/Private/Compute/Elements/PCGCustomHLSL.cpp
--- a/Engine/Plugins/PCG/Source/PCG/Private/Compute/Elements/PCGCustomHLSL.cpp
+++ b/Engine/Plugins/PCG/Source/PCG/Private/Compute/Elements/PCGCustomHLSL.cpp
@@ -243,6 +243,29 @@ EPCGChangeType UPCGCustomHLSLSettings::GetChangeTypeForProperty(const FName& InP
 }
 #endif
 
+bool UPCGCustomHLSLSettings::IsPinUsedByNodeExecution(const UPCGPin* InPin) const
+{
+	if (!Super::IsPinUsedByNodeExecution(InPin))
+	{
+		return false;
+	}
+
+	// NumElements override pins are not used by processor kernels.
+	if (InPin)
+	{
+		const FName PinLabel = InPin->Properties.Label;
+		if (IsProcessorKernel()
+			&& (PinLabel == GET_MEMBER_NAME_CHECKED(UPCGCustomHLSLSettings, NumElements)
+			|| PinLabel == GET_MEMBER_NAME_CHECKED(UPCGCustomHLSLSettings, NumElementsX)
+			|| PinLabel == GET_MEMBER_NAME_CHECKED(UPCGCustomHLSLSettings, NumElementsY)))
+		{
+			return false;
+		}
+	}
+
+	return true;
+}
+
 void UPCGCustomHLSLSettings::ApplyPreconfiguredSettings(const FPCGPreConfiguredSettingsInfo& PreconfiguredInfo)
 {
 	if (const UEnum* EnumPtr = StaticEnum<EPCGKernelType>())
diff --git a/Engine/Plugins/PCG/Source/PCG/Private/Compute/Elements/PCGCustomHLSL.h b/Engine/Plugins/PCG/Source/PCG/Private/Compute/Elements/PCGCustomHLSL.h
--- a/Engine/Plugins/PCG/Source/PCG/Private/Compute/Elements/PCGCustomHLSL.h
+++ b/Engine/Plugins/PCG/Source/PCG/Private/Compute/Elements/PCGCustomHLSL.h
@@ -81,6 +81,7 @@ protected:
 	virtual void PostEditChangeProperty(FPropertyChangedEvent& PropertyChangedEvent) override;
 	virtual EPCGChangeType GetChangeTypeForProperty(const FName& InPropertyName) const override;
 #endif
+	virtual bool IsPinUsedByNodeExecution(const UPCGPin* InPin) const override;
 	virtual FPCGElementPtr CreateElement() const override;
 	//~End UPCGSettings interface
 
diff --git a/Engine/Plugins/PCG/Source/PCG/Private/PCGSettings.cpp b/Engine/Plugins/PCG/Source/PCG/Private/PCGSettings.cpp
--- a/Engine/Plugins/PCG/Source/PCG/Private/PCGSettings.cpp
+++ b/Engine/Plugins/PCG/Source/PCG/Private/PCGSettings.cpp
@@ -794,7 +794,11 @@ bool UPCGSettings::CanEditChange(const FEditPropertyChain& InPropertyChain) cons
 	}
 
 	// If it is not marked edit const, properties that are not marked overridable or child properties explicitly marked non overridable are always editable.
-	if (!MemberProperty->HasMetaData(PCGObjectMetadata::Overridable) || Property->HasMetaData(PCGObjectMetadata::NotOverridable))
+	const bool bIsOverridable = MemberProperty->HasMetaData(PCGObjectMetadata::Overridable)
+		|| MemberProperty->HasMetaData(PCGObjectMetadata::OverridableCPUAndGPU)
+		|| MemberProperty->HasMetaData(PCGObjectMetadata::OverridableCPUAndGPUWithReadback);
+
+	if (!bIsOverridable || Property->HasMetaData(PCGObjectMetadata::NotOverridable))
 	{
 		return true;
 	}
