GPU Test Infrastructure - Isolated world, parity runner, comparison helpers

- FPCGGPUGraphTestRunner now creates an isolated UWorld instead of using the editor world
- Add RunCPUGPUParityTest for CPU vs GPU output comparison
- Add CompareBasePointData, CompareDataCollections, ComparePointDataCollections, CompareMetadata
- Add CreateFlatLandscape/CreatePerlinLandscape helpers
- Add UPCGSettings::SetExecuteOnGPU setter

#rnx
#Synced-CL 52429529

diff --git a/Engine/Plugins/PCG/Source/PCG/Private/Tests/Compute/PCGCopyPointsGPUTest.cpp b/Engine/Plugins/PCG/Source/PCG/Private/Tests/Compute/PCGCopyPointsGPUTest.cpp
new file mode 100644
--- /dev/null
+++ b/Engine/Plugins/PCG/Source/PCG/Private/Tests/Compute/PCGCopyPointsGPUTest.cpp
@@ -0,0 +1,264 @@
+// Copyright Epic Games, Inc. All Rights Reserved.
+
+#if WITH_EDITOR
+
+#include "Tests/Compute/PCGGPUTestCommon.h"
+
+#include "PCGCommon.h"
+#include "PCGGraph.h"
+#include "Data/PCGBasePointData.h"
+#include "Elements/PCGCopyPoints.h"
+#include "Elements/PCGGather.h"
+
+namespace PCGCopyPointsGPUTestHelpers
+{
+	static const TArray<FTransform> SourceTransforms = {
+		FTransform(FQuat::MakeFromEuler(FVector(10.0, 20.0, 30.0)), FVector(100.0, 200.0, 300.0), FVector(1.0, 2.0, 0.5)),
+		FTransform(FQuat::MakeFromEuler(FVector(45.0, 0.0, 90.0)), FVector(-50.0, 75.0, 150.0), FVector(0.5, 1.0, 1.5)),
+		FTransform(FQuat::Identity, FVector(0.0, 0.0, 0.0), FVector(1.0, 1.0, 1.0)),
+	};
+
+	static const TArray<FTransform> TargetTransforms = {
+		FTransform(FQuat::MakeFromEuler(FVector(0.0, 0.0, 45.0)), FVector(500.0, 0.0, 0.0), FVector(2.0, 2.0, 2.0)),
+		FTransform(FQuat::MakeFromEuler(FVector(90.0, 0.0, 0.0)), FVector(0.0, -300.0, 100.0), FVector(1.0, 0.5, 3.0)),
+	};
+
+	struct FCopyPointsGraphSetup
+	{
+		UPCGNode* GatherNode = nullptr;
+		UPCGCopyPointsSettings* CopySettings = nullptr;
+		UPCGNode* CopyNode = nullptr;
+	};
+
+	FCopyPointsGraphSetup BuildCopyPointsGraph(UPCGGraph* Graph, UPCGNode* SourceNode, UPCGNode* TargetNode)
+	{
+		FCopyPointsGraphSetup Setup;
+		Setup.CopyNode = Graph->AddNodeOfType<UPCGCopyPointsSettings>(Setup.CopySettings);
+
+		UPCGGatherSettings* GatherSettings = nullptr;
+		Setup.GatherNode = Graph->AddNodeOfType<UPCGGatherSettings>(GatherSettings);
+
+		Graph->AddEdge(SourceNode, PCGPinConstants::DefaultOutputLabel, Setup.CopyNode, PCGCopyPointsConstants::SourcePointsLabel);
+		Graph->AddEdge(TargetNode, PCGPinConstants::DefaultOutputLabel, Setup.CopyNode, PCGCopyPointsConstants::TargetPointsLabel);
+		Graph->AddEdge(Setup.CopyNode, PCGPinConstants::DefaultOutputLabel, Setup.GatherNode, PCGPinConstants::DefaultInputLabel);
+
+		return Setup;
+	}
+}
+
+/**
+ * GPU test: CopyPoints CPU/GPU parity with default settings (all Relative).
+ *
+ *   CreatePoints(Source) -\
+ *                          -> CopyPoints -> Gather
+ *   CreatePoints(Target) -/
+ */
+IMPLEMENT_CUSTOM_SIMPLE_AUTOMATION_TEST(
+	FPCGCopyPointsGPU_DefaultTest,
+	FPCGTestBaseClass,
+	"Plugins.PCG.Compute.CopyPoints.Default",
+	PCGTestsCommon::TestFlags | EAutomationTestFlags::NonNullRHI)
+
+bool FPCGCopyPointsGPU_DefaultTest::RunTest(const FString& Parameters)
+{
+	using namespace PCGCopyPointsGPUTestHelpers;
+
+	return PCGGPUTestCommon::RunCPUGPUParityTest(this,
+		[](UPCGGraph* Graph, UWorld* /*World*/) -> PCGGPUTestCommon::FParityBuildResult
+		{
+			UPCGNode* SourceNode = PCGTestsCommon::AddCreatePointsNode(Graph, SourceTransforms);
+			UPCGNode* TargetNode = PCGTestsCommon::AddCreatePointsNode(Graph, TargetTransforms);
+			FCopyPointsGraphSetup Setup = BuildCopyPointsGraph(Graph, SourceNode, TargetNode);
+			return { Setup.GatherNode, { Setup.CopySettings } };
+		},
+		[](FAutomationTestBase* Test, const FPCGDataCollection& CPUOutput, const FPCGDataCollection& GPUOutput) -> bool
+		{
+			return PCGTestsCommon::ComparePointDataCollections(Test, CPUOutput, GPUOutput);
+		});
+}
+
+/**
+ * GPU test: CopyPoints with all inheritance set to Source.
+ */
+IMPLEMENT_CUSTOM_SIMPLE_AUTOMATION_TEST(
+	FPCGCopyPointsGPU_AllSourceTest,
+	FPCGTestBaseClass,
+	"Plugins.PCG.Compute.CopyPoints.AllSource",
+	PCGTestsCommon::TestFlags | EAutomationTestFlags::NonNullRHI)
+
+bool FPCGCopyPointsGPU_AllSourceTest::RunTest(const FString& Parameters)
+{
+	using namespace PCGCopyPointsGPUTestHelpers;
+
+	return PCGGPUTestCommon::RunCPUGPUParityTest(this,
+		[](UPCGGraph* Graph, UWorld* /*World*/) -> PCGGPUTestCommon::FParityBuildResult
+		{
+			UPCGNode* SourceNode = PCGTestsCommon::AddCreatePointsNode(Graph, SourceTransforms);
+			UPCGNode* TargetNode = PCGTestsCommon::AddCreatePointsNode(Graph, TargetTransforms);
+			FCopyPointsGraphSetup Setup = BuildCopyPointsGraph(Graph, SourceNode, TargetNode);
+
+			Setup.CopySettings->RotationInheritance = EPCGCopyPointsInheritanceMode::Source;
+			Setup.CopySettings->ScaleInheritance = EPCGCopyPointsInheritanceMode::Source;
+			Setup.CopySettings->ColorInheritance = EPCGCopyPointsInheritanceMode::Source;
+			Setup.CopySettings->SeedInheritance = EPCGCopyPointsInheritanceMode::Source;
+			Setup.CopySettings->AttributeInheritance = EPCGCopyPointsMetadataInheritanceMode::SourceOnly;
+
+			return { Setup.GatherNode, { Setup.CopySettings } };
+		},
+		[](FAutomationTestBase* Test, const FPCGDataCollection& CPUOutput, const FPCGDataCollection& GPUOutput) -> bool
+		{
+			return PCGTestsCommon::ComparePointDataCollections(Test, CPUOutput, GPUOutput);
+		});
+}
+
+/**
+ * GPU test: CopyPoints with all inheritance set to Target.
+ */
+IMPLEMENT_CUSTOM_SIMPLE_AUTOMATION_TEST(
+	FPCGCopyPointsGPU_AllTargetTest,
+	FPCGTestBaseClass,
+	"Plugins.PCG.Compute.CopyPoints.AllTarget",
+	PCGTestsCommon::TestFlags | EAutomationTestFlags::NonNullRHI)
+
+bool FPCGCopyPointsGPU_AllTargetTest::RunTest(const FString& Parameters)
+{
+	using namespace PCGCopyPointsGPUTestHelpers;
+
+	return PCGGPUTestCommon::RunCPUGPUParityTest(this,
+		[](UPCGGraph* Graph, UWorld* /*World*/) -> PCGGPUTestCommon::FParityBuildResult
+		{
+			UPCGNode* SourceNode = PCGTestsCommon::AddCreatePointsNode(Graph, SourceTransforms);
+			UPCGNode* TargetNode = PCGTestsCommon::AddCreatePointsNode(Graph, TargetTransforms);
+			FCopyPointsGraphSetup Setup = BuildCopyPointsGraph(Graph, SourceNode, TargetNode);
+
+			Setup.CopySettings->RotationInheritance = EPCGCopyPointsInheritanceMode::Target;
+			Setup.CopySettings->ScaleInheritance = EPCGCopyPointsInheritanceMode::Target;
+			Setup.CopySettings->ColorInheritance = EPCGCopyPointsInheritanceMode::Target;
+			Setup.CopySettings->SeedInheritance = EPCGCopyPointsInheritanceMode::Target;
+			Setup.CopySettings->AttributeInheritance = EPCGCopyPointsMetadataInheritanceMode::TargetOnly;
+
+			return { Setup.GatherNode, { Setup.CopySettings } };
+		},
+		[](FAutomationTestBase* Test, const FPCGDataCollection& CPUOutput, const FPCGDataCollection& GPUOutput) -> bool
+		{
+			return PCGTestsCommon::ComparePointDataCollections(Test, CPUOutput, GPUOutput);
+		});
+}
+
+/**
+ * GPU test: CopyPoints with bCopyEachSourceOnEveryTarget=false (N:N pairing instead of cartesian product).
+ */
+IMPLEMENT_CUSTOM_SIMPLE_AUTOMATION_TEST(
+	FPCGCopyPointsGPU_NonCartesianTest,
+	FPCGTestBaseClass,
+	"Plugins.PCG.Compute.CopyPoints.NonCartesian",
+	PCGTestsCommon::TestFlags | EAutomationTestFlags::NonNullRHI)
+
+bool FPCGCopyPointsGPU_NonCartesianTest::RunTest(const FString& Parameters)
+{
+	using namespace PCGCopyPointsGPUTestHelpers;
+
+	return PCGGPUTestCommon::RunCPUGPUParityTest(this,
+		[](UPCGGraph* Graph, UWorld* /*World*/) -> PCGGPUTestCommon::FParityBuildResult
+		{
+			UPCGNode* SourceNode = PCGTestsCommon::AddCreatePointsNode(Graph, SourceTransforms);
+			UPCGNode* TargetNode = PCGTestsCommon::AddCreatePointsNode(Graph, TargetTransforms);
+			FCopyPointsGraphSetup Setup = BuildCopyPointsGraph(Graph, SourceNode, TargetNode);
+
+			Setup.CopySettings->bCopyEachSourceOnEveryTarget = false;
+
+			return { Setup.GatherNode, { Setup.CopySettings } };
+		},
+		[](FAutomationTestBase* Test, const FPCGDataCollection& CPUOutput, const FPCGDataCollection& GPUOutput) -> bool
+		{
+			return PCGTestsCommon::ComparePointDataCollections(Test, CPUOutput, GPUOutput);
+		});
+}
+
+/**
+ * GPU test: CopyPoints with all overridable params supplied via connected CPU override pins.
+ *
+ * Settings are left at defaults; CreateAttributeSet nodes connected to each override pin supply non-trivial values.
+ * This exercises the GPU override system for every overridable parameter on CopyPoints.
+ */
+IMPLEMENT_CUSTOM_SIMPLE_AUTOMATION_TEST(
+	FPCGCopyPointsGPU_OverrideAllParamsTest,
+	FPCGTestBaseClass,
+	"Plugins.PCG.Compute.CopyPoints.OverrideAllParams",
+	PCGTestsCommon::TestFlags | EAutomationTestFlags::NonNullRHI)
+
+bool FPCGCopyPointsGPU_OverrideAllParamsTest::RunTest(const FString& Parameters)
+{
+	using namespace PCGCopyPointsGPUTestHelpers;
+
+	return PCGGPUTestCommon::RunCPUGPUParityTest(this,
+		[](UPCGGraph* Graph, UWorld* /*World*/) -> PCGGPUTestCommon::FParityBuildResult
+		{
+			UPCGNode* SourceNode = PCGTestsCommon::AddCreatePointsNode(Graph, SourceTransforms);
+			UPCGNode* TargetNode = PCGTestsCommon::AddCreatePointsNode(Graph, TargetTransforms);
+			FCopyPointsGraphSetup Setup = BuildCopyPointsGraph(Graph, SourceNode, TargetNode);
+
+			// Override all overridable params via connected CreateAttributeSet nodes.
+			// Enum params are overridden as Int32.
+			PCGGPUTestCommon::AddOverride(Graph, Setup.CopyNode, GET_MEMBER_NAME_CHECKED(UPCGCopyPointsSettings, RotationInheritance), static_cast<int32>(EPCGCopyPointsInheritanceMode::Source));
+			PCGGPUTestCommon::AddOverride(Graph, Setup.CopyNode, GET_MEMBER_NAME_CHECKED(UPCGCopyPointsSettings, bApplyTargetRotationToPositions), false);
+			PCGGPUTestCommon::AddOverride(Graph, Setup.CopyNode, GET_MEMBER_NAME_CHECKED(UPCGCopyPointsSettings, ScaleInheritance), static_cast<int32>(EPCGCopyPointsInheritanceMode::Target));
+			PCGGPUTestCommon::AddOverride(Graph, Setup.CopyNode, GET_MEMBER_NAME_CHECKED(UPCGCopyPointsSettings, bApplyTargetScaleToPositions), false);
+			PCGGPUTestCommon::AddOverride(Graph, Setup.CopyNode, GET_MEMBER_NAME_CHECKED(UPCGCopyPointsSettings, ColorInheritance), static_cast<int32>(EPCGCopyPointsInheritanceMode::Source));
+			PCGGPUTestCommon::AddOverride(Graph, Setup.CopyNode, GET_MEMBER_NAME_CHECKED(UPCGCopyPointsSettings, SeedInheritance), static_cast<int32>(EPCGCopyPointsInheritanceMode::Target));
+			PCGGPUTestCommon::AddOverride(Graph, Setup.CopyNode, GET_MEMBER_NAME_CHECKED(UPCGCopyPointsSettings, AttributeInheritance), static_cast<int32>(EPCGCopyPointsMetadataInheritanceMode::TargetFirst));
+			PCGGPUTestCommon::AddOverride(Graph, Setup.CopyNode, GET_MEMBER_NAME_CHECKED(UPCGCopyPointsSettings, TagInheritance), static_cast<int32>(EPCGCopyPointsTagInheritanceMode::Target));
+			PCGGPUTestCommon::AddOverride(Graph, Setup.CopyNode, GET_MEMBER_NAME_CHECKED(UPCGCopyPointsSettings, bCopyEachSourceOnEveryTarget), true);
+
+			return { Setup.GatherNode, { Setup.CopySettings } };
+		},
+		[](FAutomationTestBase* Test, const FPCGDataCollection& CPUOutput, const FPCGDataCollection& GPUOutput) -> bool
+		{
+			return PCGTestsCommon::ComparePointDataCollections(Test, CPUOutput, GPUOutput);
+		});
+}
+
+/**
+ * GPU test: CopyPoints with all overridable params supplied from the GPU via CreateAttributeSet -> CustomHLSL ASP chains.
+ *
+ * Each override is produced by a CreateAttributeSet feeding into an empty CustomHLSL AttributeSetProcessor
+ * that passes the value through on the GPU. This validates that overrides arriving from GPU-resident data
+ * are correctly read back and applied.
+ */
+IMPLEMENT_CUSTOM_SIMPLE_AUTOMATION_TEST(
+	FPCGCopyPointsGPU_GPUOverridesTest,
+	FPCGTestBaseClass,
+	"Plugins.PCG.Compute.CopyPoints.GPUOverrides",
+	PCGTestsCommon::TestFlags | EAutomationTestFlags::NonNullRHI)
+
+bool FPCGCopyPointsGPU_GPUOverridesTest::RunTest(const FString& Parameters)
+{
+	using namespace PCGCopyPointsGPUTestHelpers;
+
+	return PCGGPUTestCommon::RunCPUGPUParityTest(this,
+		[](UPCGGraph* Graph, UWorld* /*World*/) -> PCGGPUTestCommon::FParityBuildResult
+		{
+			UPCGNode* SourceNode = PCGTestsCommon::AddCreatePointsNode(Graph, SourceTransforms);
+			UPCGNode* TargetNode = PCGTestsCommon::AddCreatePointsNode(Graph, TargetTransforms);
+			FCopyPointsGraphSetup Setup = BuildCopyPointsGraph(Graph, SourceNode, TargetNode);
+
+			// Override params via GPU-produced attribute sets (CreateAttributeSet -> CustomHLSL passthrough).
+			PCGGPUTestCommon::AddGPUOverride(Graph, Setup.CopyNode, GET_MEMBER_NAME_CHECKED(UPCGCopyPointsSettings, RotationInheritance), static_cast<int32>(EPCGCopyPointsInheritanceMode::Target));
+			PCGGPUTestCommon::AddGPUOverride(Graph, Setup.CopyNode, GET_MEMBER_NAME_CHECKED(UPCGCopyPointsSettings, bApplyTargetRotationToPositions), true);
+			PCGGPUTestCommon::AddGPUOverride(Graph, Setup.CopyNode, GET_MEMBER_NAME_CHECKED(UPCGCopyPointsSettings, ScaleInheritance), static_cast<int32>(EPCGCopyPointsInheritanceMode::Source));
+			PCGGPUTestCommon::AddGPUOverride(Graph, Setup.CopyNode, GET_MEMBER_NAME_CHECKED(UPCGCopyPointsSettings, bApplyTargetScaleToPositions), true);
+			PCGGPUTestCommon::AddGPUOverride(Graph, Setup.CopyNode, GET_MEMBER_NAME_CHECKED(UPCGCopyPointsSettings, ColorInheritance), static_cast<int32>(EPCGCopyPointsInheritanceMode::Relative));
+			PCGGPUTestCommon::AddGPUOverride(Graph, Setup.CopyNode, GET_MEMBER_NAME_CHECKED(UPCGCopyPointsSettings, SeedInheritance), static_cast<int32>(EPCGCopyPointsInheritanceMode::Source));
+			PCGGPUTestCommon::AddGPUOverride(Graph, Setup.CopyNode, GET_MEMBER_NAME_CHECKED(UPCGCopyPointsSettings, AttributeInheritance), static_cast<int32>(EPCGCopyPointsMetadataInheritanceMode::SourceFirst));
+			PCGGPUTestCommon::AddGPUOverride(Graph, Setup.CopyNode, GET_MEMBER_NAME_CHECKED(UPCGCopyPointsSettings, TagInheritance), static_cast<int32>(EPCGCopyPointsTagInheritanceMode::Both));
+			PCGGPUTestCommon::AddGPUOverride(Graph, Setup.CopyNode, GET_MEMBER_NAME_CHECKED(UPCGCopyPointsSettings, bCopyEachSourceOnEveryTarget), false);
+
+			return { Setup.GatherNode, { Setup.CopySettings } };
+		},
+		[](FAutomationTestBase* Test, const FPCGDataCollection& CPUOutput, const FPCGDataCollection& GPUOutput) -> bool
+		{
+			return PCGTestsCommon::ComparePointDataCollections(Test, CPUOutput, GPUOutput);
+		});
+}
+
+#endif // WITH_EDITOR
diff --git a/Engine/Plugins/PCG/Source/PCG/Private/Tests/Compute/PCGCustomHLSLAttributeTest.cpp b/Engine/Plugins/PCG/Source/PCG/Private/Tests/Compute/PCGCustomHLSLAttributeTest.cpp
--- a/Engine/Plugins/PCG/Source/PCG/Private/Tests/Compute/PCGCustomHLSLAttributeTest.cpp
+++ b/Engine/Plugins/PCG/Source/PCG/Private/Tests/Compute/PCGCustomHLSLAttributeTest.cpp
@@ -61,7 +61,7 @@ bool FPCGCustomHLSLAttributeSetDoubleValueTest::RunTest(const FString& Parameter
         return false;
     }
 
-    return PCGGPUTestCommon::ValidateSingleTaggedData(this, OutNodeData[0], [&](const FPCGTaggedData& TaggedData)
+    return PCGTestsCommon::ValidateSingleTaggedData(this, OutNodeData[0], [&](const FPCGTaggedData& TaggedData)
     {
         const UPCGParamData* OutParams = Cast<UPCGParamData>(TaggedData.Data);
         if (!TestNotNull(TEXT("Readback data is attribute set data"), OutParams))
@@ -132,7 +132,7 @@ bool FPCGCustomHLSLAttributeSetStringAttributeTest::RunTest(const FString& Param
         return false;
     }
 
-    return PCGGPUTestCommon::ValidateSingleTaggedData(this, OutNodeData[0], [&](const FPCGTaggedData& TaggedData)
+    return PCGTestsCommon::ValidateSingleTaggedData(this, OutNodeData[0], [&](const FPCGTaggedData& TaggedData)
     {
         const UPCGParamData* OutParams = Cast<UPCGParamData>(TaggedData.Data);
         if (!TestNotNull(TEXT("Readback data is attribute set data"), OutParams))
@@ -215,7 +215,7 @@ bool FPCGCustomHLSLInferredAttributeTypesTest::RunTest(const FString& Parameters
         return false;
     }
 
-    return PCGGPUTestCommon::ValidateSingleTaggedData(this, OutNodeData[0], [&](const FPCGTaggedData& TaggedData)
+    return PCGTestsCommon::ValidateSingleTaggedData(this, OutNodeData[0], [&](const FPCGTaggedData& TaggedData)
     {
         const UPCGBasePointData* OutPoints = Cast<UPCGBasePointData>(TaggedData.Data);
         if (!TestNotNull(TEXT("Readback data is point data"), OutPoints))
@@ -297,7 +297,7 @@ bool FPCGCustomHLSLPerPinAttributeSettingsEnabledTest::RunTest(const FString& Pa
         return false;
     }
 
-    return PCGGPUTestCommon::ValidateSingleTaggedData(this, OutNodeData[0], [&](const FPCGTaggedData& TaggedData)
+    return PCGTestsCommon::ValidateSingleTaggedData(this, OutNodeData[0], [&](const FPCGTaggedData& TaggedData)
     {
         const UPCGBasePointData* OutPoints = Cast<UPCGBasePointData>(TaggedData.Data);
         if (!TestNotNull(TEXT("Readback data is point data"), OutPoints))
@@ -374,7 +374,7 @@ bool FPCGCustomHLSLPerPinAttributeSettingsDisabledTest::RunTest(const FString& P
         return false;
     }
 
-    return PCGGPUTestCommon::ValidateSingleTaggedData(this, OutNodeData[0], [&](const FPCGTaggedData& TaggedData)
+    return PCGTestsCommon::ValidateSingleTaggedData(this, OutNodeData[0], [&](const FPCGTaggedData& TaggedData)
     {
         const UPCGBasePointData* OutPoints = Cast<UPCGBasePointData>(TaggedData.Data);
         if (!TestNotNull(TEXT("Readback data is point data"), OutPoints))
@@ -452,7 +452,7 @@ bool FPCGCustomHLSLAttributeCreationOverlapTest::RunTest(const FString& Paramete
         return false;
     }
 
-    return PCGGPUTestCommon::ValidateSingleTaggedData(this, OutNodeData[0], [&](const FPCGTaggedData& TaggedData)
+    return PCGTestsCommon::ValidateSingleTaggedData(this, OutNodeData[0], [&](const FPCGTaggedData& TaggedData)
     {
         const UPCGBasePointData* OutPoints = Cast<UPCGBasePointData>(TaggedData.Data);
         if (!TestNotNull(TEXT("Readback data is point data"), OutPoints))
@@ -529,7 +529,7 @@ bool FPCGCustomHLSLAuthoredAndInferredCoexistTest::RunTest(const FString& Parame
         return false;
     }
 
-    return PCGGPUTestCommon::ValidateSingleTaggedData(this, OutNodeData[0], [&](const FPCGTaggedData& TaggedData)
+    return PCGTestsCommon::ValidateSingleTaggedData(this, OutNodeData[0], [&](const FPCGTaggedData& TaggedData)
     {
         const UPCGBasePointData* OutPoints = Cast<UPCGBasePointData>(TaggedData.Data);
         if (!TestNotNull(TEXT("Readback data is point data"), OutPoints))
diff --git a/Engine/Plugins/PCG/Source/PCG/Private/Tests/Compute/PCGCustomHLSLTest.cpp b/Engine/Plugins/PCG/Source/PCG/Private/Tests/Compute/PCGCustomHLSLTest.cpp
--- a/Engine/Plugins/PCG/Source/PCG/Private/Tests/Compute/PCGCustomHLSLTest.cpp
+++ b/Engine/Plugins/PCG/Source/PCG/Private/Tests/Compute/PCGCustomHLSLTest.cpp
@@ -74,7 +74,7 @@ bool FPCGCustomHLSLPointTranslateTest::RunTest(const FString& Parameters)
 		return false;
 	}
 
-	return PCGGPUTestCommon::ValidateSingleTaggedData(this, OutNodeData[0], [&](const FPCGTaggedData& TaggedData)
+	return PCGTestsCommon::ValidateSingleTaggedData(this, OutNodeData[0], [&](const FPCGTaggedData& TaggedData)
 	{
 		const UPCGBasePointData* OutPoints = Cast<UPCGBasePointData>(TaggedData.Data);
 		if (!TestNotNull(TEXT("Readback data is point data"), OutPoints))
@@ -162,7 +162,7 @@ bool FPCGCustomHLSLChainedPointProcessorsTest::RunTest(const FString& Parameters
 		return false;
 	}
 
-	return PCGGPUTestCommon::ValidateSingleTaggedData(this, OutNodeData[0], [&](const FPCGTaggedData& TaggedData)
+	return PCGTestsCommon::ValidateSingleTaggedData(this, OutNodeData[0], [&](const FPCGTaggedData& TaggedData)
 	{
 		const UPCGBasePointData* OutPoints = Cast<UPCGBasePointData>(TaggedData.Data);
 		if (!TestNotNull(TEXT("Readback data is point data"), OutPoints))
@@ -272,7 +272,7 @@ bool FPCGCustomHLSLSanitizedPinLabelsTest::RunTest(const FString& Parameters)
 		return false;
 	}
 
-	return PCGGPUTestCommon::ValidateSingleTaggedData(this, OutNodeData[0], [&](const FPCGTaggedData& TaggedData)
+	return PCGTestsCommon::ValidateSingleTaggedData(this, OutNodeData[0], [&](const FPCGTaggedData& TaggedData)
 	{
 		const UPCGBasePointData* OutPoints = Cast<UPCGBasePointData>(TaggedData.Data);
 		if (!TestNotNull(TEXT("Readback data is point data"), OutPoints))
diff --git a/Engine/Plugins/PCG/Source/PCG/Private/Tests/Compute/PCGGPUStringTableTest.cpp b/Engine/Plugins/PCG/Source/PCG/Private/Tests/Compute/PCGGPUStringTableTest.cpp
--- a/Engine/Plugins/PCG/Source/PCG/Private/Tests/Compute/PCGGPUStringTableTest.cpp
+++ b/Engine/Plugins/PCG/Source/PCG/Private/Tests/Compute/PCGGPUStringTableTest.cpp
@@ -59,7 +59,7 @@ bool FPCGStringTableFNamePassthroughTest::RunTest(const FString& Parameters)
 		return false;
 	}
 
-	return PCGGPUTestCommon::ValidateSingleTaggedData(this, OutNodeData[0], [&](const FPCGTaggedData& TaggedData)
+	return PCGTestsCommon::ValidateSingleTaggedData(this, OutNodeData[0], [&](const FPCGTaggedData& TaggedData)
 	{
 		const UPCGParamData* OutParams = Cast<UPCGParamData>(TaggedData.Data);
 		if (!TestNotNull(TEXT("Readback data is attribute set data"), OutParams))
@@ -142,7 +142,7 @@ bool FPCGStringTableChainedGPUKernelsTest::RunTest(const FString& Parameters)
 		return false;
 	}
 
-	return PCGGPUTestCommon::ValidateSingleTaggedData(this, OutNodeData[0], [&](const FPCGTaggedData& TaggedData)
+	return PCGTestsCommon::ValidateSingleTaggedData(this, OutNodeData[0], [&](const FPCGTaggedData& TaggedData)
 	{
 		const UPCGParamData* OutParams = Cast<UPCGParamData>(TaggedData.Data);
 		if (!TestNotNull(TEXT("Readback data is attribute set data"), OutParams))
@@ -225,7 +225,7 @@ bool FPCGStringTableMultipleStringAttributesTest::RunTest(const FString& Paramet
 		return false;
 	}
 
-	return PCGGPUTestCommon::ValidateSingleTaggedData(this, OutNodeData[0], [&](const FPCGTaggedData& TaggedData)
+	return PCGTestsCommon::ValidateSingleTaggedData(this, OutNodeData[0], [&](const FPCGTaggedData& TaggedData)
 	{
 		const UPCGParamData* OutParams = Cast<UPCGParamData>(TaggedData.Data);
 		if (!TestNotNull(TEXT("Readback data is attribute set data"), OutParams))
diff --git a/Engine/Plugins/PCG/Source/PCG/Private/Tests/Compute/PCGGPUTestCommon.cpp b/Engine/Plugins/PCG/Source/PCG/Private/Tests/Compute/PCGGPUTestCommon.cpp
--- a/Engine/Plugins/PCG/Source/PCG/Private/Tests/Compute/PCGGPUTestCommon.cpp
+++ b/Engine/Plugins/PCG/Source/PCG/Private/Tests/Compute/PCGGPUTestCommon.cpp
@@ -6,48 +6,118 @@
 
 #include "PCGComponent.h"
 #include "PCGGraph.h"
+#include "PCGSettings.h"
 #include "Compute/Data/PCGProxyForGPUData.h"
 #include "Graph/PCGGraphExecutor.h"
 
 #include "ComputeWorkerInterface.h"
+#include "EngineUtils.h"
 #include "RenderingThread.h"
 #include "Components/BoxComponent.h"
 #include "ComputeFramework/ComputeFramework.h"
 #include "Containers/Ticker.h"
+#include "Engine/Engine.h"
 #include "Engine/World.h"
 #include "HAL/PlatformProcess.h"
 #include "Misc/AutomationTest.h"
 
-bool PCGGPUTestCommon::ValidateSingleTaggedData(FAutomationTestBase* Test, const FPCGDataCollection& DataCollection, TFunctionRef<bool(const FPCGTaggedData&)> Validator)
+bool PCGGPUTestCommon::RunCPUGPUParityTest(FAutomationTestBase* Test, FBuildGraphFunc BuildGraph, FValidateFunc Validate, const FPCGGPUGraphTestRunnerConfig& RunnerConfig)
 {
-	if (!Test->TestEqual(TEXT("Data collection contains a single data item"), DataCollection.TaggedData.Num(), 1))
+	// CPU execution
+	FPCGGPUGraphTestRunner CPURunner(RunnerConfig);
+	FParityBuildResult CPUSetup = BuildGraph(CPURunner.Graph, CPURunner.TestWorld);
+
+	for (UPCGSettings* Settings : CPUSetup.GPUSettings)
+	{
+		// Explicitly disable GPU mode on the settings.
+		Settings->SetExecuteOnGPU(false);
+	}
+
+	TArray<FPCGDataCollection> CPUOutput;
+	if (!CPURunner.Execute({ Test, { CPUSetup.CollectNode } }, CPUOutput))
 	{
 		return false;
 	}
-	return Validator(DataCollection.TaggedData[0]);
+
+	// GPU execution
+	FPCGGPUGraphTestRunner GPURunner(RunnerConfig);
+	FParityBuildResult GPUSetup = BuildGraph(GPURunner.Graph, GPURunner.TestWorld);
+
+	for (UPCGSettings* Settings : GPUSetup.GPUSettings)
+	{
+		Settings->SetExecuteOnGPU(true);
+	}
+
+	TArray<FPCGDataCollection> GPUOutput;
+	if (!GPURunner.Execute({ Test, { GPUSetup.CollectNode } }, GPUOutput))
+	{
+		return false;
+	}
+
+	return Validate(Test, CPUOutput[0], GPUOutput[0]);
 }
 
-FPCGGPUGraphTestRunner::FPCGGPUGraphTestRunner()
+FPCGGPUGraphTestRunner::FPCGGPUGraphTestRunner(const FPCGGPUGraphTestRunnerConfig& InConfig)
 {
-	// The test actor has no geometry, so GetActorBounds() returns an invalid FBox, which
-	// causes CreateGenerateTask to abort. Add a UBoxComponent to give the actor valid bounds.
-	UBoxComponent* BoundsBox = NewObject<UBoxComponent>(TestData.TestActor, FName(TEXT("BoundsBox")), RF_Transient);
-	BoundsBox->SetBoxExtent(FVector(100.f, 100.f, 100.f));
+	// Create an isolated game world so test artifacts (actors, landscapes, etc.) don't pollute
+	// the editor world. The world gets a renderer Scene (needed for ComputeFramework::FlushWork).
+	const FName WorldName = MakeUniqueObjectName(GetTransientPackage(), UWorld::StaticClass(), FName(TEXT("PCGGPUTestWorld")));
+	FWorldContext& WorldContext = GEngine->CreateNewWorldContext(EWorldType::Game);
+	TestWorld = UWorld::CreateWorld(EWorldType::Game, /*bInformEngineOfWorld=*/false, WorldName, GetTransientPackage());
+	check(TestWorld);
+	TestWorld->AddToRoot();
+	WorldContext.SetCurrentWorld(TestWorld);
+
+	// Spawn test actor with valid bounds. The bounds box must exist before the PCG component
+	// registers, because registration queries actor bounds to set up the execution source.
+	FActorSpawnParameters SpawnParams;
+	SpawnParams.ObjectFlags = RF_Transient;
+	TestActor = TestWorld->SpawnActor<AActor>(SpawnParams);
+	check(TestActor);
+	TestActor->SetActorLocation(InConfig.ActorLocation);
+
+	UBoxComponent* BoundsBox = NewObject<UBoxComponent>(TestActor, FName(TEXT("BoundsBox")), RF_Transient);
+	BoundsBox->SetBoxExtent(InConfig.BoundsExtent);
 	BoundsBox->RegisterComponent();
-	TestData.TestActor->AddInstanceComponent(BoundsBox);
-	if (!TestData.TestActor->GetRootComponent())
-	{
-		TestData.TestActor->SetRootComponent(BoundsBox);
-	}
+	TestActor->AddInstanceComponent(BoundsBox);
+	TestActor->SetRootComponent(BoundsBox);
+
+	TestPCGComponent = NewObject<UPCGComponent>(TestActor, FName(TEXT("TestPCGComponent")), RF_Transient);
+	check(TestPCGComponent);
+	TestActor->AddInstanceComponent(TestPCGComponent);
+	TestPCGComponent->RegisterComponent();
+	TestPCGComponent->SetIsPartitioned(false);
 
 	Graph = NewObject<UPCGGraph>(GetTransientPackage(), NAME_None, RF_Transient);
 }
 
+FPCGGPUGraphTestRunner::~FPCGGPUGraphTestRunner()
+{
+	if (TestWorld)
+	{
+		// End-play all actors so components unregister cleanly.
+		for (AActor* Actor : FActorRange(TestWorld))
+		{
+			if (Actor)
+			{
+				Actor->RouteEndPlay(EEndPlayReason::LevelTransition);
+			}
+		}
+
+		GEngine->ShutdownWorldNetDriver(TestWorld);
+		TestWorld->DestroyWorld(/*bInformEngineOfWorld=*/true);
+		TestWorld->SetPhysicsScene(nullptr);
+		GEngine->DestroyWorldContext(TestWorld);
+		TestWorld->RemoveFromRoot();
+		TestWorld = nullptr;
+	}
+}
+
 bool FPCGGPUGraphTestRunner::Execute(const FPCGGPUGraphTestRunnerParams& Params, TArray<FPCGDataCollection>& OutNodeData)
 {
 	check(Params.Test);
 
-	TestData.TestPCGComponent->SetGraph(Graph);
+	TestPCGComponent->SetGraph(Graph);
 
 	TSharedPtr<FPCGGraphExecutor> Executor = MakeShared<FPCGGraphExecutor>();
 	Executor->SetKeepIntermediateResults(true);
@@ -59,7 +129,7 @@ bool FPCGGPUGraphTestRunner::Execute(const FPCGGPUGraphTestRunnerParams& Params,
 
 	// This will execute the graph and run the lambda on each node, which we use to capture output data.
 	const FPCGTaskId FinalTaskId = Executor->ScheduleDebugWithTaskCallback(
-		TestData.TestPCGComponent,
+		TestPCGComponent,
 		[&](FPCGTaskId /*TaskId*/, const UPCGNode* Node, const FPCGDataCollection& Output)
 		{
 			for (int32 i = 0; i < Params.CollectOutputDataForNodes.Num(); ++i)
@@ -78,13 +148,12 @@ bool FPCGGPUGraphTestRunner::Execute(const FPCGGPUGraphTestRunnerParams& Params,
 	volatile bool bExecutionComplete = false;
 	Executor->ScheduleGeneric(
 		[&bExecutionComplete]() { bExecutionComplete = true; return true; },
-		TestData.TestPCGComponent,
+		TestPCGComponent,
 		{ FinalTaskId }
 	);
 
 	constexpr float FakeDeltaTime = 1.f / 30.f;
 	constexpr int32 MaxExecuteIterations = 10000;
-	UWorld* TestWorld = TestData.TestActor->GetWorld();
 	for (int32 i = 0; i < MaxExecuteIterations && !bExecutionComplete; ++i)
 	{
 		// Fires FPCGModule::Tick (flushing SleepUntilNextFrame unpause callbacks) AND any ExecuteOnGameThread
diff --git a/Engine/Plugins/PCG/Source/PCG/Private/Tests/Compute/PCGGPUTestCommon.h b/Engine/Plugins/PCG/Source/PCG/Private/Tests/Compute/PCGGPUTestCommon.h
--- a/Engine/Plugins/PCG/Source/PCG/Private/Tests/Compute/PCGGPUTestCommon.h
+++ b/Engine/Plugins/PCG/Source/PCG/Private/Tests/Compute/PCGGPUTestCommon.h
@@ -7,21 +7,65 @@
 #include "Tests/PCGTestsCommon.h"
 
 #include "PCGData.h"
+#include "Compute/Elements/PCGCustomHLSL.h"
 
 #include "Templates/Function.h"
 
+class AActor;
 class FAutomationTestBase;
+class UPCGComponent;
 class UPCGGraph;
 class UPCGNode;
+class UPCGSettings;
+class UWorld;
+
+struct FPCGGPUGraphTestRunnerConfig
+{
+	/** World-space location of the test actor. */
+	FVector ActorLocation = FVector::ZeroVector;
+
+	/** Extent of the box component attached to the test actor. Controls the PCG component's grid bounds. */
+	FVector BoundsExtent = FVector(2500.f);
+};
 
 namespace PCGGPUTestCommon
 {
+	/** Add a CreateAttributeSet override node and wire it to a target node's override pin. */
+	template <typename T>
+	void AddOverride(UPCGGraph* Graph, UPCGNode* TargetNode, FName PinName, const T& Value)
+	{
+		UPCGNode* Node = PCGTestsCommon::AddCreateAttributeSetNode(Graph, PinName, Value);
+		Graph->AddEdge(Node, PCGPinConstants::DefaultOutputLabel, TargetNode, PinName);
+	}
+
+	/** Add a CreateAttributeSet -> empty CustomHLSL ASP chain (forces GPU execution) and wire to a target node's override pin. */
+	template <typename T>
+	void AddGPUOverride(UPCGGraph* Graph, UPCGNode* TargetNode, FName PinName, const T& Value)
+	{
+		UPCGNode* CreateNode = PCGTestsCommon::AddCreateAttributeSetNode(Graph, PinName, Value);
+
+		UPCGCustomHLSLSettings* HLSLSettings = nullptr;
+		UPCGNode* HLSLNode = Graph->AddNodeOfType<UPCGCustomHLSLSettings>(HLSLSettings);
+		HLSLSettings->SetKernelType(EPCGKernelType::AttributeSetProcessor);
+
+		Graph->AddEdge(CreateNode, PCGPinConstants::DefaultOutputLabel, HLSLNode, PCGPinConstants::DefaultInputLabel);
+		Graph->AddEdge(HLSLNode, PCGPinConstants::DefaultOutputLabel, TargetNode, PinName);
+	}
+
 	/**
-	 * Verifies DataCollection contains exactly one item, then calls Validator on it.
-	 * Returns false (with a logged failure) if the count check fails, otherwise returns
-	 * whatever Validator returns.
+	 * Run the same graph on CPU and GPU (separate runners) and validate the outputs.
+	 * BuildGraph is called twice — it should build the graph and return the node to collect output from
+	 * plus the settings whose bExecuteOnGPU flag should be toggled for the GPU run.
+	 * Validate receives the two output collections and should return true if they match.
 	 */
-	bool ValidateSingleTaggedData(FAutomationTestBase* Test, const FPCGDataCollection& DataCollection, TFunctionRef<bool(const FPCGTaggedData&)> Validator);
+	struct FParityBuildResult
+	{
+		UPCGNode* CollectNode = nullptr;
+		TArray<UPCGSettings*> GPUSettings;
+	};
+	using FBuildGraphFunc = TFunctionRef<FParityBuildResult(UPCGGraph* Graph, UWorld* World)>;
+	using FValidateFunc = TFunctionRef<bool(FAutomationTestBase* Test, const FPCGDataCollection& CPUOutput, const FPCGDataCollection& GPUOutput)>;
+	bool RunCPUGPUParityTest(FAutomationTestBase* Test, FBuildGraphFunc BuildGraph, FValidateFunc Validate, const FPCGGPUGraphTestRunnerConfig& RunnerConfig = {});
 }
 
 /** Parameters for FPCGGPUGraphTestRunner::Execute(). */
@@ -49,11 +93,19 @@ struct FPCGGPUGraphTestRunnerParams
  */
 struct FPCGGPUGraphTestRunner
 {
-	PCGTestsCommon::FTestData TestData;
+	UWorld* TestWorld = nullptr;
+	AActor* TestActor = nullptr;
+	UPCGComponent* TestPCGComponent = nullptr;
 	UPCGGraph* Graph = nullptr;
 
-	/** Sets up the test actor (with valid bounds), PCG component, and graph. */
-	FPCGGPUGraphTestRunner();
+	/** Creates an isolated game world with a test actor, PCG component, and graph. */
+	explicit FPCGGPUGraphTestRunner(const FPCGGPUGraphTestRunnerConfig& InConfig = {});
+
+	/** Tears down the isolated world and all actors within it. */
+	~FPCGGPUGraphTestRunner();
+
+	FPCGGPUGraphTestRunner(const FPCGGPUGraphTestRunner&) = delete;
+	FPCGGPUGraphTestRunner& operator=(const FPCGGPUGraphTestRunner&) = delete;
 
 	/**
 	 * Runs the graph to completion, reads back (if Params.bReadbackGPUData) all
diff --git a/Engine/Plugins/PCG/Source/PCG/Private/Tests/Compute/PCGTransformPointsGPUTest.cpp b/Engine/Plugins/PCG/Source/PCG/Private/Tests/Compute/PCGTransformPointsGPUTest.cpp
new file mode 100644
--- /dev/null
+++ b/Engine/Plugins/PCG/Source/PCG/Private/Tests/Compute/PCGTransformPointsGPUTest.cpp
@@ -0,0 +1,136 @@
+// Copyright Epic Games, Inc. All Rights Reserved.
+
+#if WITH_EDITOR
+
+#include "Tests/Compute/PCGGPUTestCommon.h"
+
+#include "PCGCommon.h"
+#include "PCGGraph.h"
+#include "Data/PCGBasePointData.h"
+#include "Elements/PCGGather.h"
+#include "Elements/PCGTransformPoints.h"
+
+namespace PCGTransformPointsGPUTestHelpers
+{
+	static const TArray<FTransform> InputTransforms = {
+		FTransform(FQuat::Identity, FVector(  0.0,   0.0,   0.0), FVector(1.0, 2.0, 3.0)),
+		FTransform(FQuat::Identity, FVector(100.0, 200.0, 300.0), FVector(0.5, 1.5, 2.5)),
+		FTransform(FQuat::Identity, FVector(-50.0,  75.0, 250.0), FVector(1.0, 1.0, 1.0)),
+		FTransform(FQuat::Identity, FVector( 10.0, -30.0, -80.0), FVector(2.0, 0.5, 1.5)),
+	};
+
+	/** Create a TransformPoints node with varied settings that exercise all params. */
+	UPCGNode* AddTransformPointsNode(UPCGGraph* Graph, UPCGTransformPointsSettings*& OutSettings)
+	{
+		OutSettings = nullptr;
+		UPCGNode* Node = Graph->AddNodeOfType<UPCGTransformPointsSettings>(OutSettings);
+		OutSettings->OffsetMin = FVector(10.0, -5.0, 20.0);
+		OutSettings->OffsetMax = FVector(30.0, 15.0, 50.0);
+		OutSettings->bAbsoluteOffset = true;
+		OutSettings->RotationMin = FRotator(10.0, 20.0, 0.0);
+		OutSettings->RotationMax = FRotator(45.0, 90.0, 30.0);
+		OutSettings->bAbsoluteRotation = false;
+		OutSettings->ScaleMin = FVector(0.5, 1.0, 0.8);
+		OutSettings->ScaleMax = FVector(2.0, 3.0, 1.5);
+		OutSettings->bAbsoluteScale = false;
+		OutSettings->bUniformScale = false;
+		OutSettings->bRecomputeSeed = true;
+		return Node;
+	}
+}
+
+/**
+ * GPU test: TransformPoints CPU/GPU parity with varied offset, rotation, and scale ranges.
+ *
+ *   CreatePoints -> TransformPoints -> Gather
+ *
+ * All settings use distinct min/max ranges to exercise seed-based randomness on both paths.
+ */
+IMPLEMENT_CUSTOM_SIMPLE_AUTOMATION_TEST(
+	FPCGTransformPointsGPU_CPUParityTest,
+	FPCGTestBaseClass,
+	"Plugins.PCG.Compute.TransformPoints.CPUGPUParity",
+	PCGTestsCommon::TestFlags | EAutomationTestFlags::NonNullRHI)
+
+bool FPCGTransformPointsGPU_CPUParityTest::RunTest(const FString& Parameters)
+{
+	using namespace PCGTransformPointsGPUTestHelpers;
+
+	return PCGGPUTestCommon::RunCPUGPUParityTest(this,
+		[](UPCGGraph* Graph, UWorld* /*World*/) -> PCGGPUTestCommon::FParityBuildResult
+		{
+			UPCGNode* CreateNode = PCGTestsCommon::AddCreatePointsNode(Graph, InputTransforms);
+
+			UPCGTransformPointsSettings* TransformSettings = nullptr;
+			UPCGNode* TransformNode = AddTransformPointsNode(Graph, TransformSettings);
+
+			UPCGGatherSettings* GatherSettings = nullptr;
+			UPCGNode* GatherNode = Graph->AddNodeOfType<UPCGGatherSettings>(GatherSettings);
+
+			Graph->AddEdge(CreateNode, PCGPinConstants::DefaultOutputLabel, TransformNode, PCGPinConstants::DefaultInputLabel);
+			Graph->AddEdge(TransformNode, PCGPinConstants::DefaultOutputLabel, GatherNode, PCGPinConstants::DefaultInputLabel);
+
+			return { GatherNode, { TransformSettings } };
+		},
+		[](FAutomationTestBase* Test, const FPCGDataCollection& CPUOutput, const FPCGDataCollection& GPUOutput) -> bool
+		{
+			return PCGTestsCommon::ComparePointDataCollections(Test, CPUOutput, GPUOutput);
+		});
+}
+
+/**
+ * GPU test: TransformPoints with all overridable params supplied via connected override pins.
+ *
+ *   CreatePoints -> TransformPoints (all settings default) -> Gather
+ *                   ^ CreateAttributeSet nodes connected to each override pin
+ *
+ * Settings are left at defaults; override pins supply non-trivial values for every overridable param.
+ * CPU/GPU parity validates the GPU override system processes all params correctly.
+ */
+IMPLEMENT_CUSTOM_SIMPLE_AUTOMATION_TEST(
+	FPCGTransformPointsGPU_OverrideAllParamsTest,
+	FPCGTestBaseClass,
+	"Plugins.PCG.Compute.TransformPoints.OverrideAllParams",
+	PCGTestsCommon::TestFlags | EAutomationTestFlags::NonNullRHI)
+
+bool FPCGTransformPointsGPU_OverrideAllParamsTest::RunTest(const FString& Parameters)
+{
+	using namespace PCGTransformPointsGPUTestHelpers;
+
+	return PCGGPUTestCommon::RunCPUGPUParityTest(this,
+		[](UPCGGraph* Graph, UWorld* /*World*/) -> PCGGPUTestCommon::FParityBuildResult
+		{
+			UPCGNode* CreateNode = PCGTestsCommon::AddCreatePointsNode(Graph, InputTransforms);
+
+			// TransformPoints with default settings — all overrides come from connected pins.
+			UPCGTransformPointsSettings* TransformSettings = nullptr;
+			UPCGNode* TransformNode = Graph->AddNodeOfType<UPCGTransformPointsSettings>(TransformSettings);
+
+			// Override all overridable params via connected CreateAttributeSet nodes.
+			PCGGPUTestCommon::AddOverride(Graph, TransformNode, GET_MEMBER_NAME_CHECKED(UPCGTransformPointsSettings, OffsetMin), FVector(10.0, -5.0, 20.0));
+			PCGGPUTestCommon::AddOverride(Graph, TransformNode, GET_MEMBER_NAME_CHECKED(UPCGTransformPointsSettings, OffsetMax), FVector(30.0, 15.0, 50.0));
+			PCGGPUTestCommon::AddOverride(Graph, TransformNode, GET_MEMBER_NAME_CHECKED(UPCGTransformPointsSettings, bAbsoluteOffset), true);
+			PCGGPUTestCommon::AddOverride(Graph, TransformNode, GET_MEMBER_NAME_CHECKED(UPCGTransformPointsSettings, RotationMin), FRotator(10.0, 20.0, 0.0));
+			PCGGPUTestCommon::AddOverride(Graph, TransformNode, GET_MEMBER_NAME_CHECKED(UPCGTransformPointsSettings, RotationMax), FRotator(45.0, 90.0, 30.0));
+			PCGGPUTestCommon::AddOverride(Graph, TransformNode, GET_MEMBER_NAME_CHECKED(UPCGTransformPointsSettings, bAbsoluteRotation), false);
+			PCGGPUTestCommon::AddOverride(Graph, TransformNode, GET_MEMBER_NAME_CHECKED(UPCGTransformPointsSettings, ScaleMin), FVector(0.5, 1.0, 0.8));
+			PCGGPUTestCommon::AddOverride(Graph, TransformNode, GET_MEMBER_NAME_CHECKED(UPCGTransformPointsSettings, ScaleMax), FVector(2.0, 3.0, 1.5));
+			PCGGPUTestCommon::AddOverride(Graph, TransformNode, GET_MEMBER_NAME_CHECKED(UPCGTransformPointsSettings, bAbsoluteScale), true);
+			PCGGPUTestCommon::AddOverride(Graph, TransformNode, GET_MEMBER_NAME_CHECKED(UPCGTransformPointsSettings, bUniformScale), false);
+			PCGGPUTestCommon::AddOverride(Graph, TransformNode, GET_MEMBER_NAME_CHECKED(UPCGTransformPointsSettings, bRecomputeSeed), true);
+
+			UPCGGatherSettings* GatherSettings = nullptr;
+			UPCGNode* GatherNode = Graph->AddNodeOfType<UPCGGatherSettings>(GatherSettings);
+
+			Graph->AddEdge(CreateNode, PCGPinConstants::DefaultOutputLabel, TransformNode, PCGPinConstants::DefaultInputLabel);
+			Graph->AddEdge(TransformNode, PCGPinConstants::DefaultOutputLabel, GatherNode, PCGPinConstants::DefaultInputLabel);
+
+			return { GatherNode, { TransformSettings } };
+		},
+		[](FAutomationTestBase* Test, const FPCGDataCollection& CPUOutput, const FPCGDataCollection& GPUOutput) -> bool
+		{
+			return PCGTestsCommon::ComparePointDataCollections(Test, CPUOutput, GPUOutput);
+		});
+}
+
+#endif // WITH_EDITOR
diff --git a/Engine/Plugins/PCG/Source/PCG/Private/Tests/PCGTestsCommon.cpp b/Engine/Plugins/PCG/Source/PCG/Private/Tests/PCGTestsCommon.cpp
--- a/Engine/Plugins/PCG/Source/PCG/Private/Tests/PCGTestsCommon.cpp
+++ b/Engine/Plugins/PCG/Source/PCG/Private/Tests/PCGTestsCommon.cpp
@@ -9,6 +9,7 @@
 #include "PCGGraph.h"
 #include "PCGParamData.h"
 #include "PCGPin.h"
+#include "Elements/PCGCreatePoints.h"
 #include "Data/PCGPointData.h"
 #include "Data/PCGPolyLineData.h" // IWYU pragma: keep
 #include "Data/PCGPrimitiveData.h" // IWYU pragma: keep
@@ -16,12 +17,116 @@
 #include "Data/PCGVolumeData.h"
 
 #if WITH_EDITOR
+#include "Engine/World.h"
 #include "Editor.h"
 #include "Components/SceneComponent.h"
-#endif
+#include "Landscape.h"
+#include "LandscapeComponent.h"
+#include "LandscapeDataAccess.h"
+#include "LandscapeHeightfieldCollisionComponent.h"
+#endif // WITH_EDITOR
+
+#if WITH_EDITOR
+namespace
+{
+	constexpr int32 LandscapeNumSubsections = 1;
+	constexpr int32 LandscapeSubsectionSizeQuads = 63;
+	constexpr int32 LandscapeSizeVerts = LandscapeSubsectionSizeQuads * LandscapeNumSubsections + 1;
+	constexpr int32 LandscapeTotalVerts = LandscapeSizeVerts * LandscapeSizeVerts;
+
+	ALandscape* CreateLandscapeFromHeightData(UWorld* World, const FTransform& Transform, TArray<uint16>&& HeightData)
+	{
+		check(World);
+		check(HeightData.Num() == LandscapeTotalVerts);
+
+		FActorSpawnParameters SpawnParams;
+		SpawnParams.ObjectFlags = RF_Transient;
+		ALandscape* Landscape = World->SpawnActor<ALandscape>(Transform.GetLocation(), Transform.GetRotation().Rotator(), SpawnParams);
+		check(Landscape);
+
+		TMap<FGuid, TArray<uint16>> HeightDataPerLayer;
+		HeightDataPerLayer.Add(FGuid(), MoveTemp(HeightData));
+
+		TMap<FGuid, TArray<FLandscapeImportLayerInfo>> MaterialLayerDataPerLayer;
+		MaterialLayerDataPerLayer.Add(FGuid(), TArray<FLandscapeImportLayerInfo>());
+
+		Landscape->Import(
+			FGuid::NewGuid(),
+			0, 0, LandscapeSizeVerts - 1, LandscapeSizeVerts - 1,
+			LandscapeNumSubsections, LandscapeSubsectionSizeQuads,
+			HeightDataPerLayer,
+			nullptr,
+			MaterialLayerDataPerLayer,
+			ELandscapeImportAlphamapType::Additive,
+			MakeArrayView<const FLandscapeLayer>({})
+		);
+
+		Landscape->CreateLandscapeInfo();
+
+		return Landscape;
+	}
+}
+#endif // WITH_EDITOR
 
 namespace PCGTestsCommon
 {
+#if WITH_EDITOR
+	ALandscape* CreateFlatLandscape(UWorld* World, const FTransform& Transform)
+	{
+		TArray<uint16> HeightData;
+		HeightData.SetNum(LandscapeTotalVerts);
+		for (uint16& H : HeightData)
+		{
+			H = LandscapeDataAccess::GetTexHeight(0.0f);
+		}
+
+		return CreateLandscapeFromHeightData(World, Transform, MoveTemp(HeightData));
+	}
+
+	ALandscape* CreatePerlinLandscape(UWorld* World, const FTransform& Transform, int32 Seed)
+	{
+		TArray<uint16> HeightData;
+		HeightData.SetNum(LandscapeTotalVerts);
+
+		// Offset the noise domain by seed so different seeds produce different terrain.
+		const float SeedOffset = Seed * 17.31f;
+
+		for (int32 Y = 0; Y < LandscapeSizeVerts; ++Y)
+		{
+			for (int32 X = 0; X < LandscapeSizeVerts; ++X)
+			{
+				const float NX = X / static_cast<float>(LandscapeSizeVerts) + SeedOffset;
+				const float NY = Y / static_cast<float>(LandscapeSizeVerts) + SeedOffset;
+
+				// Three octaves: broad hills, medium ridges, fine bumps.
+				float H = 0.0f;
+				H += 200.0f * FMath::PerlinNoise2D(FVector2D(NX * 2.0f, NY * 2.0f));
+				H +=  50.0f * FMath::PerlinNoise2D(FVector2D(NX * 8.0f, NY * 8.0f));
+				H +=  10.0f * FMath::PerlinNoise2D(FVector2D(NX * 20.0f, NY * 20.0f));
+
+				HeightData[Y * LandscapeSizeVerts + X] = LandscapeDataAccess::GetTexHeight(H);
+			}
+		}
+
+		return CreateLandscapeFromHeightData(World, Transform, MoveTemp(HeightData));
+	}
+#endif // WITH_EDITOR
+
+#if WITH_EDITOR
+	UPCGNode* AddCreatePointsNode(UPCGGraph* Graph, TArrayView<const FTransform> Transforms)
+	{
+		UPCGCreatePointsSettings* Settings = nullptr;
+		UPCGNode* Node = Graph->AddNodeOfType<UPCGCreatePointsSettings>(Settings);
+		Settings->PointsToCreate.Empty();
+		Settings->CoordinateSpace = EPCGCoordinateSpace::World;
+		for (const FTransform& Transform : Transforms)
+		{
+			Settings->PointsToCreate.AddDefaulted_GetRef().Transform = Transform;
+		}
+		return Node;
+	}
+#endif // WITH_EDITOR
+
 	FTestData::FTestData(int32 RandomSeed, UPCGSettings* DefaultSettings, TSubclassOf<AActor> ActorClass)
 		: Settings(DefaultSettings)
 		, Seed(RandomSeed)
@@ -295,6 +400,208 @@ namespace PCGTestsCommon
 		// Transform checks with epsilon
 		return FirstPoint.Transform.Equals(SecondPoint.Transform);
 	}
+
+#if WITH_EDITOR
+	bool ValidateSingleTaggedData(FAutomationTestBase* Test, const FPCGDataCollection& DataCollection, TFunctionRef<bool(const FPCGTaggedData&)> Validator)
+	{
+		if (!Test->TestEqual(TEXT("Data collection contains a single data item"), DataCollection.TaggedData.Num(), 1))
+		{
+			return false;
+		}
+		return Validator(DataCollection.TaggedData[0]);
+	}
+
+	bool CompareMetadata(FAutomationTestBase* Test, const UPCGMetadata* ExpMeta, const UPCGMetadata* ActMeta, int32 NumEntries)
+	{
+		// Both null is fine — no metadata to compare.
+		if (!ExpMeta && !ActMeta)
+		{
+			return true;
+		}
+
+		if (!Test->TestNotNull(TEXT("Expected metadata"), ExpMeta) || !Test->TestNotNull(TEXT("Actual metadata"), ActMeta))
+		{
+			return false;
+		}
+
+		TArray<FName> ExpNames, ActNames;
+		TArray<EPCGMetadataTypes> ExpTypes, ActTypes;
+		ExpMeta->GetAttributes(ExpNames, ExpTypes);
+		ActMeta->GetAttributes(ActNames, ActTypes);
+
+		if (!Test->TestEqual(TEXT("Metadata attribute count"), ActNames.Num(), ExpNames.Num()))
+		{
+			return false;
+		}
+
+		bool bPassed = true;
+
+		for (int32 AttrIdx = 0; AttrIdx < ExpNames.Num(); ++AttrIdx)
+		{
+			const FName& AttrName = ExpNames[AttrIdx];
+			const int32 ActIdx = ActNames.IndexOfByKey(AttrName);
+
+			if (!Test->TestTrue(*FString::Printf(TEXT("Attribute '%s' exists in actual"), *AttrName.ToString()), ActIdx != INDEX_NONE))
+			{
+				bPassed = false;
+				continue;
+			}
+
+			if (!Test->TestEqual(*FString::Printf(TEXT("Attribute '%s' type"), *AttrName.ToString()), static_cast<uint8>(ActTypes[ActIdx]), static_cast<uint8>(ExpTypes[AttrIdx])))
+			{
+				bPassed = false;
+				continue;
+			}
+
+			const FPCGMetadataAttributeBase* ExpAttr = ExpMeta->GetConstAttribute(AttrName);
+			const FPCGMetadataAttributeBase* ActAttr = ActMeta->GetConstAttribute(AttrName);
+
+			if (!Test->TestNotNull(*FString::Printf(TEXT("Attribute '%s' expected base"), *AttrName.ToString()), ExpAttr)
+				|| !Test->TestNotNull(*FString::Printf(TEXT("Attribute '%s' actual base"), *AttrName.ToString()), ActAttr))
+			{
+				bPassed = false;
+				continue;
+			}
+
+			// Type-dispatch to compare values per entry.
+			bPassed &= PCGMetadataAttribute::CallbackWithRightType(ExpAttr->GetTypeId(),
+				[&](auto Dummy) -> bool
+				{
+					using T = decltype(Dummy);
+					const auto* TypedExp = static_cast<const FPCGMetadataAttribute<T>*>(ExpAttr);
+					const auto* TypedAct = static_cast<const FPCGMetadataAttribute<T>*>(ActAttr);
+					bool bAttrPassed = true;
+
+					for (int32 Entry = 0; Entry < NumEntries; ++Entry)
+					{
+						const T ExpVal = TypedExp->GetValueFromItemKey(static_cast<PCGMetadataEntryKey>(Entry));
+						const T ActVal = TypedAct->GetValueFromItemKey(static_cast<PCGMetadataEntryKey>(Entry));
+
+						if (!PCG::Private::MetadataTraits<T>::Equal(ExpVal, ActVal))
+						{
+							Test->AddError(FString::Printf(TEXT("[%d] Attribute '%s' value mismatch"), Entry, *AttrName.ToString()));
+							bAttrPassed = false;
+						}
+					}
+
+					return bAttrPassed;
+				});
+		}
+
+		return bPassed;
+	}
+
+	bool CompareBasePointData(FAutomationTestBase* Test, const UPCGBasePointData* Expected, const UPCGBasePointData* Actual)
+	{
+		if (!Test->TestNotNull(TEXT("Expected point data"), Expected) || !Test->TestNotNull(TEXT("Actual point data"), Actual))
+		{
+			return false;
+		}
+
+		const int32 NumPoints = Expected->GetNumPoints();
+		if (!Test->TestEqual(TEXT("Point count"), Actual->GetNumPoints(), NumPoints))
+		{
+			return false;
+		}
+
+		if (NumPoints == 0)
+		{
+			return true;
+		}
+
+		bool bPassed = true;
+		const FConstPCGPointValueRanges E(Expected);
+		const FConstPCGPointValueRanges A(Actual);
+
+		for (int32 i = 0; i < NumPoints; ++i)
+		{
+			// Position: world units (cm), float vs double precision loss
+			bPassed &= Test->TestNearlyEqual(*FString::Printf(TEXT("[%d] Position"), i), A.TransformRange[i].GetLocation(), E.TransformRange[i].GetLocation(), 1.0);
+
+			// Rotation: quaternion dot product (1.0 = identical, 0.001 tolerance ≈ 2.6 degrees)
+			const FQuat ERot = E.TransformRange[i].GetRotation();
+			const FQuat ARot = A.TransformRange[i].GetRotation();
+			const double RotDiff = 1.0 - FMath::Abs(ERot | ARot);
+			bPassed &= Test->TestTrue(*FString::Printf(TEXT("[%d] Rotation (dot diff %.6f)"), i, RotDiff), RotDiff < 0.001);
+
+			bPassed &= Test->TestNearlyEqual(*FString::Printf(TEXT("[%d] Scale"), i), A.TransformRange[i].GetScale3D(), E.TransformRange[i].GetScale3D(), 0.01);
+			bPassed &= Test->TestNearlyEqual(*FString::Printf(TEXT("[%d] Density"), i), static_cast<double>(A.DensityRange[i]), static_cast<double>(E.DensityRange[i]), 0.01);
+			bPassed &= Test->TestEqual(*FString::Printf(TEXT("[%d] Seed"), i), A.SeedRange[i], E.SeedRange[i]);
+			bPassed &= Test->TestNearlyEqual(*FString::Printf(TEXT("[%d] BoundsMin"), i), A.BoundsMinRange[i], E.BoundsMinRange[i], 0.01);
+			bPassed &= Test->TestNearlyEqual(*FString::Printf(TEXT("[%d] BoundsMax"), i), A.BoundsMaxRange[i], E.BoundsMaxRange[i], 0.01);
+			bPassed &= Test->TestNearlyEqual(*FString::Printf(TEXT("[%d] Steepness"), i), static_cast<double>(A.SteepnessRange[i]), static_cast<double>(E.SteepnessRange[i]), 0.01);
+
+			const FVector4& EColor = E.ColorRange[i];
+			const FVector4& AColor = A.ColorRange[i];
+			bPassed &= Test->TestNearlyEqual(*FString::Printf(TEXT("[%d] Color R"), i), AColor.X, EColor.X, 0.01);
+			bPassed &= Test->TestNearlyEqual(*FString::Printf(TEXT("[%d] Color G"), i), AColor.Y, EColor.Y, 0.01);
+			bPassed &= Test->TestNearlyEqual(*FString::Printf(TEXT("[%d] Color B"), i), AColor.Z, EColor.Z, 0.01);
+			bPassed &= Test->TestNearlyEqual(*FString::Printf(TEXT("[%d] Color A"), i), AColor.W, EColor.W, 0.01);
+		}
+
+		// Metadata attributes
+		bPassed &= CompareMetadata(Test, Expected->Metadata.Get(), Actual->Metadata.Get(), NumPoints);
+
+		return bPassed;
+	}
+
+	bool CompareDataCollections(FAutomationTestBase* Test, const FPCGDataCollection& Expected, const FPCGDataCollection& Actual, FTaggedDataComparator InComparator)
+	{
+		if (!Test->TestEqual(TEXT("Tagged data count"), Actual.TaggedData.Num(), Expected.TaggedData.Num()))
+		{
+			return false;
+		}
+
+		bool bPassed = true;
+
+		for (int32 i = 0; i < Expected.TaggedData.Num(); ++i)
+		{
+			const FPCGTaggedData& ExpectedEntry = Expected.TaggedData[i];
+			const FPCGTaggedData& ActualEntry = Actual.TaggedData[i];
+
+			if (!Test->TestNotNull(*FString::Printf(TEXT("Data[%d] expected"), i), ExpectedEntry.Data.Get()))
+			{
+				bPassed = false;
+				continue;
+			}
+			if (!Test->TestNotNull(*FString::Printf(TEXT("Data[%d] actual"), i), ActualEntry.Data.Get()))
+			{
+				bPassed = false;
+				continue;
+			}
+
+			if (!Test->TestEqual(*FString::Printf(TEXT("Data[%d] type"), i), ActualEntry.Data->GetClass()->GetName(), ExpectedEntry.Data->GetClass()->GetName()))
+			{
+				bPassed = false;
+				continue;
+			}
+
+			bPassed &= InComparator(Test, i, ExpectedEntry, ActualEntry);
+		}
+
+		return bPassed;
+	}
+
+	bool ComparePointDataCollections(FAutomationTestBase* Test, const FPCGDataCollection& Expected, const FPCGDataCollection& Actual)
+	{
+		return CompareDataCollections(Test, Expected, Actual, [](FAutomationTestBase* Test, int32 Index, const FPCGTaggedData& ExpectedEntry, const FPCGTaggedData& ActualEntry)
+		{
+			const UPCGBasePointData* ExpectedPoints = Cast<UPCGBasePointData>(ExpectedEntry.Data);
+			const UPCGBasePointData* ActualPoints = Cast<UPCGBasePointData>(ActualEntry.Data);
+
+			if (!Test->TestNotNull(*FString::Printf(TEXT("Data[%d] expected is point data"), Index), ExpectedPoints))
+			{
+				return false;
+			}
+			if (!Test->TestNotNull(*FString::Printf(TEXT("Data[%d] actual is point data"), Index), ActualPoints))
+			{
+				return false;
+			}
+
+			return CompareBasePointData(Test, ExpectedPoints, ActualPoints);
+		});
+	}
+#endif // WITH_EDITOR
 }
 
 bool FPCGTestBaseClass::SmokeTestAnyValidInput(UPCGSettings* InSettings, TFunction<bool(const FPCGDataCollection&, const FPCGDataCollection&)> ValidationFn)
diff --git a/Engine/Plugins/PCG/Source/PCG/Public/PCGSettings.h b/Engine/Plugins/PCG/Source/PCG/Public/PCGSettings.h
--- a/Engine/Plugins/PCG/Source/PCG/Public/PCGSettings.h
+++ b/Engine/Plugins/PCG/Source/PCG/Public/PCGSettings.h
@@ -613,6 +613,9 @@ public:
 	UE_API virtual bool ShouldExecuteOnGPU() const;
 
 #if WITH_EDITOR
+	/** Set whether this node should be executed on the GPU. */
+	void SetExecuteOnGPU(bool bValue) { bExecuteOnGPU = bValue; }
+
 	/** Create the GPU kernels and kernel edges that should execute for this node. */
 	virtual void CreateKernels(FPCGGPUCompilationContext& InOutContext, UObject* InObjectOuter, TArray<UPCGComputeKernel*>& OutKernels, TArray<FPCGKernelEdge>& OutEdges) const {};
 
diff --git a/Engine/Plugins/PCG/Source/PCG/Public/Tests/PCGTestsCommon.h b/Engine/Plugins/PCG/Source/PCG/Public/Tests/PCGTestsCommon.h
--- a/Engine/Plugins/PCG/Source/PCG/Public/Tests/PCGTestsCommon.h
+++ b/Engine/Plugins/PCG/Source/PCG/Public/Tests/PCGTestsCommon.h
@@ -3,11 +3,13 @@
 #pragma once
 
 #include "PCGData.h"
+#include "PCGGraph.h"
 #include "PCGPoint.h"
 #include "PCGSettings.h"
 #include "Data/PCGBasePointData.h"
 #include "Data/PCGPointData.h"
 #include "Data/PCGPointArrayData.h"
+#include "Elements/PCGCreateAttribute.h"
 #include "Metadata/PCGMetadata.h"
 #include "Metadata/PCGMetadataAttributeTpl.h"
 #include "Metadata/PCGMetadataAttributeTraits.h"
@@ -17,8 +19,10 @@
 
 #define UE_API PCG_API
 
+class ALandscape;
 class IPCGElement;
 class UPCGComponent;
+class UPCGMetadata;
 class UPCGNode;
 class UPCGParamData;
 class UPCGPolyLineData;
@@ -56,6 +60,44 @@ namespace PCGTestsCommon
 		FRandomStream RandomStream;
 	};
 
+#if WITH_EDITOR
+	/** Spawn a small flat landscape (Z=0 everywhere) in the given world. */
+	ALandscape* CreateFlatLandscape(UWorld* World, const FTransform& Transform = FTransform::Identity);
+
+	/** Spawn a landscape with layered Perlin noise (hills, medium detail, small bumps). Deterministic for a given Seed. */
+	ALandscape* CreatePerlinLandscape(UWorld* World, const FTransform& Transform = FTransform::Identity, int32 Seed = 0);
+
+	/** Add a CreatePoints node to a graph with the given transforms (world space). */
+	UPCGNode* AddCreatePointsNode(UPCGGraph* Graph, TArrayView<const FTransform> Transforms);
+
+	/** Add a CreateAttributeSet node to a graph. Sets the type and value from the C++ type. */
+	template <typename T>
+	UPCGNode* AddCreateAttributeSetNode(UPCGGraph* Graph, FName AttributeName, const T& Value)
+	{
+		UPCGCreateAttributeSetSettings* Settings = nullptr;
+		UPCGNode* Node = Graph->AddNodeOfType<UPCGCreateAttributeSetSettings>(Settings);
+		Settings->OutputTarget.SetAttributeName(AttributeName);
+
+		FPCGMetadataTypesConstantStruct& Attr = Settings->AttributeTypes;
+		if constexpr (std::is_same_v<T, bool>)            { Attr.Type = EPCGMetadataTypes::Boolean;    Attr.BoolValue = Value; }
+		else if constexpr (std::is_same_v<T, int32>)      { Attr.Type = EPCGMetadataTypes::Integer32;  Attr.Int32Value = Value; }
+		else if constexpr (std::is_same_v<T, int64>)      { Attr.Type = EPCGMetadataTypes::Integer64;  Attr.IntValue = Value; }
+		else if constexpr (std::is_same_v<T, float>)      { Attr.Type = EPCGMetadataTypes::Float;      Attr.FloatValue = Value; }
+		else if constexpr (std::is_same_v<T, double>)     { Attr.Type = EPCGMetadataTypes::Double;     Attr.DoubleValue = Value; }
+		else if constexpr (std::is_same_v<T, FVector2D>)  { Attr.Type = EPCGMetadataTypes::Vector2;    Attr.Vector2Value = Value; }
+		else if constexpr (std::is_same_v<T, FVector>)    { Attr.Type = EPCGMetadataTypes::Vector;     Attr.VectorValue = Value; }
+		else if constexpr (std::is_same_v<T, FVector4>)   { Attr.Type = EPCGMetadataTypes::Vector4;    Attr.Vector4Value = Value; }
+		else if constexpr (std::is_same_v<T, FQuat>)      { Attr.Type = EPCGMetadataTypes::Quaternion; Attr.QuatValue = Value; }
+		else if constexpr (std::is_same_v<T, FRotator>)   { Attr.Type = EPCGMetadataTypes::Rotator;    Attr.RotatorValue = Value; }
+		else if constexpr (std::is_same_v<T, FTransform>) { Attr.Type = EPCGMetadataTypes::Transform;  Attr.TransformValue = Value; }
+		else if constexpr (std::is_same_v<T, FString>)    { Attr.Type = EPCGMetadataTypes::String;     Attr.StringValue = Value; }
+		else if constexpr (std::is_same_v<T, FName>)      { Attr.Type = EPCGMetadataTypes::Name;       Attr.NameValue = Value; }
+		else { static_assert(sizeof(T) == 0, "Unsupported type for AddCreateAttributeSetNode"); }
+
+		return Node;
+	}
+#endif // WITH_EDITOR
+
 	AActor* CreateTemporaryActor();
 	UPCGParamData* CreateEmptyParamData();
 
@@ -140,6 +182,33 @@ namespace PCGTestsCommon
 	/** Validates that two Spatial Points are identical */
 	bool PointsAreIdentical(const FPCGPoint& FirstPoint, const FPCGPoint& SecondPoint);
 
+#if WITH_EDITOR
+	/**
+	 * Verifies DataCollection contains exactly one item, then calls Validator on it.
+	 * Returns false (with a logged failure) if the count check fails, otherwise returns
+	 * whatever Validator returns.
+	 */
+	bool ValidateSingleTaggedData(FAutomationTestBase* Test, const FPCGDataCollection& DataCollection, TFunctionRef<bool(const FPCGTaggedData&)> Validator);
+
+	/** Per-entry comparator for CompareDataCollections. Index is the TaggedData index. */
+	using FTaggedDataComparator = TFunctionRef<bool(FAutomationTestBase* Test, int32 Index, const FPCGTaggedData& Expected, const FPCGTaggedData& Actual)>;
+
+	/**
+	 * Compare two data collections structurally (same tagged data count, matching data types per entry),
+	 * then delegate per-entry content comparison to InComparator.
+	 */
+	bool CompareDataCollections(FAutomationTestBase* Test, const FPCGDataCollection& Expected, const FPCGDataCollection& Actual, FTaggedDataComparator InComparator);
+
+	/** Convenience overload: compares data collections assuming all entries are point data. */
+	bool ComparePointDataCollections(FAutomationTestBase* Test, const FPCGDataCollection& Expected, const FPCGDataCollection& Actual);
+
+	/** Compare two UPCGBasePointData objects: native point properties and metadata attributes. */
+	bool CompareBasePointData(FAutomationTestBase* Test, const UPCGBasePointData* Expected, const UPCGBasePointData* Actual);
+
+	/** Compare metadata attributes. Checks attribute names, types, and per-entry values. */
+	bool CompareMetadata(FAutomationTestBase* Test, const UPCGMetadata* Expected, const UPCGMetadata* Actual, int32 NumEntries);
+#endif // WITH_EDITOR
+
 	/** Generates settings based upon a UPCGSettings subclass */
 	template<typename SettingsType>
 	SettingsType* GenerateSettings(FTestData& TestData, TFunction<void(FTestData&)> ExtraSettingsDelegate = nullptr)
