[PCG] Fix C6011 null-pointer warning in PCGTestsCommon

#rnx
#Synced-CL 52429529

diff --git a/Engine/Plugins/PCG/Source/PCG/Private/Tests/PCGTestsCommon.cpp b/Engine/Plugins/PCG/Source/PCG/Private/Tests/PCGTestsCommon.cpp
--- a/Engine/Plugins/PCG/Source/PCG/Private/Tests/PCGTestsCommon.cpp
+++ b/Engine/Plugins/PCG/Source/PCG/Private/Tests/PCGTestsCommon.cpp
@@ -424,6 +424,12 @@ namespace PCGTestsCommon
 			return false;
 		}
 
+		// Explicit null guard to satisfy static analysis (C6011). The TestNotNull calls above already return false on null.
+		if (!ExpMeta || !ActMeta)
+		{
+			return false;
+		}
+
 		TArray<FName> ExpNames, ActNames;
 		TArray<EPCGMetadataTypes> ExpTypes, ActTypes;
 		ExpMeta->GetAttributes(ExpNames, ExpTypes);
