[PCG] Remove `FPCGMetadataAttributeGeneric`

* Now that the merge is done, simple renaming CL to remove reference to Generic
* Still some references related to the UI/widgets, that will be worked on soon

#rnx
#Synced-CL 52470086

#Synced-CL 52470086

diff --git a/Engine/Plugins/PCG/Source/PCG/Private/Metadata/Accessors/PCGAttributeAccessorHelpers.cpp b/Engine/Plugins/PCG/Source/PCG/Private/Metadata/Accessors/PCGAttributeAccessorHelpers.cpp
--- a/Engine/Plugins/PCG/Source/PCG/Private/Metadata/Accessors/PCGAttributeAccessorHelpers.cpp
+++ b/Engine/Plugins/PCG/Source/PCG/Private/Metadata/Accessors/PCGAttributeAccessorHelpers.cpp
@@ -358,13 +358,13 @@ TUniquePtr<const IPCGAttributeAccessor> PCGAttributeAccessorHelpers::CreateConst
 
 TUniquePtr<const IPCGAttributeAccessor> PCGAttributeAccessorHelpers::CreateConstAccessor(const FPCGMetadataAttributeBase* InAttribute, const FPCGMetadataDomain* InMetadata, bool bQuiet)
 {
-	if (!InMetadata || !InAttribute || !InAttribute->IsGeneric())
+	if (!InMetadata || !InAttribute)
 	{
 		return {};
 	}
 	else
 	{
-		return MakeUnique<const FPCGAttributeGenericAccessor>(static_cast<const FPCGMetadataAttributeGeneric*>(InAttribute), InMetadata);
+		return MakeUnique<const FPCGAttributeGenericAccessor>(InAttribute, InMetadata);
 	}
 }
 
@@ -430,13 +430,13 @@ TUniquePtr<IPCGAttributeAccessor> PCGAttributeAccessorHelpers::CreateAccessor(FP
 
 TUniquePtr<IPCGAttributeAccessor> PCGAttributeAccessorHelpers::CreateAccessor(FPCGMetadataAttributeBase* InAttribute, FPCGMetadataDomain* InMetadata, bool bQuiet)
 {
-	if (!InMetadata || !InAttribute || !InAttribute->IsGeneric())
+	if (!InMetadata || !InAttribute)
 	{
 		return {};
 	}
 	else
 	{
-		return MakeUnique<FPCGAttributeGenericAccessor>(static_cast<FPCGMetadataAttributeGeneric*>(InAttribute), InMetadata);
+		return MakeUnique<FPCGAttributeGenericAccessor>(InAttribute, InMetadata);
 	}
 }
 
diff --git a/Engine/Plugins/PCG/Source/PCG/Public/Metadata/Accessors/IPCGAttributeAccessor.h b/Engine/Plugins/PCG/Source/PCG/Public/Metadata/Accessors/IPCGAttributeAccessor.h
--- a/Engine/Plugins/PCG/Source/PCG/Public/Metadata/Accessors/IPCGAttributeAccessor.h
+++ b/Engine/Plugins/PCG/Source/PCG/Public/Metadata/Accessors/IPCGAttributeAccessor.h
@@ -150,7 +150,8 @@ public:
 
 	virtual void Prepare(IPCGAttributeAccessorKeys& Keys, int32 Count, const bool bCanReuseEntryKeys) {}
 
-	int16 GetUnderlyingType() const { return UnderlyingType; }
+	// Since underlying type was only working for Single Values (and not arrays or other containers), make sure we never return a valid value if this is not a single value.
+	int16 GetUnderlyingType() const { return UnderlyingDesc.IsSingleValue() ? UnderlyingType : static_cast<int16>(EPCGMetadataTypes::Unknown); }
 	const FPCGMetadataAttributeDesc& GetUnderlyingDesc() const { return UnderlyingDesc; }
 	bool IsReadOnly() const { return bReadOnly; }
 
diff --git a/Engine/Plugins/PCG/Source/PCG/Public/Metadata/Accessors/PCGAttributeAccessor.h b/Engine/Plugins/PCG/Source/PCG/Public/Metadata/Accessors/PCGAttributeAccessor.h
--- a/Engine/Plugins/PCG/Source/PCG/Public/Metadata/Accessors/PCGAttributeAccessor.h
+++ b/Engine/Plugins/PCG/Source/PCG/Public/Metadata/Accessors/PCGAttributeAccessor.h
@@ -5,7 +5,7 @@
 #include "IPCGAttributeAccessorTpl.h"
 
 #include "Metadata/PCGMetadata.h"
-#include "Metadata/PCGMetadataAttributeGeneric.h"
+#include "Metadata/PCGMetadataAttribute.h"
 #include "Metadata/PCGMetadataAttributeTpl.h"
 #include "Metadata/PCGMetadataDomain.h"
 
@@ -21,7 +21,7 @@ public:
 	using Super = IPCGAttributeAccessor;
 
 	// Can't write if metadata is null
-	FPCGAttributeGenericAccessor(FPCGMetadataAttributeGeneric* InAttribute, FPCGMetadataDomain* InMetadataDomain, bool bForceReadOnly = false)
+	FPCGAttributeGenericAccessor(FPCGMetadataAttributeBase* InAttribute, FPCGMetadataDomain* InMetadataDomain, bool bForceReadOnly = false)
 		: Super(/*bInReadOnly=*/ InMetadataDomain == nullptr || bForceReadOnly, static_cast<int16>(InAttribute->GetAttributeDesc().ValueType))
 		, Attribute(InAttribute)
 		, MetadataDomain(InMetadataDomain)
@@ -30,9 +30,9 @@ public:
 		UnderlyingDesc = InAttribute->GetAttributeDesc();
 	}
 
-	FPCGAttributeGenericAccessor(const FPCGMetadataAttributeGeneric* InAttribute, const FPCGMetadataDomain* InMetadataDomain, bool bForceReadOnly = false)
+	FPCGAttributeGenericAccessor(const FPCGMetadataAttributeBase* InAttribute, const FPCGMetadataDomain* InMetadataDomain, bool bForceReadOnly = false)
 		: Super(/*bInReadOnly=*/ true, static_cast<int16>(InAttribute->GetAttributeDesc().ValueType))
-		, Attribute(const_cast<FPCGMetadataAttributeGeneric*>(InAttribute))
+		, Attribute(const_cast<FPCGMetadataAttributeBase*>(InAttribute))
 		, MetadataDomain(const_cast<FPCGMetadataDomain*>(InMetadataDomain))
 	{
 		check(InAttribute && (!InMetadataDomain || InAttribute->GetMetadataDomain() == InMetadataDomain));
@@ -190,7 +190,7 @@ public:
 	}
 
 protected:
-	FPCGMetadataAttributeGeneric* Attribute = nullptr;
+	FPCGMetadataAttributeBase* Attribute = nullptr;
 	FPCGMetadataDomain* MetadataDomain = nullptr;
 	bool bWasPrepared = false;
 	int32 StartIndex = INDEX_NONE;
diff --git a/Engine/Plugins/PCG/Source/PCG/Public/Metadata/PCGMetadataAttribute.h b/Engine/Plugins/PCG/Source/PCG/Public/Metadata/PCGMetadataAttribute.h
--- a/Engine/Plugins/PCG/Source/PCG/Public/Metadata/PCGMetadataAttribute.h
+++ b/Engine/Plugins/PCG/Source/PCG/Public/Metadata/PCGMetadataAttribute.h
@@ -106,9 +106,6 @@ public:
 		return PCG::Private::IsEquivalentDesc<T>(CachedDesc);
 	}
 
-	// @todo_pcg All generic methods to be removed in a follow-up CL. Will ease the refactor review.
-	bool IsGeneric() const { return true; }
-
 	/** Unparents current attribute by flattening the values, entries, etc. */
 	UE_API virtual void Flatten();
 	/** Unparents current attribute by flattening the values, entries, etc while only keeping the entries referenced in InEntryKeysToKeep. There must be NO invalid entry keys. */
diff --git a/Engine/Plugins/PCG/Source/PCG/Public/Metadata/PCGMetadataAttributeGeneric.h b/Engine/Plugins/PCG/Source/PCG/Public/Metadata/PCGMetadataAttributeGeneric.h
deleted file mode 100644
diff --git a/Engine/Plugins/PCG/Source/PCG/Public/Metadata/PCGMetadataAttributeTpl.h b/Engine/Plugins/PCG/Source/PCG/Public/Metadata/PCGMetadataAttributeTpl.h
--- a/Engine/Plugins/PCG/Source/PCG/Public/Metadata/PCGMetadataAttributeTpl.h
+++ b/Engine/Plugins/PCG/Source/PCG/Public/Metadata/PCGMetadataAttributeTpl.h
@@ -6,7 +6,7 @@
 
 #include "PCGModule.h"
 #include "Helpers/PCGMetadataHelpers.h"
-#include "Metadata/PCGMetadataAttributeGeneric.h"
+#include "Metadata/PCGMetadataAttribute.h"
 #include "Metadata/PCGMetadataAttributeTraits.h"
 #include "Metadata/PCGMetadataCommon.h"
 
@@ -142,7 +142,7 @@ protected:
 			}
 		};
 		
-		FPCGMetadataAttributeGeneric::CopyInternal(AttributeCopy, bKeepParent, bCopyEntries, bCopyValues, CopyFunctor);
+		FPCGMetadataAttributeBase::CopyInternal(AttributeCopy, bKeepParent, bCopyEntries, bCopyValues, CopyFunctor);
 
 		return AttributeCopy;
 	}
@@ -160,19 +160,12 @@ public:
 		{
 			check(ItemKey != PCGInvalidEntryKey);
 
-			if (InAttribute && !InAttribute->IsGeneric())
+			if (InAttribute && !InAttribute->IsOfType<T>())
 			{
 				return;
 			}
 
-			const FPCGMetadataAttributeGeneric* Attribute = static_cast<const FPCGMetadataAttributeGeneric*>(InAttribute);
-		
-			if (Attribute && !Attribute->IsOfType<T>())
-			{
-				return;
-			}
-		
-			Accumulate(ItemKey, Attribute, InEntryKey, Weight);
+			Accumulate(ItemKey, InAttribute, InEntryKey, Weight);
 		}
 	}
 
@@ -182,19 +175,12 @@ public:
 		{
 			check(ItemKey != PCGInvalidEntryKey);
 
-			if (InAttribute && !InAttribute->IsGeneric())
+			if (InAttribute && !InAttribute->IsOfType<T>())
 			{
 				return;
 			}
 
-			const FPCGMetadataAttributeGeneric* Attribute = static_cast<const FPCGMetadataAttributeGeneric*>(InAttribute);
-		
-			if (Attribute && !Attribute->IsOfType<T>())
-			{
-				return;
-			}
-		
-			Accumulate(ItemKey, Attribute, InWeightedKeys);
+			Accumulate(ItemKey, InAttribute, InWeightedKeys);
 		}
 	}
 	
@@ -210,76 +196,68 @@ public:
 		check(ItemKey != PCGInvalidEntryKey);
 		bool bAppliedValue = false;
 
-		if ((InAttributeA && !InAttributeA->IsGeneric()) || (InAttributeB && !InAttributeB->IsGeneric()))
+		if ((InAttributeA && !InAttributeA->IsOfType<T>()) || (InAttributeB && !InAttributeB->IsOfType<T>()))
 		{
 			return;
 		}
 
-		const FPCGMetadataAttributeGeneric* AttributeA = static_cast<const FPCGMetadataAttributeGeneric*>(InAttributeA);
-		const FPCGMetadataAttributeGeneric* AttributeB = static_cast<const FPCGMetadataAttributeGeneric*>(InAttributeB);
-		
-		if ((AttributeA && !AttributeA->IsOfType<T>()) || (AttributeB && !AttributeB->IsOfType<T>()))
-		{
-			return;
-		}
-		
-		if (Op == EPCGMetadataOp::TargetValue && AttributeB)
+		if (Op == EPCGMetadataOp::TargetValue && InAttributeB)
 		{
 			// Take value of second attribute.
-			if (AttributeB == this)
+			if (InAttributeB == this)
 			{
 				SetValueFromValueKey(ItemKey, GetValueKey(InEntryKeyB));
 			}
 			else
 			{
-				SetValue(ItemKey, AttributeB->GetValueFromItemKey<T>(InEntryKeyB));
+				SetValue(ItemKey, InAttributeB->GetValueFromItemKey<T>(InEntryKeyB));
 			}
 
 			bAppliedValue = true;
 		}
-		else if (Op == EPCGMetadataOp::SourceValue && AttributeA)
+		else if (Op == EPCGMetadataOp::SourceValue && InAttributeA)
 		{
 			// Take value of first attribute.
-			if (AttributeA == this)
+			if (InAttributeA == this)
 			{
 				SetValueFromValueKey(ItemKey, GetValueKey(InEntryKeyA));
 			}
 			else
 			{
-				SetValue(ItemKey, AttributeA->GetValueFromItemKey<T>(InEntryKeyA));
+				SetValue(ItemKey, InAttributeA->GetValueFromItemKey<T>(InEntryKeyA));
 			}
 
 			bAppliedValue = true;
 		}
-		else if (AttributeA && AttributeB && bAllowsInterpolation)
+		else if (InAttributeA && InAttributeB && bAllowsInterpolation)
 		{
 			// Combine attributes using specified operation.
 			if (Op == EPCGMetadataOp::Min)
 			{
-				bAppliedValue = SetMin(ItemKey, AttributeA, InEntryKeyA, AttributeB, InEntryKeyB);
+				bAppliedValue = SetMin(ItemKey, InAttributeA, InEntryKeyA, InAttributeB, InEntryKeyB);
 			}
 			else if (Op == EPCGMetadataOp::Max)
 			{
-				bAppliedValue = SetMax(ItemKey, AttributeA, InEntryKeyA, AttributeB, InEntryKeyB);
+				bAppliedValue = SetMax(ItemKey, InAttributeA, InEntryKeyA, InAttributeB, InEntryKeyB);
 			}
 			else if (Op == EPCGMetadataOp::Sub)
 			{
-				bAppliedValue = SetSub(ItemKey, AttributeA, InEntryKeyA, AttributeB, InEntryKeyB);
+				bAppliedValue = SetSub(ItemKey, InAttributeA, InEntryKeyA, InAttributeB, InEntryKeyB);
 			}
 			else if (Op == EPCGMetadataOp::Add)
 			{
-				bAppliedValue = SetAdd(ItemKey, AttributeA, InEntryKeyA, AttributeB, InEntryKeyB);
+				bAppliedValue = SetAdd(ItemKey, InAttributeA, InEntryKeyA, InAttributeB, InEntryKeyB);
 			}
 			else if (Op == EPCGMetadataOp::Mul)
 			{
-				bAppliedValue = SetMul(ItemKey, AttributeA, InEntryKeyA, AttributeB, InEntryKeyB);
+				bAppliedValue = SetMul(ItemKey, InAttributeA, InEntryKeyA, InAttributeB, InEntryKeyB);
 			}
 			else if (Op == EPCGMetadataOp::Div)
 			{
-				bAppliedValue = SetDiv(ItemKey, AttributeA, InEntryKeyA, AttributeB, InEntryKeyB);
+				bAppliedValue = SetDiv(ItemKey, InAttributeA, InEntryKeyA, InAttributeB, InEntryKeyB);
 			}
 		}
-		else if (AttributeA && AttributeB && HasNonDefaultValue(ItemKey))
+		else if (InAttributeA && InAttributeB && HasNonDefaultValue(ItemKey))
 		{
 			// In this case, the current already has a value, in which case we should not update it
 			bAppliedValue = true;
@@ -289,26 +267,26 @@ public:
 		{
 			// Nothing to do
 		}
-		else if (AttributeA)
+		else if (InAttributeA)
 		{
-			if (AttributeA == this)
+			if (InAttributeA == this)
 			{
 				SetValueFromValueKey(ItemKey, GetValueKey(InEntryKeyA));
 			}
 			else
 			{
-				SetValue(ItemKey, AttributeA->GetValueFromItemKey<T>(InEntryKeyA));
+				SetValue(ItemKey, InAttributeA->GetValueFromItemKey<T>(InEntryKeyA));
 			}
 		}
-		else if (AttributeB)
+		else if (InAttributeB)
 		{
-			if (AttributeB == this)
+			if (InAttributeB == this)
 			{
 				SetValueFromValueKey(ItemKey, GetValueKey(InEntryKeyB));
 			}
 			else
 			{
-				SetValue(ItemKey, AttributeB->GetValueFromItemKey<T>(InEntryKeyB));
+				SetValue(ItemKey, InAttributeB->GetValueFromItemKey<T>(InEntryKeyB));
 			}
 		}
 	}
@@ -316,38 +294,38 @@ public:
 	/** Adds the value, returns the value key for the given value */
 	PCGMetadataValueKey AddValue(const T& InValue)
 	{
-		TArray<PCGMetadataValueKey> Result = FPCGMetadataAttributeGeneric::AddValues<T>(MakeArrayView<const T>(&InValue, 1));
+		TArray<PCGMetadataValueKey> Result = FPCGMetadataAttributeBase::AddValues<T>(MakeArrayView<const T>(&InValue, 1));
 		return !Result.IsEmpty() ? Result[0] : PCGDefaultValueKey; 
 	}
 
 	TArray<PCGMetadataValueKey> AddValues(const TArrayView<const T>& InValues)
 	{
-		return FPCGMetadataAttributeGeneric::AddValues<T>(InValues);
+		return FPCGMetadataAttributeBase::AddValues<T>(InValues);
 	}
 
 	void SetValue(PCGMetadataEntryKey ItemKey, const T& InValue)
 	{
-		FPCGMetadataAttributeGeneric::SetValue<T>(ItemKey, InValue);
+		FPCGMetadataAttributeBase::SetValue<T>(ItemKey, InValue);
 	}
 
 	void SetValues(const TArrayView<const PCGMetadataEntryKey>& ItemKeys, const TArrayView<const T>& InValues)
 	{
-		FPCGMetadataAttributeGeneric::SetValues<T>(ItemKeys, InValues);
+		FPCGMetadataAttributeBase::SetValues<T>(ItemKeys, InValues);
 	}
 
 	void SetValues(const TArrayView<const PCGMetadataEntryKey * const>& ItemKeys, const TArrayView<const T>& InValues)
 	{
-		FPCGMetadataAttributeGeneric::SetValues<T>(ItemKeys, InValues);
+		FPCGMetadataAttributeBase::SetValues<T>(ItemKeys, InValues);
 	}
 	
 	void SetValues(TConstPCGValueRange<PCGMetadataEntryKey> ItemKeys, const TArrayView<const T>& InValues)
 	{
-		FPCGMetadataAttributeGeneric::SetValues<T>(ItemKeys, InValues);
+		FPCGMetadataAttributeBase::SetValues<T>(ItemKeys, InValues);
 	}
 	
 	void SetValues(TPCGValueRange<PCGMetadataEntryKey> ItemKeys, const TArrayView<const T>& InValues)
 	{
-		FPCGMetadataAttributeGeneric::SetValues<T>(ItemKeys, InValues);
+		FPCGMetadataAttributeBase::SetValues<T>(ItemKeys, InValues);
 	}
 
 	template<typename U> requires (!std::is_same_v<T, U> && PCG::Private::IsBroadcastableOrConstructible<U, T>())
@@ -355,25 +333,25 @@ public:
 	{
 		T Value{};
 		PCG::Private::GetValueWithBroadcastAndConstructible(InValue, Value);
-		FPCGMetadataAttributeGeneric::SetValue<T>(ItemKey, Value);
+		FPCGMetadataAttributeBase::SetValue<T>(ItemKey, Value);
 	}
 	
 	// Special version for enum classes, as enums are not directly constructible
 	template<typename U> requires (!std::is_same_v<T, U> && !PCG::Private::IsBroadcastableOrConstructible<U, T>() && std::is_enum_v<U> && std::is_convertible_v<std::underlying_type_t<U>, T>)
 	void SetValue(PCGMetadataEntryKey ItemKey, const U& InValue)
 	{
-		FPCGMetadataAttributeGeneric::SetValue<T>(ItemKey, static_cast<T>(InValue));
+		FPCGMetadataAttributeBase::SetValue<T>(ItemKey, static_cast<T>(InValue));
 	}
 
 	T GetValueFromItemKey(PCGMetadataEntryKey ItemKey) const
 	{
-		return FPCGMetadataAttributeGeneric::GetValueFromItemKey<T>(ItemKey);
+		return FPCGMetadataAttributeBase::GetValueFromItemKey<T>(ItemKey);
 	}
 
 	T GetValue(PCGMetadataValueKey ValueKey) const
 	{
 		T Value{};
-		FPCGMetadataAttributeGeneric::GetValues<T>(MakeConstArrayView(&ValueKey, 1), MakeArrayView(&Value, 1));
+		FPCGMetadataAttributeBase::GetValues<T>(MakeConstArrayView(&ValueKey, 1), MakeArrayView(&Value, 1));
 		return Value;
 	}
 
@@ -383,7 +361,7 @@ public:
 	*/
 	void GetValues(const TArrayView<const PCGMetadataValueKey> ValueKeys, TArrayView<T> OutValues) const
 	{
-		return FPCGMetadataAttributeGeneric::GetValues<T>(ValueKeys, OutValues);
+		return FPCGMetadataAttributeBase::GetValues<T>(ValueKeys, OutValues);
 	}
 
 	/** 
@@ -394,12 +372,12 @@ public:
 	*/
 	void GetValuesFromItemKeys(TConstArrayView<PCGMetadataEntryKey> EntryKeys, TArrayView<T> OutValues) const
 	{
-		FPCGMetadataAttributeGeneric::GetValuesFromItemKeys<T>(EntryKeys, OutValues);
+		FPCGMetadataAttributeBase::GetValuesFromItemKeys<T>(EntryKeys, OutValues);
 	}
 	
 	void GetValuesFromItemKeys(TConstPCGValueRange<PCGMetadataEntryKey> EntryKeys, TArrayView<T> OutValues) const
 	{
-		FPCGMetadataAttributeGeneric::GetValuesFromItemKeys<T>(EntryKeys, OutValues);
+		FPCGMetadataAttributeBase::GetValuesFromItemKeys<T>(EntryKeys, OutValues);
 	}
 
 	/** 
@@ -409,7 +387,7 @@ public:
 	*/
 	void GetValuesFromItemKeys(TArrayView<PCGMetadataEntryKey> EntryKeys, TArrayView<T> OutValues) const
 	{
-		FPCGMetadataAttributeGeneric::GetValuesFromItemKeys<T>(EntryKeys, OutValues);
+		FPCGMetadataAttributeBase::GetValuesFromItemKeys<T>(EntryKeys, OutValues);
 	}
 	
 	PCGMetadataValueKey FindValue(const T& InValue) const 
@@ -430,7 +408,7 @@ public:
 	{
 		if constexpr (PCG::Private::MetadataTraits<T>::CompressData)
 		{
-			return FPCGMetadataAttributeGeneric::FindValues(OutValueKeys, PCG::Private::FInValues{TInPlaceType<PCG::Private::FInValuesByValue>{}, InValues.GetData(), InValues.Num()}, InValues.Num());
+			return FPCGMetadataAttributeBase::FindValues(OutValueKeys, PCG::Private::FInValues{TInPlaceType<PCG::Private::FInValuesByValue>{}, InValues.GetData(), InValues.Num()}, InValues.Num());
 		}
 		else
 		{
@@ -441,12 +419,12 @@ public:
 
 	void SetValues_TryLockless(TArrayView<PCGMetadataEntryKey*> EntryKeys, TArrayView<const T> InValues, int32 StartIndex)
 	{
-		FPCGMetadataAttributeGeneric::SetValues_TryLockless<T>(EntryKeys, InValues, StartIndex);
+		FPCGMetadataAttributeBase::SetValues_TryLockless<T>(EntryKeys, InValues, StartIndex);
 	}
 
 protected:
 	/** Code related to computing compared values (min, max, sub, add) */
-	bool SetMin(PCGMetadataEntryKey ItemKey, const FPCGMetadataAttributeGeneric* InAttributeA, PCGMetadataEntryKey InEntryKeyA, const FPCGMetadataAttributeGeneric* InAttributeB, PCGMetadataEntryKey InEntryKeyB)
+	bool SetMin(PCGMetadataEntryKey ItemKey, const FPCGMetadataAttributeBase* InAttributeA, PCGMetadataEntryKey InEntryKeyA, const FPCGMetadataAttributeBase* InAttributeB, PCGMetadataEntryKey InEntryKeyB)
 	{
 		if constexpr (PCG::Private::MetadataTraits<T>::CanMinMax)
 		{
@@ -463,7 +441,7 @@ protected:
 		}
 	}
 	
-	bool SetMax(PCGMetadataEntryKey ItemKey, const FPCGMetadataAttributeGeneric* InAttributeA, PCGMetadataEntryKey InEntryKeyA, const FPCGMetadataAttributeGeneric* InAttributeB, PCGMetadataEntryKey InEntryKeyB)
+	bool SetMax(PCGMetadataEntryKey ItemKey, const FPCGMetadataAttributeBase* InAttributeA, PCGMetadataEntryKey InEntryKeyA, const FPCGMetadataAttributeBase* InAttributeB, PCGMetadataEntryKey InEntryKeyB)
 	{
 		if constexpr (PCG::Private::MetadataTraits<T>::CanMinMax)
 		{
@@ -479,7 +457,7 @@ protected:
 		}
 	}
 	
-	bool SetAdd(PCGMetadataEntryKey ItemKey, const FPCGMetadataAttributeGeneric* InAttributeA, PCGMetadataEntryKey InEntryKeyA, const FPCGMetadataAttributeGeneric* InAttributeB, PCGMetadataEntryKey InEntryKeyB)
+	bool SetAdd(PCGMetadataEntryKey ItemKey, const FPCGMetadataAttributeBase* InAttributeA, PCGMetadataEntryKey InEntryKeyA, const FPCGMetadataAttributeBase* InAttributeB, PCGMetadataEntryKey InEntryKeyB)
 	{
 		if constexpr (PCG::Private::MetadataTraits<T>::CanSubAdd)
 		{
@@ -495,7 +473,7 @@ protected:
 		}
 	}
 	
-	bool SetSub(PCGMetadataEntryKey ItemKey, const FPCGMetadataAttributeGeneric* InAttributeA, PCGMetadataEntryKey InEntryKeyA, const FPCGMetadataAttributeGeneric* InAttributeB, PCGMetadataEntryKey InEntryKeyB)
+	bool SetSub(PCGMetadataEntryKey ItemKey, const FPCGMetadataAttributeBase* InAttributeA, PCGMetadataEntryKey InEntryKeyA, const FPCGMetadataAttributeBase* InAttributeB, PCGMetadataEntryKey InEntryKeyB)
 	{
 		if constexpr (PCG::Private::MetadataTraits<T>::CanSubAdd)
 		{
@@ -511,7 +489,7 @@ protected:
 		}
 	}
 	
-	bool SetMul(PCGMetadataEntryKey ItemKey, const FPCGMetadataAttributeGeneric* InAttributeA, PCGMetadataEntryKey InEntryKeyA, const FPCGMetadataAttributeGeneric* InAttributeB, PCGMetadataEntryKey InEntryKeyB)
+	bool SetMul(PCGMetadataEntryKey ItemKey, const FPCGMetadataAttributeBase* InAttributeA, PCGMetadataEntryKey InEntryKeyA, const FPCGMetadataAttributeBase* InAttributeB, PCGMetadataEntryKey InEntryKeyB)
 	{
 		if constexpr (PCG::Private::MetadataTraits<T>::CanMulDiv)
 		{
@@ -527,7 +505,7 @@ protected:
 		}
 	}
 	
-	bool SetDiv(PCGMetadataEntryKey ItemKey, const FPCGMetadataAttributeGeneric* InAttributeA, PCGMetadataEntryKey InEntryKeyA, const FPCGMetadataAttributeGeneric* InAttributeB, PCGMetadataEntryKey InEntryKeyB)
+	bool SetDiv(PCGMetadataEntryKey ItemKey, const FPCGMetadataAttributeBase* InAttributeA, PCGMetadataEntryKey InEntryKeyA, const FPCGMetadataAttributeBase* InAttributeB, PCGMetadataEntryKey InEntryKeyB)
 	{
 		if constexpr (PCG::Private::MetadataTraits<T>::CanMulDiv)
 		{
@@ -552,7 +530,7 @@ protected:
 		}
 	}
 
-	void Accumulate(PCGMetadataEntryKey ItemKey, const FPCGMetadataAttributeGeneric* InAttribute, PCGMetadataEntryKey InEntryKey, float Weight)
+	void Accumulate(PCGMetadataEntryKey ItemKey, const FPCGMetadataAttributeBase* InAttribute, PCGMetadataEntryKey InEntryKey, float Weight)
 	{
 		if constexpr (PCG::Private::MetadataTraits<T>::CanInterpolate)
 		{
@@ -573,7 +551,7 @@ protected:
 		}
 	}
 	
-	void Accumulate(PCGMetadataEntryKey ItemKey, const FPCGMetadataAttributeGeneric* InAttribute, const TArrayView<const TPair<PCGMetadataEntryKey, float>>& InWeightedKeys)
+	void Accumulate(PCGMetadataEntryKey ItemKey, const FPCGMetadataAttributeBase* InAttribute, const TArrayView<const TPair<PCGMetadataEntryKey, float>>& InWeightedKeys)
 	{
 		if constexpr (PCG::Private::MetadataTraits<T>::CanInterpolate)
 		{
diff --git a/Engine/Plugins/PCG/Source/PCG/Public/Metadata/PCGMetadataDomain.h b/Engine/Plugins/PCG/Source/PCG/Public/Metadata/PCGMetadataDomain.h
--- a/Engine/Plugins/PCG/Source/PCG/Public/Metadata/PCGMetadataDomain.h
+++ b/Engine/Plugins/PCG/Source/PCG/Public/Metadata/PCGMetadataDomain.h
@@ -3,7 +3,6 @@
 #pragma once
 
 #include "PCGCommon.h"
-#include "PCGMetadataAttributeGeneric.h"
 #include "Metadata/PCGMetadataCommon.h"
 #include "Metadata/PCGMetadataAttributeTpl.h"
 #include "Utils/PCGValueRange.h"
@@ -252,24 +251,23 @@ public:
 	*/
 	UE_API FPCGMetadataAttributeBase* CreateAttribute(const FPCGMetadataAttributeDesc& AttributeDesc, bool bAllowsInterpolation = true, bool bOverrideParent = true);
 
-	// @todo_pcg All generic methods to be removed in a follow-up CL. Will ease the refactor review.
 	template<typename T>
-	FPCGMetadataAttributeGeneric* CreateGenericAttribute(FName AttributeName, const T& DefaultValue, bool bAllowsInterpolation = true, bool bOverrideParent = false)
+	FPCGMetadataAttributeBase* CreateGenericAttribute(FName AttributeName, const T& DefaultValue, bool bAllowsInterpolation = true, bool bOverrideParent = false)
 	{
 		return CreateAttribute<T>(AttributeName, DefaultValue, bAllowsInterpolation, bOverrideParent);
 	}
 	
-	FPCGMetadataAttributeGeneric* CreateGenericAttribute(const FPCGMetadataAttributeDesc& Desc, bool bAllowsInterpolation = true, bool bOverrideParent = false)
+	FPCGMetadataAttributeBase* CreateGenericAttribute(const FPCGMetadataAttributeDesc& Desc, bool bAllowsInterpolation = true, bool bOverrideParent = false)
 	{
 		return CreateAttribute(Desc, bAllowsInterpolation, bOverrideParent);
 	}
 	
-	FPCGMetadataAttributeGeneric* GetMutableGenericAttribute(FName AttributeName)
+	FPCGMetadataAttributeBase* GetMutableGenericAttribute(FName AttributeName)
 	{
 		return GetMutableAttribute(AttributeName);
 	}
 	
-	const FPCGMetadataAttributeGeneric* GetConstGenericAttribute(FName AttributeName) const
+	const FPCGMetadataAttributeBase* GetConstGenericAttribute(FName AttributeName) const
 	{
 		return GetConstAttribute(AttributeName);
 	}
diff --git a/Engine/Plugins/PCG/Source/PCG/Public/PCGDataVisualization.h b/Engine/Plugins/PCG/Source/PCG/Public/PCGDataVisualization.h
--- a/Engine/Plugins/PCG/Source/PCG/Public/PCGDataVisualization.h
+++ b/Engine/Plugins/PCG/Source/PCG/Public/PCGDataVisualization.h
@@ -64,7 +64,7 @@ struct FPCGTableVisualizerColumnInfo
 	EPCGTableVisualizerCellAlignment CellAlignment = EPCGTableVisualizerCellAlignment::Right;
 	TSharedPtr<const IPCGAttributeAccessor> Accessor;
 	TSharedPtr<const IPCGAttributeAccessorKeys> AccessorKeys;
-	const FPCGMetadataAttributeBase* GenericAttribute = nullptr;
+	const FPCGMetadataAttributeBase* UnderlyingAttribute = nullptr;
 };
 
 struct FPCGTableVisualizerInfo
diff --git a/Engine/Plugins/PCG/Source/PCGEditor/Private/DataVisualizations/PCGDataVisualizationHelpers.cpp b/Engine/Plugins/PCG/Source/PCGEditor/Private/DataVisualizations/PCGDataVisualizationHelpers.cpp
--- a/Engine/Plugins/PCG/Source/PCGEditor/Private/DataVisualizations/PCGDataVisualizationHelpers.cpp
+++ b/Engine/Plugins/PCG/Source/PCGEditor/Private/DataVisualizations/PCGDataVisualizationHelpers.cpp
@@ -158,32 +158,26 @@ namespace PCGDataVisualizationHelpers
 		FPCGAttributePropertySelector Selector{};
 		Data->SetDomainFromDomainID(DomainID, Selector);
 		
-		auto AddInfo = [Data, &OutInfo, &Selector]<typename T>(T, const FName AttributeName)
+		auto AddInfo = [Data, &OutInfo, &Selector]<typename T>(T, const FName AttributeName) -> bool
 		{
 			Selector.SetAttributeName(AttributeName);
 			AddTypedColumnInfo<T>(OutInfo, Data, Selector);
+			return true;
 		};
 
 		for (int32 I = 0; I < AttributeNames.Num(); I++)
 		{
-			const FPCGMetadataAttributeBase* GenericAttr = MetadataDomain->GetConstAttribute(AttributeNames[I]);
-			const int32 ColumnInfosNum = OutInfo.ColumnInfos.Num();
+			const FPCGMetadataAttributeBase* Attribute = MetadataDomain->GetConstAttribute(AttributeNames[I]);
 
-			// If we don't have a generic attribute or the generic attribute is a single value (not in a container), try to add the columns for that type. 
-			if (!GenericAttr || GenericAttr->GetAttributeDesc().IsSingleValue())
-			{
-				PCGMetadataAttribute::CallbackWithRightType(static_cast<uint16>(AttributeTypes[I]), AddInfo, AttributeNames[I]);
-			}
+			const bool bColumnAdded = PCGMetadataAttribute::CallbackWithRightType(static_cast<uint16>(AttributeTypes[I]), AddInfo, AttributeNames[I]);
 
-			// If it failed (there was no column added), and we have a generic attribute, then configure the column to get the generic attribute.
-			// We don't want to do it for all attributes as we still want the basic types to be displayed as before, even if they are generic attributes
-			// and not typed attributes.
-			if (ColumnInfosNum == OutInfo.ColumnInfos.Num() && GenericAttr)
+			// If it failed (there was no column added), and we have an attribute, then configure the column to get the attribute with the generic viewer.
+			if (!bColumnAdded && Attribute)
 			{
 				FPCGTableVisualizerColumnInfo& ColumnInfo = OutInfo.ColumnInfos.Emplace_GetRef();
 				ColumnInfo.Label = FText::FromName(AttributeNames[I]);
 				ColumnInfo.Id = AttributeNames[I] == NAME_None ? TEXT("AttrNone") : AttributeNames[I];
-				ColumnInfo.GenericAttribute = GenericAttr;
+				ColumnInfo.UnderlyingAttribute = Attribute;
 				Selector.SetAttributeName(AttributeNames[I]);
 				ColumnInfo.AccessorKeys = TSharedPtr<const IPCGAttributeAccessorKeys>(PCGAttributeAccessorHelpers::CreateConstKeys(Data, Selector).Release());
 			}
diff --git a/Engine/Plugins/PCG/Source/PCGEditor/Private/Widgets/SPCGEditorGraphAttributeListView.cpp b/Engine/Plugins/PCG/Source/PCGEditor/Private/Widgets/SPCGEditorGraphAttributeListView.cpp
--- a/Engine/Plugins/PCG/Source/PCGEditor/Private/Widgets/SPCGEditorGraphAttributeListView.cpp
+++ b/Engine/Plugins/PCG/Source/PCGEditor/Private/Widgets/SPCGEditorGraphAttributeListView.cpp
@@ -143,19 +143,23 @@ TSharedRef<SWidget> SPCGListViewItemRow::GenerateWidgetForColumn(const FName& Co
 	// Make the text clickable for soft object paths
 	bool bIsClickable = false;
 	
-	// @todo_pcg: To be removed when we do not make the distinction between typed and generic.
-	bool bDataRetrievedWithAccessors = false;
-	TTuple<const FPCGMetadataAttributeGeneric*, PCGMetadataEntryKey> AttrAndKey{nullptr, PCGInvalidEntryKey};
+	TOptional<TTuple<const FPCGMetadataAttributeBase*, PCGMetadataEntryKey>> AttrAndKey{};
 
 	ETextOverflowPolicy OverflowPolicy = GetDefault<UPCGEditorSettings>()->AttributeListViewSettings.GeneralOverflowPolicy;
 
 	if (FPCGColumnData* PCGColumnData = SharedAttributeListView->PCGColumnData.Find(ColumnId))
 	{
-		const bool bShouldUseGenericAttributeViewer = PCGColumnData->GenericAttribute != nullptr;
+		if (PCGColumnData->DataKeys.IsValid() && PCGColumnData->UnderlyingAttribute)
+		{
+			const PCGMetadataEntryKey* Key = nullptr;
+			if (PCGColumnData->DataKeys->GetKey<PCGMetadataEntryKey>(Key))
+			{
+				AttrAndKey = MakeTuple(PCGColumnData->UnderlyingAttribute, *Key);
+			}
+		}
 		// We have to make sure the data is still valid otherwise this will crash.
-		if (ensure(SharedAttributeListView->DataPtr) && PCGColumnData->DataAccessor.IsValid() && PCGColumnData->DataKeys.IsValid() && !bShouldUseGenericAttributeViewer)
+		else if (ensure(SharedAttributeListView->DataPtr) && PCGColumnData->DataAccessor.IsValid() && PCGColumnData->DataKeys.IsValid() && !PCGColumnData->UnderlyingAttribute)
 		{
-			bDataRetrievedWithAccessors = true;
 			const int32 Index = InternalItem->Index;
 			auto Callback = [&PCGColumnData, Index, &RowText, &bShouldComputeCellWidth, &bIsClickable, &OverflowPolicy, ColumnId] (auto Dummy)
 			{
@@ -208,15 +212,6 @@ TSharedRef<SWidget> SPCGListViewItemRow::GenerateWidgetForColumn(const FName& Co
 
 			PCGMetadataAttribute::CallbackWithRightType(PCGColumnData->DataAccessor->GetUnderlyingType(), Callback);
 		}
-		else if (PCGColumnData->DataKeys.IsValid() && PCGColumnData->GenericAttribute)
-		{
-			const PCGMetadataEntryKey* Key = nullptr;
-			if (PCGColumnData->DataKeys->GetKey<PCGMetadataEntryKey>(Key))
-			{
-				AttrAndKey.Get<0>() = PCGColumnData->GenericAttribute;
-				AttrAndKey.Get<1>() = *Key;
-			}
-		}
 	}
 
 	if (bShouldComputeCellWidth)
@@ -253,7 +248,7 @@ TSharedRef<SWidget> SPCGListViewItemRow::GenerateWidgetForColumn(const FName& Co
 
 		return Hyperlink.ToSharedRef();
 	}
-	else if (bDataRetrievedWithAccessors)
+	else if (!AttrAndKey.IsSet())
 	{
 		return SNew(STextBlock)
 			.Text(RowText)
@@ -261,7 +256,7 @@ TSharedRef<SWidget> SPCGListViewItemRow::GenerateWidgetForColumn(const FName& Co
 			.Margin(Margin)
 			.TextStyle(TextStyle);
 	}
-	else if (AttrAndKey.Get<0>() && SharedAttributeListView->DataPtr)
+	else if (AttrAndKey.IsSet() && AttrAndKey->Get<0>() && SharedAttributeListView->DataPtr)
 	{
 		return SNew(SButton)
 			.Text(INVTEXT("..."))
@@ -272,7 +267,7 @@ TSharedRef<SWidget> SPCGListViewItemRow::GenerateWidgetForColumn(const FName& Co
 						.ClientSize(FVector2D(700, 700));
 
 					TSharedRef<SPCGMetadataAttributeGenericViewer> Viewer = SNew(SPCGMetadataAttributeGenericViewer);
-					Viewer->Setup(WeakData, AttrAndKey.Get<0>(), AttrAndKey.Get<1>());
+					Viewer->Setup(WeakData, AttrAndKey->Get<0>(), AttrAndKey->Get<1>());
 
 					ModalWindow->SetContent(Viewer);
 
@@ -1835,7 +1830,7 @@ void SPCGEditorGraphAttributeListView::AddColumn(const FPCGTableVisualizerColumn
 	FPCGColumnData& ColumnData = PCGColumnData.Add(InColumnInfo.Id);
 	ColumnData.DataAccessor = InColumnInfo.Accessor;
 	ColumnData.DataKeys = InColumnInfo.AccessorKeys;
-	ColumnData.GenericAttribute = InColumnInfo.GenericAttribute;
+	ColumnData.UnderlyingAttribute = InColumnInfo.UnderlyingAttribute;
 
 	const float ColumnWidth = InColumnInfo.Width < 0.0f ? PCGEditorGraphAttributeListView::CalculateColumnWidth(InColumnInfo.Label) : InColumnInfo.Width;
 	ColumnsMaxWidthMapping.FindOrAdd(InColumnInfo.Id) = ColumnWidth;
diff --git a/Engine/Plugins/PCG/Source/PCGEditor/Private/Widgets/SPCGEditorGraphAttributeListView.h b/Engine/Plugins/PCG/Source/PCGEditor/Private/Widgets/SPCGEditorGraphAttributeListView.h
--- a/Engine/Plugins/PCG/Source/PCGEditor/Private/Widgets/SPCGEditorGraphAttributeListView.h
+++ b/Engine/Plugins/PCG/Source/PCGEditor/Private/Widgets/SPCGEditorGraphAttributeListView.h
@@ -57,8 +57,7 @@ struct FPCGColumnData
 {
 	TSharedPtr<const IPCGAttributeAccessor> DataAccessor;
 	TSharedPtr<const IPCGAttributeAccessorKeys> DataKeys;
-	// @todo_pcg: To be removed when we do not make the distinction between typed and generic.
-	const FPCGMetadataAttributeGeneric* GenericAttribute = nullptr;
+	const FPCGMetadataAttributeBase* UnderlyingAttribute = nullptr;
 };
 
 typedef TSharedPtr<FPCGListViewItem> PCGListViewItemPtr;
diff --git a/Engine/Plugins/PCG/Source/PCGEditor/Private/Widgets/SPCGMetadataAttributeGenericViewer.cpp b/Engine/Plugins/PCG/Source/PCGEditor/Private/Widgets/SPCGMetadataAttributeGenericViewer.cpp
--- a/Engine/Plugins/PCG/Source/PCGEditor/Private/Widgets/SPCGMetadataAttributeGenericViewer.cpp
+++ b/Engine/Plugins/PCG/Source/PCGEditor/Private/Widgets/SPCGMetadataAttributeGenericViewer.cpp
@@ -2,14 +2,14 @@
 
 #include "Widgets/SPCGMetadataAttributeGenericViewer.h"
 
-#include "DetailCategoryBuilder.h"
-#include "DetailLayoutBuilder.h"
 #include "PCGData.h"
+#include "Metadata/PCGMetadataAttribute.h"
 
+#include "DetailCategoryBuilder.h"
+#include "DetailLayoutBuilder.h"
 #include "DetailsViewArgs.h"
 #include "InstancedPropertyBagStructureDataProvider.h"
 #include "PropertyEditorModule.h"
-#include "Metadata/PCGMetadataAttributeGeneric.h"
 #include "Modules/ModuleManager.h"
 
 TSharedRef<IDetailCustomization> FPCGMetadataAttributeGenericViewerDetailCustomization::MakeInstance(TSharedPtr<SPCGMetadataAttributeGenericViewer> Viewer)
@@ -97,7 +97,7 @@ void SPCGMetadataAttributeGenericViewer::Refresh()
 	}
 }
 
-void SPCGMetadataAttributeGenericViewer::Setup(TWeakObjectPtr<const UPCGData> InData, const FPCGMetadataAttributeGeneric* InAttr, PCGMetadataEntryKey InEntryKey)
+void SPCGMetadataAttributeGenericViewer::Setup(TWeakObjectPtr<const UPCGData> InData, const FPCGMetadataAttributeBase* InAttr, PCGMetadataEntryKey InEntryKey)
 {
 	Data = InData;
 	Attr = InAttr;
diff --git a/Engine/Plugins/PCG/Tests/Private/Benchmarks/Metadata/PCGMetadataAttributeGenericBenchmarks.cpp b/Engine/Plugins/PCG/Tests/Private/Benchmarks/Metadata/PCGMetadataAttributeGenericBenchmarks.cpp
--- a/Engine/Plugins/PCG/Tests/Private/Benchmarks/Metadata/PCGMetadataAttributeGenericBenchmarks.cpp
+++ b/Engine/Plugins/PCG/Tests/Private/Benchmarks/Metadata/PCGMetadataAttributeGenericBenchmarks.cpp
@@ -35,7 +35,7 @@ TEST_CASE_METHOD(PCGTests::FPCGBaseTestWithBasicContext, "PCG::AttributeGeneric:
 		TPCGValueRange<int64> PointArrayDataMetadataEntryRange = PointArrayData->GetMetadataEntryValueRange(/*bAllocate=*/false);
 
 		FPCGMetadataDomain* Domain = PointArrayData->MutableMetadata()->GetDefaultMetadataDomain();
-		FPCGMetadataAttributeGeneric* GenericAttribute = Domain->CreateGenericAttribute<double>(DoubleGenericAttributeName, 0.);
+		FPCGMetadataAttributeBase* GenericAttribute = Domain->CreateAttribute<double>(DoubleGenericAttributeName, 0.);
 		FPCGMetadataAttribute<double>* TypedAttribute = Domain->CreateAttribute<double>(DoubleAttributeName, 0., /*bAllowsInterpolation=*/true, /*bOverrideParent=*/false);
 		
 		// Allocate entry keys and values
@@ -114,7 +114,7 @@ TEST_CASE_METHOD(PCGTests::FPCGBaseTestWithBasicContext, "PCG::AttributeGeneric:
 		TPCGValueRange<int64> ChildPointArrayDataMetadataEntryRange = ChildPointArrayData->GetMetadataEntryValueRange(/*bAllocate=*/false);
 
 		FPCGMetadataDomain* ChildDomain = ChildPointArrayData->MutableMetadata()->GetDefaultMetadataDomain();
-		FPCGMetadataAttributeGeneric* ChildGenericAttribute = ChildDomain->GetMutableGenericAttribute(DoubleGenericAttributeName);
+		FPCGMetadataAttributeBase* ChildGenericAttribute = ChildDomain->GetMutableAttribute(DoubleGenericAttributeName);
 		FPCGMetadataAttribute<double>* ChildTypedAttribute = ChildDomain->GetMutableTypedAttribute<double>(DoubleAttributeName);
 		
 		CHECK_NOT_EQUAL(ChildGenericAttribute, nullptr);
@@ -162,7 +162,7 @@ TEST_CASE_METHOD(PCGTests::FPCGBaseTestWithBasicContext, "PCG::AttributeGeneric:
 		TPCGValueRange<int64> PointArrayDataMetadataEntryRange = PointArrayData->GetMetadataEntryValueRange(/*bAllocate=*/false);
 
 		FPCGMetadataDomain* Domain = PointArrayData->MutableMetadata()->GetDefaultMetadataDomain();
-		FPCGMetadataAttributeGeneric* GenericAttribute = Domain->CreateGenericAttribute<FString>(StringGenericAttributeName, FString{});
+		FPCGMetadataAttributeBase* GenericAttribute = Domain->CreateAttribute<FString>(StringGenericAttributeName, FString{});
 		FPCGMetadataAttribute<FString>* TypedAttribute = Domain->CreateAttribute<FString>(StringAttributeName, FString{}, /*bAllowsInterpolation=*/true, /*bOverrideParent=*/false);
 
 		// Allocate entry keys and values
@@ -241,7 +241,7 @@ TEST_CASE_METHOD(PCGTests::FPCGBaseTestWithBasicContext, "PCG::AttributeGeneric:
 		TPCGValueRange<int64> ChildPointArrayDataMetadataEntryRange = ChildPointArrayData->GetMetadataEntryValueRange(/*bAllocate=*/false);
 
 		FPCGMetadataDomain* ChildDomain = ChildPointArrayData->MutableMetadata()->GetDefaultMetadataDomain();
-		FPCGMetadataAttributeGeneric* ChildGenericAttribute = ChildDomain->GetMutableGenericAttribute(StringGenericAttributeName);
+		FPCGMetadataAttributeBase* ChildGenericAttribute = ChildDomain->GetMutableAttribute(StringGenericAttributeName);
 		FPCGMetadataAttribute<FString>* ChildTypedAttribute = ChildDomain->GetMutableTypedAttribute<FString>(StringAttributeName);
 		
 		CHECK_NOT_EQUAL(ChildGenericAttribute, nullptr);
diff --git a/Engine/Plugins/PCG/Tests/Private/Metadata/Accessors/PCGMetadataAttributeGenericAccessorTests.cpp b/Engine/Plugins/PCG/Tests/Private/Metadata/Accessors/PCGMetadataAttributeGenericAccessorTests.cpp
--- a/Engine/Plugins/PCG/Tests/Private/Metadata/Accessors/PCGMetadataAttributeGenericAccessorTests.cpp
+++ b/Engine/Plugins/PCG/Tests/Private/Metadata/Accessors/PCGMetadataAttributeGenericAccessorTests.cpp
@@ -53,17 +53,17 @@ TEST_CASE_METHOD(PCGTests::FPCGBaseTestWithBasicContext, "PCG::Metadata::Accesso
 		
 		FRandomStream RandomStream(Seed);
 		
-		FPCGMetadataAttributeGeneric* GenericAttribute = nullptr;
+		FPCGMetadataAttributeBase* GenericAttribute = nullptr;
 	
 		if (AttributeTester.OverrideDesc.IsSet())
 		{
-			GenericAttribute = Domain->CreateGenericAttribute(*AttributeTester.OverrideDesc);
+			GenericAttribute = Domain->CreateAttribute(*AttributeTester.OverrideDesc);
 			REQUIRE_NOT_EQUAL(GenericAttribute, nullptr);
 			GenericAttribute->SetDefaultValue<T>(AttributeTester.DefaultValue);
 		}
 		else
 		{
-			GenericAttribute = Domain->CreateGenericAttribute<T>(PCGAttributeTestsCommonHelper::AttributeName, AttributeTester.DefaultValue);
+			GenericAttribute = Domain->CreateAttribute<T>(PCGAttributeTestsCommonHelper::AttributeName, AttributeTester.DefaultValue);
 			REQUIRE_NOT_EQUAL(GenericAttribute, nullptr);
 		}
 		
@@ -187,7 +187,7 @@ TEST_CASE_METHOD(PCGTests::FPCGBaseTestWithBasicContext, "PCG::Metadata::Accesso
 		
 	FRandomStream RandomStream(Seed);
 		
-	FPCGMetadataAttributeGeneric* GenericAttribute = Domain->CreateGenericAttribute<double>(DoubleAttributeName, -1.0);
+	FPCGMetadataAttributeBase* GenericAttribute = Domain->CreateAttribute<double>(DoubleAttributeName, -1.0);
 	REQUIRE_NOT_EQUAL(GenericAttribute, nullptr);
 	
 	TUniquePtr<IPCGAttributeAccessor> DoubleAttributeAccessor = PCGAttributeAccessorHelpers::CreateAccessor(PointArrayData, FPCGAttributePropertySelector::CreateAttributeSelector(DoubleAttributeName));
@@ -242,8 +242,7 @@ TEST_CASE_METHOD(PCGTests::FPCGBaseTestWithBasicContext, "PCG::Metadata::Accesso
 		
 	FRandomStream RandomStream(Seed);
 		
-	FPCGMetadataAttributeGeneric* GenericAttribute = nullptr;
-	GenericAttribute = Domain->CreateGenericAttribute<FString>(PCGAttributeTestsCommonHelper::AttributeName, FString{});
+	FPCGMetadataAttributeBase* GenericAttribute = Domain->CreateAttribute<FString>(PCGAttributeTestsCommonHelper::AttributeName, FString{});
 	REQUIRE_NOT_EQUAL(GenericAttribute, nullptr);
 		
 	TUniquePtr<IPCGAttributeAccessor> AttributeAccessor = PCGAttributeAccessorHelpers::CreateAccessor(PointArrayData, FPCGAttributePropertySelector::CreateAttributeSelector(PCGAttributeTestsCommonHelper::AttributeName));
diff --git a/Engine/Plugins/PCG/Tests/Private/Metadata/PCGMetadataAttributeGenericTests.cpp b/Engine/Plugins/PCG/Tests/Private/Metadata/PCGMetadataAttributeGenericTests.cpp
--- a/Engine/Plugins/PCG/Tests/Private/Metadata/PCGMetadataAttributeGenericTests.cpp
+++ b/Engine/Plugins/PCG/Tests/Private/Metadata/PCGMetadataAttributeGenericTests.cpp
@@ -52,17 +52,17 @@ TEST_CASE_METHOD(PCGTests::FPCGBaseTestWithBasicContext, "PCG::AttributeGeneric:
 		
 		FRandomStream RandomStream(Seed);
 		
-		FPCGMetadataAttributeGeneric* GenericAttribute = nullptr;
+		FPCGMetadataAttributeBase* GenericAttribute = nullptr;
 	
 		if (AttributeTester.OverrideDesc.IsSet())
 		{
-			GenericAttribute = Domain->CreateGenericAttribute(*AttributeTester.OverrideDesc);
+			GenericAttribute = Domain->CreateAttribute(*AttributeTester.OverrideDesc);
 			REQUIRE_NOT_EQUAL(GenericAttribute, nullptr);
 			GenericAttribute->SetDefaultValue<T>(AttributeTester.DefaultValue);
 		}
 		else
 		{
-			GenericAttribute = Domain->CreateGenericAttribute<T>(PCGAttributeTestsCommonHelper::AttributeName, AttributeTester.DefaultValue);
+			GenericAttribute = Domain->CreateAttribute<T>(PCGAttributeTestsCommonHelper::AttributeName, AttributeTester.DefaultValue);
 			REQUIRE_NOT_EQUAL(GenericAttribute, nullptr);
 		}
 		
@@ -105,7 +105,7 @@ TEST_CASE_METHOD(PCGTests::FPCGBaseTestWithBasicContext, "PCG::AttributeGeneric:
 
 			TPCGValueRange<int64> ChildPointArrayDataMetadataEntryRange = ChildPointArrayData->GetMetadataEntryValueRange(/*bAllocate=*/false);
 
-			FPCGMetadataAttributeGeneric* ChildGenericAttribute = ChildDomain->GetMutableGenericAttribute(PCGAttributeTestsCommonHelper::AttributeName);
+			FPCGMetadataAttributeBase* ChildGenericAttribute = ChildDomain->GetMutableAttribute(PCGAttributeTestsCommonHelper::AttributeName);
 			REQUIRE_NOT_EQUAL(ChildGenericAttribute, nullptr);
 			
 			SECTION("Expected attribute descriptor")
@@ -222,17 +222,17 @@ TEST_CASE_METHOD(PCGTests::FPCGBaseTestWithBasicContext, "PCG::AttributeGeneric:
 		
 		FRandomStream RandomStream(Seed);
 		
-		FPCGMetadataAttributeGeneric* GenericAttribute = nullptr;
+		FPCGMetadataAttributeBase* GenericAttribute = nullptr;
 	
 		if (AttributeTester.OverrideDesc.IsSet())
 		{
-			GenericAttribute = Domain->CreateGenericAttribute(*AttributeTester.OverrideDesc);
+			GenericAttribute = Domain->CreateAttribute(*AttributeTester.OverrideDesc);
 			REQUIRE_NOT_EQUAL(GenericAttribute, nullptr);
 			GenericAttribute->SetDefaultValue<T>(AttributeTester.DefaultValue);
 		}
 		else
 		{
-			GenericAttribute = Domain->CreateGenericAttribute<T>(PCGAttributeTestsCommonHelper::AttributeName, AttributeTester.DefaultValue);
+			GenericAttribute = Domain->CreateAttribute<T>(PCGAttributeTestsCommonHelper::AttributeName, AttributeTester.DefaultValue);
 			REQUIRE_NOT_EQUAL(GenericAttribute, nullptr);
 		}
 		
@@ -262,7 +262,7 @@ TEST_CASE_METHOD(PCGTests::FPCGBaseTestWithBasicContext, "PCG::AttributeGeneric:
 
 			FPCGMetadataDomain* ChildDomain = ChildParamData->MutableMetadata()->GetDefaultMetadataDomain();
 
-			FPCGMetadataAttributeGeneric* ChildGenericAttribute = ChildDomain->GetMutableGenericAttribute(PCGAttributeTestsCommonHelper::AttributeName);
+			FPCGMetadataAttributeBase* ChildGenericAttribute = ChildDomain->GetMutableAttribute(PCGAttributeTestsCommonHelper::AttributeName);
 			REQUIRE_NOT_EQUAL(ChildGenericAttribute, nullptr);
 			
 			SECTION("Expected attribute descriptor")
@@ -373,7 +373,7 @@ TEST_CASE_METHOD(PCGTests::FPCGBaseTestWithBasicContext, "PCG::AttributeGeneric:
 		UPCGParamData* ParamData = CreateData<UPCGParamData>();
 		FPCGMetadataDomain* Domain = ParamData->MutableMetadata()->GetDefaultMetadataDomain();
 		
-		FPCGMetadataAttributeGeneric* GenericAttribute = Domain->CreateGenericAttribute<RealType>(PCGAttributeTestsCommonHelper::AttributeName, InValue);
+		FPCGMetadataAttributeBase* GenericAttribute = Domain->CreateAttribute<RealType>(PCGAttributeTestsCommonHelper::AttributeName, InValue);
 		
 		SECTION("Get Value with wrong type")
 		{
@@ -523,9 +523,9 @@ TEST_CASE_METHOD(PCGTests::FPCGBaseTestWithBasicContext, "PCG::AttributeGeneric:
 	};
 	
 	// Initialize the default value to a value we will add later on.
-	FPCGMetadataAttributeGeneric* GenericAttribute = Domain->CreateGenericAttribute<FString>(PCGAttributeTestsCommonHelper::AttributeName, TEXT("Value: 0"));
-	FPCGMetadataAttributeGeneric* GenericAttributePath = Domain->CreateGenericAttribute<FSoftObjectPath>(PCGAttributeTestsCommonHelper::AttributeName2, FSoftObjectPath(FString::Format(*BasePathFormat, { 0 })));
-	FPCGMetadataAttributeGeneric* GenericAttributeClass = Domain->CreateGenericAttribute<FSoftClassPath>(PCGAttributeTestsCommonHelper::AttributeName3, UClassesPool[0]);
+	FPCGMetadataAttributeBase* GenericAttribute = Domain->CreateAttribute<FString>(PCGAttributeTestsCommonHelper::AttributeName, TEXT("Value: 0"));
+	FPCGMetadataAttributeBase* GenericAttributePath = Domain->CreateAttribute<FSoftObjectPath>(PCGAttributeTestsCommonHelper::AttributeName2, FSoftObjectPath(FString::Format(*BasePathFormat, { 0 })));
+	FPCGMetadataAttributeBase* GenericAttributeClass = Domain->CreateAttribute<FSoftClassPath>(PCGAttributeTestsCommonHelper::AttributeName3, UClassesPool[0]);
 	REQUIRE_NOT_EQUAL(GenericAttribute, nullptr);
 	REQUIRE_NOT_EQUAL(GenericAttributePath, nullptr);
 	REQUIRE_NOT_EQUAL(GenericAttributeClass, nullptr);
@@ -617,7 +617,7 @@ TEST_CASE_METHOD(PCGTests::FPCGBaseTestWithBasicContext, "PCG::AttributeGeneric:
 	
 		FRandomStream RandomStream(Seed);
 	
-		FPCGMetadataAttributeGeneric* GenericAttribute = Domain->CreateGenericAttribute<T>(PCGAttributeTestsCommonHelper::AttributeName, AttributeTester.DefaultValue);
+		FPCGMetadataAttributeBase* GenericAttribute = Domain->CreateAttribute<T>(PCGAttributeTestsCommonHelper::AttributeName, AttributeTester.DefaultValue);
 		REQUIRE_NOT_EQUAL(GenericAttribute, nullptr);
 	
 		TStaticArray<T, NumEntries> Values{};
@@ -857,7 +857,7 @@ TEST_CASE_METHOD(PCGTests::FPCGBaseTestWithBasicContext, "PCG::AttributeGeneric:
 
 		FRandomStream RandomStream(42);
 
-		FPCGMetadataAttributeGeneric* GenericAttribute = Domain->CreateGenericAttribute<T>(PCGAttributeTestsCommonHelper::AttributeName, AttributeTester.DefaultValue);
+		FPCGMetadataAttributeBase* GenericAttribute = Domain->CreateAttribute<T>(PCGAttributeTestsCommonHelper::AttributeName, AttributeTester.DefaultValue);
 		REQUIRE_NOT_EQUAL(GenericAttribute, nullptr);
 
 		// Populate initial entries with random values.
@@ -932,9 +932,9 @@ TEST_CASE_METHOD(PCGTests::FPCGBaseTestWithBasicContext, "PCG::AttributeGeneric:
 
 		FRandomStream RandomStream(42);
 
-		FPCGMetadataAttributeGeneric* GenericAttribute = nullptr;
+		FPCGMetadataAttributeBase* GenericAttribute = nullptr;
 
-		GenericAttribute = Domain->CreateGenericAttribute<T>(PCGAttributeTestsCommonHelper::AttributeName, AttributeTester.DefaultValue);
+		GenericAttribute = Domain->CreateAttribute<T>(PCGAttributeTestsCommonHelper::AttributeName, AttributeTester.DefaultValue);
 		REQUIRE_NOT_EQUAL(GenericAttribute, nullptr);
 
 		TStaticArray<T, NumPoints> Values{};
@@ -966,7 +966,7 @@ TEST_CASE_METHOD(PCGTests::FPCGBaseTestWithBasicContext, "PCG::AttributeGeneric:
 
 			TPCGValueRange<int64> ChildPointArrayDataMetadataEntryRange = ChildPointArrayData->GetMetadataEntryValueRange(/*bAllocate=*/false);
 
-			FPCGMetadataAttributeGeneric* ChildGenericAttribute = ChildDomain->GetMutableGenericAttribute(PCGAttributeTestsCommonHelper::AttributeName);
+			FPCGMetadataAttributeBase* ChildGenericAttribute = ChildDomain->GetMutableAttribute(PCGAttributeTestsCommonHelper::AttributeName);
 			REQUIRE_NOT_EQUAL(ChildGenericAttribute, nullptr);
 
 			// Add more values
diff --git a/Engine/Plugins/PCG/Tests/Private/Metadata/PCGMetadataAttributeTests.cpp b/Engine/Plugins/PCG/Tests/Private/Metadata/PCGMetadataAttributeTests.cpp
--- a/Engine/Plugins/PCG/Tests/Private/Metadata/PCGMetadataAttributeTests.cpp
+++ b/Engine/Plugins/PCG/Tests/Private/Metadata/PCGMetadataAttributeTests.cpp
@@ -146,7 +146,7 @@ namespace PCGAttributeTypedTests
 	template <typename T>
 	bool IsOfType(const FPCGMetadataAttributeBase* Attribute)
 	{
-		return Attribute && Attribute->IsGeneric() && static_cast<const FPCGMetadataAttributeGeneric*>(Attribute)->IsOfType<T>();
+		return Attribute && Attribute->IsOfType<T>();
 	}
 	
 }
@@ -685,13 +685,13 @@ TEST_CASE_METHOD(PCGTests::FPCGBaseTestWithBasicContext, "PCG::AttributeTyped::G
 	FPCGMetadataAttribute<double>* TypedAttr = Domain->CreateAttribute<double>(AttributeName, 3.14, /*bAllowsInterpolation=*/true, /*bOverrideParent=*/false);
 	REQUIRE_NOT_EQUAL(TypedAttr, nullptr);
 
-	SECTION("GetMutableGenericAttribute returns the same attribute")
+	SECTION("GetMutableAttribute returns the same attribute")
 	{
-		FPCGMetadataAttributeGeneric* GenericAttr = Domain->GetMutableGenericAttribute(AttributeName);
+		FPCGMetadataAttributeBase* GenericAttr = Domain->GetMutableAttribute(AttributeName);
 		REQUIRE_NOT_EQUAL(GenericAttr, nullptr);
 		REQUIRE(IsOfType<double>(GenericAttr));
 	
-		REQUIRE_EQUAL(static_cast<FPCGMetadataAttributeGeneric*>(TypedAttr), GenericAttr);
+		REQUIRE_EQUAL(static_cast<FPCGMetadataAttributeBase*>(TypedAttr), GenericAttr);
 	}
 
 	SECTION("GetMutableTypedAttribute returns typed attribute")
@@ -709,7 +709,7 @@ TEST_CASE_METHOD(PCGTests::FPCGBaseTestWithBasicContext, "PCG::AttributeTyped::G
 	SECTION("Setting via generic, reading via typed")
 	{
 		PCGMetadataEntryKey EntryKey = Domain->AddEntry();
-		FPCGMetadataAttributeGeneric* GenericAttr = Domain->GetMutableGenericAttribute(AttributeName);
+		FPCGMetadataAttributeBase* GenericAttr = Domain->GetMutableAttribute(AttributeName);
 		GenericAttr->SetValue<double>(EntryKey, 2.718);
 
 		double ReadValue = TypedAttr->GetValueFromItemKey(EntryKey);
@@ -721,7 +721,7 @@ TEST_CASE_METHOD(PCGTests::FPCGBaseTestWithBasicContext, "PCG::AttributeTyped::G
 	 	PCGMetadataEntryKey EntryKey = Domain->AddEntry();
 	 	TypedAttr->SetValue(EntryKey, 1.618);
 	
-	 	FPCGMetadataAttributeGeneric* GenericAttr = Domain->GetMutableGenericAttribute(AttributeName);
+	 	FPCGMetadataAttributeBase* GenericAttr = Domain->GetMutableAttribute(AttributeName);
 	 	double ReadValue = GenericAttr->GetValueFromItemKey<double>(EntryKey);
 	 	REQUIRE_EQUAL(ReadValue, 1.618);
 	 }
@@ -1202,23 +1202,6 @@ TEST_CASE_METHOD(PCGTests::FPCGBaseTestWithBasicContext, "PCG::AttributeTyped::B
 	}
 }
 
-/**
-* Test that the typed attribute inherits IsGeneric() == true from FPCGMetadataAttributeGeneric.
-*/
-TEST_CASE_METHOD(PCGTests::FPCGBaseTestWithBasicContext, "PCG::AttributeTyped::IsGenericFlag", "[PCG][AttributeTyped]")
-{
-	using namespace PCGAttributeTypedTests;
-
-	UPCGParamData* ParamData = CreateData<UPCGParamData>();
-	FPCGMetadataDomain* Domain = ParamData->MutableMetadata()->GetDefaultMetadataDomain();
-
-	FPCGMetadataAttribute<float>* Attr = Domain->CreateAttribute<float>(AttributeName, 0.0f, true, false);
-	REQUIRE_NOT_EQUAL(Attr, nullptr);
-
-	// FPCGMetadataAttribute<T> inherits from FPCGMetadataAttributeGeneric, so IsGeneric should be true
-	REQUIRE(Attr->IsGeneric());
-}
-
 /**
  * Test CopyExistingAttribute through the domain for typed attributes.
  */
diff --git a/Engine/Plugins/PCG/Tests/Private/Metadata/PCGMetadataContainerTypesTests.cpp b/Engine/Plugins/PCG/Tests/Private/Metadata/PCGMetadataContainerTypesTests.cpp
--- a/Engine/Plugins/PCG/Tests/Private/Metadata/PCGMetadataContainerTypesTests.cpp
+++ b/Engine/Plugins/PCG/Tests/Private/Metadata/PCGMetadataContainerTypesTests.cpp
@@ -29,7 +29,7 @@ TEST_CASE_METHOD(PCGTests::FPCGBaseTestWithBasicContext, "PCG::MetadataContainer
 	
 	SECTION("TScriptSetWrapper")
 	{
-		FPCGMetadataAttributeGeneric* SetGenericAttribute = Domain->CreateGenericAttribute<TSet<double>>("SetAttr", TSet<double>{});
+		FPCGMetadataAttributeBase* SetGenericAttribute = Domain->CreateAttribute<TSet<double>>("SetAttr", TSet<double>{});
 		REQUIRE_NOT_EQUAL(SetGenericAttribute, nullptr);
 		
 		const TSet<double> Values = {1.0, 2.0, 5.0, 3.0};
@@ -103,7 +103,7 @@ TEST_CASE_METHOD(PCGTests::FPCGBaseTestWithBasicContext, "PCG::MetadataContainer
 	
 	SECTION("TScriptMapWrapper")
 	{
-		FPCGMetadataAttributeGeneric* MapGenericAttribute = Domain->CreateGenericAttribute<TMap<FString, double>>("MapAttr", TMap<FString, double>{});
+		FPCGMetadataAttributeBase* MapGenericAttribute = Domain->CreateAttribute<TMap<FString, double>>("MapAttr", TMap<FString, double>{});
 		REQUIRE_NOT_EQUAL(MapGenericAttribute, nullptr);
 		
 		TMap<FString, double> Values = 
diff --git a/Engine/Plugins/PCG/Tests/Private/Metadata/PCGMetadataFlattenTests.cpp b/Engine/Plugins/PCG/Tests/Private/Metadata/PCGMetadataFlattenTests.cpp
--- a/Engine/Plugins/PCG/Tests/Private/Metadata/PCGMetadataFlattenTests.cpp
+++ b/Engine/Plugins/PCG/Tests/Private/Metadata/PCGMetadataFlattenTests.cpp
@@ -197,8 +197,8 @@ TEST_CASE_METHOD(PCGTests::FPCGBaseTestWithBasicContext, "PCG::Metadata::Flatten
 	RootPointData->AllocateProperties(EPCGPointNativeProperties::MetadataEntry);
 	RootPointData->SetNumPoints(NumPoints);
 	
-	FPCGMetadataAttributeGeneric* Attribute1 = RootPointData->MutableMetadata()->GetDefaultMetadataDomain()->CreateGenericAttribute<float>(Attribute1Name, -0.1f);
-	FPCGMetadataAttributeGeneric* Attribute2 = RootPointData->MutableMetadata()->GetDefaultMetadataDomain()->CreateGenericAttribute<FString>(Attribute2Name, TEXT("Default"));
+	FPCGMetadataAttributeBase* Attribute1 = RootPointData->MutableMetadata()->GetDefaultMetadataDomain()->CreateAttribute<float>(Attribute1Name, -0.1f);
+	FPCGMetadataAttributeBase* Attribute2 = RootPointData->MutableMetadata()->GetDefaultMetadataDomain()->CreateAttribute<FString>(Attribute2Name, TEXT("Default"));
 
 	check(Attribute1 && Attribute2);
 
@@ -223,8 +223,8 @@ TEST_CASE_METHOD(PCGTests::FPCGBaseTestWithBasicContext, "PCG::Metadata::Flatten
 	}
 
 	UPCGBasePointData* FirstChildPointData = Cast<UPCGBasePointData>(RootPointData->DuplicateData(nullptr));
-	Attribute1 = FirstChildPointData->MutableMetadata()->GetDefaultMetadataDomain()->GetMutableGenericAttribute(Attribute1Name);
-	Attribute2 = FirstChildPointData->MutableMetadata()->GetDefaultMetadataDomain()->GetMutableGenericAttribute(Attribute2Name);
+	Attribute1 = FirstChildPointData->MutableMetadata()->GetDefaultMetadataDomain()->GetMutableAttribute(Attribute1Name);
+	Attribute2 = FirstChildPointData->MutableMetadata()->GetDefaultMetadataDomain()->GetMutableAttribute(Attribute2Name);
 
 	SECTION("Attributes exists in first child")
 	{
@@ -236,7 +236,7 @@ TEST_CASE_METHOD(PCGTests::FPCGBaseTestWithBasicContext, "PCG::Metadata::Flatten
 
 	// Override all the values for Attribute 1, and replace all "1" by "0" in attribute 2
 	// Also add a third attribute
-	FPCGMetadataAttributeGeneric* Attribute3 = FirstChildPointData->MutableMetadata()->GetDefaultMetadataDomain()->CreateGenericAttribute<int>(Attribute3Name, -1);
+	FPCGMetadataAttributeBase* Attribute3 = FirstChildPointData->MutableMetadata()->GetDefaultMetadataDomain()->CreateAttribute<int>(Attribute3Name, -1);
 	check(Attribute3);
 
 	SECTION("First child has right number of points")
@@ -275,7 +275,7 @@ TEST_CASE_METHOD(PCGTests::FPCGBaseTestWithBasicContext, "PCG::Metadata::Flatten
 	InitializeFromDataParams.bInheritSpatialData = false;
 	SecondChildPointData->InitializeFromDataWithParams(InitializeFromDataParams);
 	
-	Attribute3 = SecondChildPointData->MutableMetadata()->GetDefaultMetadataDomain()->GetMutableGenericAttribute(Attribute3Name);
+	Attribute3 = SecondChildPointData->MutableMetadata()->GetDefaultMetadataDomain()->GetMutableAttribute(Attribute3Name);
 	SECTION("Attributes exists in second child")
 	{
 		REQUIRE_NOT_EQUAL(Attribute3, nullptr);
@@ -314,9 +314,9 @@ TEST_CASE_METHOD(PCGTests::FPCGBaseTestWithBasicContext, "PCG::Metadata::Flatten
 	UPCGBasePointData* FinalPointData = Cast<UPCGBasePointData>(SecondChildPointData->DuplicateData(nullptr));
 	FinalPointData->Flatten();
 
-	Attribute1 = FinalPointData->MutableMetadata()->GetDefaultMetadataDomain()->GetMutableGenericAttribute(Attribute1Name);
-	Attribute2 = FinalPointData->MutableMetadata()->GetDefaultMetadataDomain()->GetMutableGenericAttribute(Attribute2Name);
-	Attribute3 = FinalPointData->MutableMetadata()->GetDefaultMetadataDomain()->GetMutableGenericAttribute(Attribute3Name);
+	Attribute1 = FinalPointData->MutableMetadata()->GetDefaultMetadataDomain()->GetMutableAttribute(Attribute1Name);
+	Attribute2 = FinalPointData->MutableMetadata()->GetDefaultMetadataDomain()->GetMutableAttribute(Attribute2Name);
+	Attribute3 = FinalPointData->MutableMetadata()->GetDefaultMetadataDomain()->GetMutableAttribute(Attribute3Name);
 
 	SECTION("Attributes exists in final child")
 	{
@@ -426,8 +426,8 @@ TEST_CASE_METHOD(PCGTests::FPCGBaseTestWithBasicContext, "PCG::Metadata::Flatten
 	RootPointData->AllocateProperties(EPCGPointNativeProperties::MetadataEntry);
 	RootPointData->SetNumPoints(NumPoints);
 	
-	FPCGMetadataAttributeGeneric* Attribute1 = RootPointData->MutableMetadata()->GetDefaultMetadataDomain()->CreateGenericAttribute<float>(Attribute1Name, -0.1f);
-	FPCGMetadataAttributeGeneric* Attribute2 = RootPointData->MutableMetadata()->GetDefaultMetadataDomain()->CreateGenericAttribute<FString>(Attribute2Name, TEXT("Default"));
+	FPCGMetadataAttributeBase* Attribute1 = RootPointData->MutableMetadata()->GetDefaultMetadataDomain()->CreateAttribute<float>(Attribute1Name, -0.1f);
+	FPCGMetadataAttributeBase* Attribute2 = RootPointData->MutableMetadata()->GetDefaultMetadataDomain()->CreateAttribute<FString>(Attribute2Name, TEXT("Default"));
 	
 	REQUIRE_NOT_EQUAL(Attribute1, nullptr);
 	REQUIRE_NOT_EQUAL(Attribute2, nullptr);
@@ -453,8 +453,8 @@ TEST_CASE_METHOD(PCGTests::FPCGBaseTestWithBasicContext, "PCG::Metadata::Flatten
 	ChildPointData->InitializeFromDataWithParams(FPCGInitializeFromDataParams{RootPointData});
 	ChildPointData->SetNumPoints(NumPoints * 2);
 	
-	Attribute1 = ChildPointData->MutableMetadata()->GetDefaultMetadataDomain()->GetMutableGenericAttribute(Attribute1Name);
-	Attribute2 = ChildPointData->MutableMetadata()->GetDefaultMetadataDomain()->GetMutableGenericAttribute(Attribute2Name);
+	Attribute1 = ChildPointData->MutableMetadata()->GetDefaultMetadataDomain()->GetMutableAttribute(Attribute1Name);
+	Attribute2 = ChildPointData->MutableMetadata()->GetDefaultMetadataDomain()->GetMutableAttribute(Attribute2Name);
 	
 	REQUIRE_NOT_EQUAL(Attribute1, nullptr);
 	REQUIRE_NOT_EQUAL(Attribute2, nullptr);
