[PCG] Extend "Get Class From Attribute" to also optionally output the class path.

#rnx
#Synced-CL 52521088

diff --git a/Engine/Plugins/PCG/Source/PCG/Private/Elements/PCGGetClassFromAttribute.cpp b/Engine/Plugins/PCG/Source/PCG/Private/Elements/PCGGetClassFromAttribute.cpp
--- a/Engine/Plugins/PCG/Source/PCG/Private/Elements/PCGGetClassFromAttribute.cpp
+++ b/Engine/Plugins/PCG/Source/PCG/Private/Elements/PCGGetClassFromAttribute.cpp
@@ -38,6 +38,12 @@ FPCGElementPtr UPCGGetClassFromAttributeSettings::CreateElement() const
 	return MakeShared<FPCGGetClassFromAttributeElement>();
 }
 
+UPCGGetClassFromAttributeSettings::UPCGGetClassFromAttributeSettings(const FObjectInitializer& Initializer)
+	: Super(Initializer)
+{
+	ClassPathOutputTarget.SetAttributeName("ClassPath");
+}
+
 #if WITH_EDITOR
 FText UPCGGetClassFromAttributeSettings::GetDefaultNodeTitle() const
 {
@@ -95,7 +101,7 @@ bool FPCGGetClassFromAttributeElement::ExecuteInternal(FPCGContext* Context) con
 		check(OutputData);
 
 		// Fix @Source on the output selector against the duplicated data
-		FPCGAttributePropertyOutputSelector FixedOutput = Settings->OutputTarget.CopyAndFixSource(&FixedInput, OutputData);
+		const FPCGAttributePropertyOutputSelector FixedOutput = Settings->OutputTarget.CopyAndFixSource(&FixedInput, OutputData);
 
 		// Create or find an FString attribute in the domain specified by the output selector
 		FPCGMetadataDomain* OutputDomain = OutputData->MutableMetadata()->GetMetadataDomainFromSelector(FixedOutput);
@@ -116,6 +122,26 @@ bool FPCGGetClassFromAttributeElement::ExecuteInternal(FPCGContext* Context) con
 			PCGLog::Metadata::LogFailToCreateAccessorError(FixedOutput, Context);
 			continue;
 		}
+		
+		// Optional accessor for class path
+		TUniquePtr<IPCGAttributeAccessor> OutputPathAccessor;
+		TUniquePtr<IPCGAttributeAccessorKeys> OutputPathKeys;
+		
+		if (Settings->bAlsoOutputClassPath)
+		{
+			const FPCGAttributePropertyOutputSelector FixedClassPathOutput = Settings->ClassPathOutputTarget.CopyAndFixSource(&FixedInput, OutputData);
+			
+			OutputDomain->FindOrCreateAttribute<FSoftClassPath>(FixedClassPathOutput.GetName(), FSoftClassPath(), /*bAllowsInterpolation=*/false, /*bOverrideParent=*/false);
+			
+			OutputPathAccessor = PCGAttributeAccessorHelpers::CreateAccessor(OutputData, FixedClassPathOutput);
+			OutputPathKeys = PCGAttributeAccessorHelpers::CreateKeys(OutputData, FixedClassPathOutput);
+			
+			if (!OutputPathAccessor || !OutputPathKeys)
+			{
+				PCGLog::Metadata::LogFailToCreateAccessorError(FixedClassPathOutput, Context);
+				continue;
+			}
+		}
 
 		// When the attribute stores FSoftClassPath, the resolved object IS the class itself.
 		// Calling GetClass()->GetName() on a UClass returns "Class" (the meta-class name), so we detect this case upfront and use GetName() directly instead.
@@ -125,14 +151,22 @@ bool FPCGGetClassFromAttributeElement::ExecuteInternal(FPCGContext* Context) con
 		const int32 NumEntries = InputKeys->GetNum();
 		TArray<FSoftObjectPath> SoftPaths;
 		SoftPaths.SetNum(NumEntries);
-		InputAccessor->GetRange<FSoftObjectPath>(SoftPaths, 0, *InputKeys, EPCGAttributeAccessorFlags::AllowBroadcastAndConstructible);
+		if (!InputAccessor->GetRange<FSoftObjectPath>(SoftPaths, 0, *InputKeys, EPCGAttributeAccessorFlags::AllowBroadcastAndConstructible))
+		{
+			PCGLog::Metadata::LogFailToGetAttributeError<FSoftObjectPath>(FixedInput, InputAccessor.Get(), Context);
+			continue;
+		}
 
 		// Resolve class name for each entry (no load)
 		TArray<FString> ClassNames;
 		ClassNames.Reserve(NumEntries);
 
+		TArray<FSoftClassPath> ClassPaths;
+		ClassPaths.Reserve(Settings->bAlsoOutputClassPath ? NumEntries : 0);
+
 		for (const FSoftObjectPath& Path : SoftPaths)
 		{
+			FSoftClassPath ClassPath{};
 			FString ClassName;
 
 			if (!Path.IsNull())
@@ -140,7 +174,10 @@ bool FPCGGetClassFromAttributeElement::ExecuteInternal(FPCGContext* Context) con
 				// Resolve already-loaded object — no asset load triggered
 				if (const UObject* Object = Path.ResolveObject())
 				{
-					ClassName = Object->IsA<UClass>() ? Object->GetName() : Object->GetClass()->GetName();
+					const UClass* Class = Object->IsA<UClass>() ? CastChecked<const UClass>(Object) : Object->GetClass();
+					check(Class);
+					ClassName = Class->GetName();
+					ClassPath = FSoftClassPath(Class);
 				}
 #if WITH_EDITOR
 				// Object not in memory — fall back to the Asset Registry (editor only).
@@ -152,6 +189,7 @@ bool FPCGGetClassFromAttributeElement::ExecuteInternal(FPCGContext* Context) con
 						if (AssetData.IsValid())
 						{
 							ClassName = AssetData.AssetClassPath.GetAssetName().ToString();
+							ClassPath = FSoftClassPath(AssetData.AssetClassPath.ToString());
 						}
 					}
 				}
@@ -164,9 +202,18 @@ bool FPCGGetClassFromAttributeElement::ExecuteInternal(FPCGContext* Context) con
 			}
 
 			ClassNames.Add(MoveTemp(ClassName));
+			if (Settings->bAlsoOutputClassPath)
+			{
+				ClassPaths.Add(MoveTemp(ClassPath));
+			}
 		}
 
 		OutputAccessor->SetRange<FString>(ClassNames, 0, *OutputKeys);
+		if (Settings->bAlsoOutputClassPath)
+		{
+			check(OutputPathAccessor && OutputPathKeys);
+			OutputPathAccessor->SetRange<FSoftClassPath>(ClassPaths, 0, *OutputPathKeys);
+		}
 
 		Context->OutputData.TaggedData.Add_GetRef(Output).Data = OutputData;
 	}
diff --git a/Engine/Plugins/PCG/Source/PCG/Public/Elements/PCGGetClassFromAttribute.h b/Engine/Plugins/PCG/Source/PCG/Public/Elements/PCGGetClassFromAttribute.h
--- a/Engine/Plugins/PCG/Source/PCG/Public/Elements/PCGGetClassFromAttribute.h
+++ b/Engine/Plugins/PCG/Source/PCG/Public/Elements/PCGGetClassFromAttribute.h
@@ -19,6 +19,8 @@ class UPCGGetClassFromAttributeSettings : public UPCGSettings
 	GENERATED_BODY()
 
 public:
+	UPCGGetClassFromAttributeSettings(const FObjectInitializer& Initializer);
+	
 	//~Begin UPCGSettings interface
 #if WITH_EDITOR
 	virtual FName GetDefaultNodeName() const override { return FName("GetClassFromAttribute"); }
@@ -43,6 +45,14 @@ public:
 	/** Attribute to write the class name string to. */
 	UPROPERTY(BlueprintReadWrite, EditAnywhere, Category = Settings, meta = (PCG_Overridable, PCG_DiscardPropertySelection, PCG_DiscardExtraSelection))
 	FPCGAttributePropertyOutputSelector OutputTarget;
+	
+	/** If we also want to output the full Soft Class Path of the class. */
+	UPROPERTY(BlueprintReadWrite, EditAnywhere, Category = Settings, meta = (InlineEditConditionToggle))
+	bool bAlsoOutputClassPath = false;
+	
+	/** Optional attribute to write the full Soft Class Path. Useful to get the class path of the input object (if it is an actor for example). */
+	UPROPERTY(BlueprintReadWrite, EditAnywhere, Category = Settings, meta = (PCG_Overridable, PCG_DiscardPropertySelection, PCG_DiscardExtraSelection, EditCondition = "bAlsoOutputClassPath"))
+	FPCGAttributePropertyOutputSelector ClassPathOutputTarget;
 
 	/** If true, suppresses warnings when an input path cannot be resolved to a valid class name. */
 	UPROPERTY(BlueprintReadWrite, EditAnywhere, Category = Settings, AdvancedDisplay, meta = (PCG_Overridable))
