[PCG] Add method to get current task for FPCGDefaultExecutionState

#Synced-CL 52517676

diff --git a/Engine/Plugins/PCG/Source/PCG/Public/PCGDefaultExecutionSource.h b/Engine/Plugins/PCG/Source/PCG/Public/PCGDefaultExecutionSource.h
--- a/Engine/Plugins/PCG/Source/PCG/Public/PCGDefaultExecutionSource.h
+++ b/Engine/Plugins/PCG/Source/PCG/Public/PCGDefaultExecutionSource.h
@@ -111,6 +111,8 @@ public:
 
 	PCG_API void Sunset();
 
+	FPCGTaskId GetCurrentGenerationTask() const {return CurrentGenerationTask;};
+	
 protected:
 	FPCGTaskId CurrentGenerationTask = InvalidPCGTaskId;
 	
