[PCG] Data Overrides for Manual Editing in PCG

* Runtime Updates
+ Graph Editor integration and Data Overrides Panel
+ Viewport Visualizer (Extensions and Widgets)

#Synced-CL 52525912

diff --git a/Engine/Plugins/PCG/Source/PCG/Private/Graph/DataOverride/PCGDataOverrideHelpers.cpp b/Engine/Plugins/PCG/Source/PCG/Private/Graph/DataOverride/PCGDataOverrideHelpers.cpp
--- a/Engine/Plugins/PCG/Source/PCG/Private/Graph/DataOverride/PCGDataOverrideHelpers.cpp
+++ b/Engine/Plugins/PCG/Source/PCG/Private/Graph/DataOverride/PCGDataOverrideHelpers.cpp
@@ -40,13 +40,19 @@ namespace PCG::DataOverride
 			return Result;
 		}
 
-		FXxHash64 ComputeDataHashWithPin(const FPCGStack& InStack, const FName PinLabel)
+		FXxHash64 ComputeNodeOverrideKey(const FPCGStack& InExecutionStack, const UPCGNode* InNode)
 		{
-			const FXxHash64 PinlessKey = ComputeDataHash(InStack);
-			return ComputeDataHashWithPin(PinlessKey, PinLabel);
+			FPCGStack NodeStack = InExecutionStack;
+			NodeStack.PushFrame(InNode);
+			return ComputeDataHash(NodeStack);
 		}
 
-		FXxHash64 ComputeDataHashWithPin(const FXxHash64 NodeKey, const FName PinLabel)
+		FXxHash64 ComputePinOverrideKey(const FPCGStack& InExecutionStack, const UPCGNode* InNode, const FName PinLabel)
+		{
+			return ComputePinOverrideKey(ComputeNodeOverrideKey(InExecutionStack, InNode), PinLabel);
+		}
+
+		FXxHash64 ComputePinOverrideKey(const FXxHash64 NodeKey, const FName PinLabel)
 		{
 			if (PinLabel == NAME_None)
 			{
@@ -54,7 +60,7 @@ namespace PCG::DataOverride
 			}
 
 			FXxHash64Builder HashBuilder;
-			HashBuilder.Update(&NodeKey, sizeof(NodeKey));
+			HashBuilder.Update(&NodeKey.Hash, sizeof(NodeKey.Hash));
 
 			const FString PinString = PinLabel.ToString();
 			HashBuilder.Update(*PinString, PinString.Len() * sizeof(TCHAR));
@@ -67,8 +73,9 @@ namespace PCG::DataOverride
 	{
 		FXxHash64 ComputePositionHash(const FVector& InLocation, const double Tolerance)
 		{
+			// Note: this assumes there's no padding in FVector that could contain uncompared data.
 			// 0 Tolerance is an exact spatial hash
-			if (FMath::Max(Tolerance, 0.0) == 0.0)
+			if (FMath::IsNearlyZero(FMath::Max(Tolerance, 0.0)))
 			{
 				return FXxHash64::HashBuffer(&InLocation, sizeof(InLocation));
 			}
@@ -84,9 +91,21 @@ namespace PCG::DataOverride
 
 		FXxHash64 ComputeTransformHash(const FTransform& InTransform, const double Tolerance)
 		{
-			if (FMath::Max(Tolerance, 0.0) == 0.0)
+			if (FMath::IsNearlyZero(FMath::Max(Tolerance, 0.0)))
 			{
-				return FXxHash64::HashBuffer(&InTransform, sizeof(InTransform));
+				// FTransform uses vectorized register storage (double[4] per member)
+				// Translation and Scale3D use uninitialized padding so decomposition is needed to hash only meaningful
+				// bytes. This matches the pattern in GetTypeHash(TTransform).
+				const FQuat Rotation = InTransform.GetRotation();
+				const FVector Translation = InTransform.GetTranslation();
+				const FVector Scale = InTransform.GetScale3D();
+
+				FXxHash64Builder HashBuilder;
+				HashBuilder.Update(&Rotation, sizeof(Rotation));
+				HashBuilder.Update(&Translation, sizeof(Translation));
+				HashBuilder.Update(&Scale, sizeof(Scale));
+
+				return HashBuilder.Finalize();
 			}
 
 			const FVector Location = InTransform.GetLocation();
@@ -110,7 +129,8 @@ namespace PCG::DataOverride
 
 		FXxHash64 ComputeAABBHash(const FBox& InBounds, const double Tolerance)
 		{
-			if (FMath::Max(Tolerance, 0.0) == 0.0)
+			// Note: this assumes there's no padding in FBox that could contain uncompared data.
+			if (FMath::IsNearlyZero(FMath::Max(Tolerance, 0.0)))
 			{
 				return FXxHash64::HashBuffer(&InBounds, sizeof(InBounds));
 			}
@@ -126,6 +146,49 @@ namespace PCG::DataOverride
 
 			return FXxHash64::HashBuffer(&Grid, sizeof(Grid));
 		}
+
+		FXxHash64 ComputeCompositeHash(const FTransform& InTransform, const double Tolerance, const bool bIncludePosition, const bool bIncludeRotation, const bool bIncludeScale)
+		{
+			// Fall back to the full transform if no components are selected to avoid degenerate hashes.
+			if (!bIncludePosition && !bIncludeRotation && !bIncludeScale)
+			{
+				return ComputeTransformHash(InTransform, Tolerance);
+			}
+
+			// Fast paths for common cases
+			if (bIncludePosition && bIncludeRotation && bIncludeScale)
+			{
+				return ComputeTransformHash(InTransform, Tolerance);
+			}
+
+			if (bIncludePosition && !bIncludeRotation && !bIncludeScale)
+			{
+				return ComputePositionHash(InTransform.GetLocation(), Tolerance);
+			}
+
+			// Composite: hash selected components individually
+			FXxHash64Builder HashBuilder;
+
+			if (bIncludePosition)
+			{
+				const FXxHash64 PositionHash = ComputePositionHash(InTransform.GetLocation(), Tolerance);
+				HashBuilder.Update(&PositionHash, sizeof(PositionHash));
+			}
+
+			if (bIncludeRotation)
+			{
+				const FQuat Rotation = InTransform.GetRotation();
+				HashBuilder.Update(&Rotation, sizeof(Rotation));
+			}
+
+			if (bIncludeScale)
+			{
+				const FVector Scale = InTransform.GetScale3D();
+				HashBuilder.Update(&Scale, sizeof(Scale));
+			}
+
+			return HashBuilder.Finalize();
+		}
 	} // namespace Keys
 
 	bool HasNodeLevelDataOverrides(FPCGContext* Context)
@@ -137,11 +200,7 @@ namespace PCG::DataOverride
 			return false;
 		}
 
-		// @todo_pcg: Investigate the need to push the node to the stack
-		// First check if the node has data override deltas.
-		FPCGStack NodeStack = *Stack;
-		NodeStack.PushFrame(Context->Node);
-		const FXxHash64 NodeKey = Helpers::ComputeDataHash(NodeStack);
+		const FXxHash64 NodeKey = Helpers::ComputeNodeOverrideKey(*Stack, Context->Node);
 		const FPCGSourceDataStorageKey NodeStorageKey(Constants::DefaultOverrideLabel, NodeKey.Hash);
 		return DataContainer->Get<FPCGDeltaCollection>(NodeStorageKey).IsValid();
 	}
@@ -155,10 +214,7 @@ namespace PCG::DataOverride
 			return false;
 		}
 
-		FPCGStack NodeStack = *Stack;
-		NodeStack.PushFrame(Context->Node);
-		const FXxHash64 NodeKey = Helpers::ComputeDataHash(NodeStack);
-		const FXxHash64 PinKey = Helpers::ComputeDataHashWithPin(NodeKey, InTaggedData.Pin);
+		const FXxHash64 PinKey = Helpers::ComputePinOverrideKey(*Stack, Context->Node, InTaggedData.Pin);
 		const FPCGSourceDataStorageKey PinStorageKey(Constants::DefaultOverrideLabel, PinKey.Hash);
 		return DataContainer->Get<FPCGDeltaCollection>(PinStorageKey).IsValid();
 	}
@@ -196,16 +252,14 @@ namespace PCG::DataOverride
 			return;
 		}
 
-		FPCGStack NodeStack = *Stack;
-		NodeStack.PushFrame(Context->Node);
-		const FXxHash64 NodeKey = Helpers::ComputeDataHash(NodeStack);
+		const FXxHash64 NodeKey = Helpers::ComputeNodeOverrideKey(*Stack, Context->Node);
 		for (FPCGTaggedData& TaggedData : OutTaggedData)
 		{
 			// @todo_pcg: This should be part of a registry. For now, default to points, but should be element
 			// We have data ownership at either phase and thus const_cast is okay.
 			if (UPCGData* PCGData = const_cast<UPCGData*>(Cast<UPCGData>(TaggedData.Data)))
 			{
-				FPCGSourceDataStorageKey PinDataKey(Constants::DefaultOverrideLabel, Helpers::ComputeDataHashWithPin(NodeKey, TaggedData.Pin).Hash);
+				FPCGSourceDataStorageKey PinDataKey(Constants::DefaultOverrideLabel, Helpers::ComputePinOverrideKey(NodeKey, TaggedData.Pin).Hash);
 
 				// Check to see if the pin itself has overrides
 				FConstSharedStruct SharedStruct = DataContainer->Get<FPCGDeltaCollection>(PinDataKey);
@@ -227,7 +281,14 @@ namespace PCG::DataOverride
 					{
 						if (const FPCGDeltaBase* Delta = DeltaStruct.GetPtr<FPCGDeltaBase>())
 						{
-							if (const PCGIndexing::FPCGIndexCollection Candidates = Delta->FilterCandidates(PCGData); !Candidates.IsEmpty())
+							const PCGIndexing::FPCGIndexCollection Candidates = Delta->FilterCandidates(PCGData);
+
+							if (!Candidates.IsValid())
+							{
+								// Delta does not use candidate filtering.
+								Delta->Apply(PCGData, INDEX_NONE);
+							}
+							else if (!Candidates.IsEmpty())
 							{
 								if (const int32 ResolvedIndex = Delta->Resolve(PCGData, Candidates, DeltaCollection->Settings); ResolvedIndex != INDEX_NONE)
 								{
@@ -236,7 +297,7 @@ namespace PCG::DataOverride
 										if (UPCGData* ReplacementData = Delta->GetReplacementData())
 										{
 											TaggedData.Data = ReplacementData;
-											// stop iterating if replacement occurred.
+											// Stop iterating if replacement occurred.
 											return false;
 										}
 										else
@@ -248,9 +309,10 @@ namespace PCG::DataOverride
 									{
 										Delta->Apply(PCGData, ResolvedIndex);
 									}
-							   }
-						   }
+								}
+							}
 						}
+
 						return true;
 					});
 				}
diff --git a/Engine/Plugins/PCG/Source/PCG/Private/Graph/DataOverride/PCGDataOverridePoints.cpp b/Engine/Plugins/PCG/Source/PCG/Private/Graph/DataOverride/PCGDataOverridePoints.cpp
--- a/Engine/Plugins/PCG/Source/PCG/Private/Graph/DataOverride/PCGDataOverridePoints.cpp
+++ b/Engine/Plugins/PCG/Source/PCG/Private/Graph/DataOverride/PCGDataOverridePoints.cpp
@@ -139,7 +139,7 @@ FName FPCGPointInsertionDelta::GetDeltaName() const
 
 PCGIndexing::FPCGIndexCollection FPCGPointInsertionDelta::FilterCandidates(const UPCGData* InData) const
 {
-	// Filtering not needed. Points are simply injected.
+	// Filtering not needed. Points are simply inserted.
 	return PCGIndexing::FPCGIndexCollection::Invalid();
 }
 
@@ -147,25 +147,25 @@ bool FPCGPointInsertionDelta::Apply(UPCGData* InData, const int32 ResolvedIndex)
 {
 	if (UPCGPointArrayData* PointData = Cast<UPCGPointArrayData>(InData))
 	{
-		if (InjectedPoints.IsEmpty())
+		if (InsertedPoints.IsEmpty())
 		{
 			return false;
 		}
 
 		const int32 OriginalPointCount = PointData->GetNumPoints();
-		if (!ensure(std::numeric_limits<int32>::max() - OriginalPointCount >= InjectedPoints.Num()))
+		if (!ensure(std::numeric_limits<int32>::max() - OriginalPointCount >= InsertedPoints.Num()))
 		{
-			// Too many point injections
+			// Too many point insertions
 			return false;
 		}
 
-		PointData->SetNumPoints(OriginalPointCount + InjectedPoints.Num());
+		PointData->SetNumPoints(OriginalPointCount + InsertedPoints.Num());
 
 		FPCGPointValueRanges PointRanges(PointData, /*bAllocate=*/true);
 
-		for (int32 I = 0; I < InjectedPoints.Num(); ++I)
+		for (int32 I = 0; I < InsertedPoints.Num(); ++I)
 		{
-			PointRanges.SetFromPoint(OriginalPointCount + I, InjectedPoints[I]);
+			PointRanges.SetFromPoint(OriginalPointCount + I, InsertedPoints[I]);
 		}
 
 		return true;
diff --git a/Engine/Plugins/PCG/Source/PCG/Private/PCGDebug.cpp b/Engine/Plugins/PCG/Source/PCG/Private/PCGDebug.cpp
--- a/Engine/Plugins/PCG/Source/PCG/Private/PCGDebug.cpp
+++ b/Engine/Plugins/PCG/Source/PCG/Private/PCGDebug.cpp
@@ -6,7 +6,7 @@
 
 namespace PCGDebugVisConstants
 {
-	const FSoftObjectPath DefaultPointMesh = FSoftObjectPath(TEXT("/PCG/DebugObjects/PCG_Cube.PCG_Cube"));
+	extern const FSoftObjectPath DefaultPointMesh = FSoftObjectPath(TEXT("/PCG/DebugObjects/PCG_Cube.PCG_Cube"));
 	const FSoftObjectPath MaterialForDefaultPointMesh = FSoftObjectPath(TEXT("Material'/PCG/DebugObjects/PCG_DebugMaterial.PCG_DebugMaterial'"));
 }
 
diff --git a/Engine/Plugins/PCG/Source/PCG/Private/PCGElement.cpp b/Engine/Plugins/PCG/Source/PCG/Private/PCGElement.cpp
--- a/Engine/Plugins/PCG/Source/PCG/Private/PCGElement.cpp
+++ b/Engine/Plugins/PCG/Source/PCG/Private/PCGElement.cpp
@@ -571,6 +571,15 @@ bool IPCGElement::ConvertInputsIfNeeded(FPCGContext* Context) const
 		if (bHasAnyOverrides)
 		{
 			PCG::DataOverride::ApplyDataOverrides(Context, Context->InputData.TaggedData);
+
+#if WITH_EDITOR
+			// Notify inspection that input data overrides were applied (pre-execute).
+			// @todo_pcg: Track orphaned overrides
+			if (Context->Node && Context->GetStack() && Context->ExecutionSource.IsValid())
+			{
+				Context->ExecutionSource->GetExecutionState().GetInspection().NotifyDataOverridesApplied(Context->Node, Context->GetStack());
+			}
+#endif
 		}
 	}
 
@@ -615,6 +624,15 @@ void IPCGElement::PostExecute(FPCGContext* Context) const
 	{
 		TRACE_CPUPROFILER_EVENT_SCOPE(IPCGElement::PostExecute::ApplyDataOverrides);
 		PCG::DataOverride::ApplyDataOverrides(Context, Context->OutputData.TaggedData);
+
+#if WITH_EDITOR
+		// Notify inspection that output data overrides were applied (post-execute).
+		// @todo_pcg: Track orphaned overrides
+		if (Context->Node && Context->GetStack() && Context->ExecutionSource.IsValid())
+		{
+			Context->ExecutionSource->GetExecutionState().GetInspection().NotifyDataOverridesApplied(Context->Node, Context->GetStack());
+		}
+#endif
 	}
 
 	if (Context->IsCacheEnabled())
diff --git a/Engine/Plugins/PCG/Source/PCG/Private/PCGGraphExecutionInspection.cpp b/Engine/Plugins/PCG/Source/PCG/Private/PCGGraphExecutionInspection.cpp
--- a/Engine/Plugins/PCG/Source/PCG/Private/PCGGraphExecutionInspection.cpp
+++ b/Engine/Plugins/PCG/Source/PCG/Private/PCGGraphExecutionInspection.cpp
@@ -137,7 +137,8 @@ void FPCGGraphExecutionInspection::StoreInspectionData(const FPCGStack* InStack,
 		}
 	}
 
-	if (IsInspecting() || (InNode->GetSettingsInterface() && InNode->GetSettingsInterface()->bViewportEditingEnabled))
+	// Retain inspection data for nodes that are actively being viewport-edited or persistently marked for it.
+	if (IsInspecting() || (InNode->GetSettingsInterface() && (InNode->GetSettingsInterface()->bTemporaryManualEditingEnabled || InNode->GetSettingsInterface()->bMarkedForManualEditing)))
 	{
 		InInputData.MarkUsage(EPCGDataUsage::ComponentInspectionData);
 		InOutputData.MarkUsage(EPCGDataUsage::ComponentInspectionData);
@@ -271,6 +272,11 @@ void FPCGGraphExecutionInspection::ClearInspectionData(bool bClearPerNodeExecuti
 			NodeToStacksTriggeringGPUUploads.Reset();
 			NodeToStacksTriggeringGPUReadbacks.Reset();
 		}
+
+		{
+			PCG::TUniqueScopeLock WriteLock(NodeToStacksWithOverridesAppliedLock);
+			NodeToStacksWithOverridesApplied.Reset();
+		}
 	}
 }
 
@@ -319,6 +325,24 @@ bool FPCGGraphExecutionInspection::DidNodeTriggerCPUToGPUUpload(const UPCGNode*
 	return FoundStacks && FoundStacks->Contains(Stack);
 }
 
+void FPCGGraphExecutionInspection::NotifyDataOverridesApplied(const UPCGNode* InNode, const FPCGStack* InStack) const
+{
+	if (!ensure(InNode) || !ensure(InStack))
+	{
+		return;
+	}
+
+	PCG::TUniqueScopeLock WriteLock(NodeToStacksWithOverridesAppliedLock);
+	NodeToStacksWithOverridesApplied.FindOrAdd(InNode).Add(*InStack);
+}
+
+bool FPCGGraphExecutionInspection::NodeAppliedDataOverrides(const UPCGNode* InNode, const FPCGStack& Stack) const
+{
+	PCG::TSharedScopeLock ReadLock(NodeToStacksWithOverridesAppliedLock);
+	const TSet<FPCGStack>* FoundStacks = NodeToStacksWithOverridesApplied.Find(InNode);
+	return FoundStacks && FoundStacks->Contains(Stack);
+}
+
 void FPCGGraphExecutionInspection::AddReferencedObjects(FReferenceCollector& Collector)
 {
 	PCG::TSharedScopeLock ReadLock(InspectionCacheLock);
diff --git a/Engine/Plugins/PCG/Source/PCG/Public/Graph/DataOverride/PCGDataOverrideHelpers.h b/Engine/Plugins/PCG/Source/PCG/Public/Graph/DataOverride/PCGDataOverrideHelpers.h
--- a/Engine/Plugins/PCG/Source/PCG/Public/Graph/DataOverride/PCGDataOverrideHelpers.h
+++ b/Engine/Plugins/PCG/Source/PCG/Public/Graph/DataOverride/PCGDataOverrideHelpers.h
@@ -8,6 +8,7 @@
 
 struct FXxHash64;
 struct FPCGTaggedData;
+class UPCGNode;
 
 namespace PCG::DataOverride
 {
@@ -18,12 +19,12 @@ namespace PCG::DataOverride
 
 	namespace Helpers
 	{
-		/** Compute the node-level override key from the execution stack. */
-		PCG_API FXxHash64 ComputeDataHash(const FPCGStack& InStack);
-		/** Compute a pin-level override key by combining the stack and pin label. */
-		PCG_API FXxHash64 ComputeDataHashWithPin(const FPCGStack& InStack, FName PinLabel);
+		/** Compute a node-level override key from an execution stack and node. Pushes the node frame and hashes. */
+		PCG_API FXxHash64 ComputeNodeOverrideKey(const FPCGStack& InExecutionStack, const UPCGNode* InNode);
+		/** Compute a pin-level override key from an execution stack, node, and pin label. */
+		PCG_API FXxHash64 ComputePinOverrideKey(const FPCGStack& InExecutionStack, const UPCGNode* InNode, FName PinLabel);
 		/** Compute a pin-level override key by appending a pin label to a precomputed node key. */
-		PCG_API FXxHash64 ComputeDataHashWithPin(FXxHash64 NodeKey, FName PinLabel);
+		PCG_API FXxHash64 ComputePinOverrideKey(FXxHash64 NodeKey, FName PinLabel);
 	} // namespace Helpers
 
 	namespace Keys
@@ -77,6 +78,26 @@ namespace PCG::DataOverride
 			FPCGDeltaKey Key;
 		};
 
+		/** Compute a hash from selected transform components. Short-circuits to existing functions when possible. */
+		PCG_API FXxHash64 ComputeCompositeHash(const FTransform& InTransform, double Tolerance, bool bIncludePosition, bool bIncludeRotation, bool bIncludeScale);
+
+		struct FPCGCompositeTransformDeltaKey
+		{
+			explicit FPCGCompositeTransformDeltaKey(
+				const FTransform& InTransform,
+				const double InTolerance,
+				const bool bPosition,
+				const bool bRotation,
+				const bool bScale,
+				const FName DeltaLabel = NAME_None)
+				: Key(ComputeCompositeHash(InTransform, InTolerance, bPosition, bRotation, bScale), DeltaLabel) {}
+
+			operator FPCGDeltaKey() const { return Key; }
+
+		private:
+			FPCGDeltaKey Key;
+		};
+
 		// @todo_pcg: still to implement in future pass
 		PCG_API FPCGDeltaKey ComputeAttributeKey(const FPCGAttributePropertyInputSelector& InSelector);
 	} // namespace Keys
diff --git a/Engine/Plugins/PCG/Source/PCG/Public/Graph/DataOverride/PCGDataOverridePoints.h b/Engine/Plugins/PCG/Source/PCG/Public/Graph/DataOverride/PCGDataOverridePoints.h
--- a/Engine/Plugins/PCG/Source/PCG/Public/Graph/DataOverride/PCGDataOverridePoints.h
+++ b/Engine/Plugins/PCG/Source/PCG/Public/Graph/DataOverride/PCGDataOverridePoints.h
@@ -49,15 +49,15 @@ struct FPCGPointEditDelta : public FPCGPointDeltaBase
 	UPROPERTY()
 	FTransform TransformOverride;
 
-	/** Whether to override the point's position. */
+	/** Override the point's position. */
 	UPROPERTY()
 	bool bOverridePosition = true;
 
-	/** Whether to override the point's rotation. */
+	/** Override the point's rotation. */
 	UPROPERTY()
 	bool bOverrideRotation = true;
 
-	/** Whether to override the point's scale. */
+	/** Override the point's scale. */
 	UPROPERTY()
 	bool bOverrideScale = true;
 
@@ -85,7 +85,7 @@ struct FPCGPointDeletionDelta : public FPCGPointDeltaBase
 #endif // WITH_EDITORONLY_DATA
 };
 
-/** Delta that injects a new point. */
+/** Delta that inserts a new point. */
 USTRUCT()
 struct FPCGPointInsertionDelta : public FPCGDeltaBase
 {
@@ -98,5 +98,5 @@ struct FPCGPointInsertionDelta : public FPCGDeltaBase
 
 	/** Points to be appended to the point data. */
 	UPROPERTY()
-	TArray<FPCGPoint> InjectedPoints;
+	TArray<FPCGPoint> InsertedPoints;
 };
diff --git a/Engine/Plugins/PCG/Source/PCG/Public/PCGDebug.h b/Engine/Plugins/PCG/Source/PCG/Public/PCGDebug.h
--- a/Engine/Plugins/PCG/Source/PCG/Public/PCGDebug.h
+++ b/Engine/Plugins/PCG/Source/PCG/Public/PCGDebug.h
@@ -47,4 +47,10 @@ public:
 #endif
 };
 
+namespace PCGDebugVisConstants
+{
+	/** Default mesh used for PCG debug point visualization. */
+	UE_API extern const FSoftObjectPath DefaultPointMesh;
+}
+
 #undef UE_API
diff --git a/Engine/Plugins/PCG/Source/PCG/Public/PCGGraphExecutionInspection.h b/Engine/Plugins/PCG/Source/PCG/Public/PCGGraphExecutionInspection.h
--- a/Engine/Plugins/PCG/Source/PCG/Public/PCGGraphExecutionInspection.h
+++ b/Engine/Plugins/PCG/Source/PCG/Public/PCGGraphExecutionInspection.h
@@ -77,6 +77,12 @@ public:
 	PCG_API bool DidNodeTriggerGPUToCPUReadback(const UPCGNode* InNode, const FPCGStack& Stack) const;
 	PCG_API bool DidNodeTriggerCPUToGPUUpload(const UPCGNode* InNode, const FPCGStack& Stack) const;
 
+	/** Called at execution time when data overrides have been applied to a node's data. */
+	PCG_API void NotifyDataOverridesApplied(const UPCGNode* InNode, const FPCGStack* InStack) const;
+
+	/** Returns true if the given node had data overrides applied during the last execution. */
+	PCG_API bool NodeAppliedDataOverrides(const UPCGNode* InNode, const FPCGStack& Stack) const;
+
 	PCG_API void AddReferencedObjects(FReferenceCollector& Collector);
 
 	UE_DEPRECATED(5.8, "Use GetInspectionDataPtr/InspectData instead")
@@ -104,5 +110,9 @@ private:
 	mutable TMap<TObjectKey<const UPCGNode>, TSet<FPCGStack>> NodeToStacksTriggeringGPUUploads;
 	mutable TMap<TObjectKey<const UPCGNode>, TSet<FPCGStack>> NodeToStacksTriggeringGPUReadbacks;
 	mutable PCG::FSharedLock NodeToStacksTriggeringGPUTransfersLock;
+
+	/** Map from nodes to all stacks for which data overrides were applied. */
+	mutable TMap<TObjectKey<const UPCGNode>, TSet<FPCGStack>> NodeToStacksWithOverridesApplied;
+	mutable PCG::FSharedLock NodeToStacksWithOverridesAppliedLock;
 };
 #endif
diff --git a/Engine/Plugins/PCG/Source/PCG/Public/PCGSettings.h b/Engine/Plugins/PCG/Source/PCG/Public/PCGSettings.h
--- a/Engine/Plugins/PCG/Source/PCG/Public/PCGSettings.h
+++ b/Engine/Plugins/PCG/Source/PCG/Public/PCGSettings.h
@@ -219,8 +219,13 @@ public:
 	bool bDebug = false;
 
 #if WITH_EDITORONLY_DATA
-	UPROPERTY(Transient, EditAnywhere, BlueprintReadWrite, Category = Debug, meta = (EditCondition = bDisplayDebuggingProperties, EditConditionHides, HideEditConditionToggle))
-	bool bViewportEditingEnabled = false;
+	/** Set while the user is actively viewport-editing this node. */
+	UPROPERTY(Transient, DuplicateTransient, EditAnywhere, BlueprintReadWrite, Category = Debug, meta = (EditCondition = bDisplayDebuggingProperties, EditConditionHides, HideEditConditionToggle))
+	bool bTemporaryManualEditingEnabled = false;
+
+	/** Saved to the graph asset so viewport editing markers survive across sessions. */
+	UPROPERTY(EditAnywhere, BlueprintReadWrite, Category = Debug, meta = (EditCondition = bDisplayDebuggingProperties, EditConditionHides, HideEditConditionToggle, PCGNoHash))
+	bool bMarkedForManualEditing = false;
 
 	UPROPERTY(EditAnywhere, BlueprintReadWrite, Category = Debug, meta = (ShowOnlyInnerProperties, PCGNoHash))
 	FPCGDebugVisualizationSettings DebugSettings;
diff --git a/Engine/Plugins/PCG/Source/PCGEditor/Private/DeltaViewportExtensions/PCGDeltaViewportExtension.cpp b/Engine/Plugins/PCG/Source/PCGEditor/Private/DeltaViewportExtensions/PCGDeltaViewportExtension.cpp
new file mode 100644
--- /dev/null
+++ b/Engine/Plugins/PCG/Source/PCGEditor/Private/DeltaViewportExtensions/PCGDeltaViewportExtension.cpp
@@ -0,0 +1,83 @@
+// Copyright Epic Games, Inc. All Rights Reserved.
+
+#include "DeltaViewportExtensions/PCGDeltaViewportExtension.h"
+
+#include "Graph/DataOverride/PCGDataOverride.h"
+
+void FPCGDeltaViewportExtensionRegistry::RegisterInternalExtension(const UScriptStruct* DeltaStruct, TUniquePtr<IPCGDeltaViewportExtension> Extension)
+{
+	if (ensureMsgf(!InternalRegistry.Find(DeltaStruct), TEXT("Cannot register multiple internal viewport extensions for the same struct type.")))
+	{
+		InternalRegistry.Add(DeltaStruct, MoveTemp(Extension));
+	}
+}
+
+void FPCGDeltaViewportExtensionRegistry::UnregisterAllInternalExtensions()
+{
+	InternalRegistry.Empty();
+}
+
+void FPCGDeltaViewportExtensionRegistry::RegisterExtension(const UScriptStruct* DeltaStruct, TUniquePtr<IPCGDeltaViewportExtension> Extension)
+{
+	if (ensureMsgf(!ExternalRegistry.Find(DeltaStruct), TEXT("Cannot register multiple viewport extensions for the same struct type.")))
+	{
+		ExternalRegistry.Add(DeltaStruct, MoveTemp(Extension));
+	}
+}
+
+void FPCGDeltaViewportExtensionRegistry::UnregisterExtension(const UScriptStruct* DeltaStruct)
+{
+	ExternalRegistry.Remove(DeltaStruct);
+}
+
+IPCGDeltaViewportExtension* FPCGDeltaViewportExtensionRegistry::GetExtension(const UScriptStruct* DeltaStruct) const
+{
+	const UScriptStruct* Current = DeltaStruct;
+	const UScriptStruct* DeltaBaseStruct = FPCGDeltaBase::StaticStruct();
+
+	// Walk the struct hierarchy from the concrete type up through FPCGDeltaBase (inclusive).
+	while (Current)
+	{
+		// Always try external implementations first so users can override built-in extensions.
+		if (const TUniquePtr<IPCGDeltaViewportExtension>* ExtensionPtr = ExternalRegistry.Find(Current))
+		{
+			return ExtensionPtr->Get();
+		}
+
+		if (const TUniquePtr<IPCGDeltaViewportExtension>* ExtensionPtr = InternalRegistry.Find(Current))
+		{
+			return ExtensionPtr->Get();
+		}
+
+		// Stop after checking FPCGDeltaBase -- don't walk into non-delta struct territory.
+		if (Current == DeltaBaseStruct)
+		{
+			break;
+		}
+
+		Current = Cast<UScriptStruct>(Current->GetSuperStruct());
+	}
+
+	return nullptr;
+}
+
+TArray<const UScriptStruct*> FPCGDeltaViewportExtensionRegistry::GetRegisteredDeltaTypes() const
+{
+	TMap<const UScriptStruct*, int32> TypePriorities;
+
+	for (const auto& [DeltaStruct, Extension] : InternalRegistry)
+	{
+		TypePriorities.Add(DeltaStruct, Extension ? Extension->GetSortPriority() : std::numeric_limits<int32>::max());
+	}
+
+	for (const auto& [DeltaStruct, Extension] : ExternalRegistry)
+	{
+		TypePriorities.Add(DeltaStruct, Extension ? Extension->GetSortPriority() : std::numeric_limits<int32>::max());
+	}
+
+	TypePriorities.ValueSort([](const int32 A, const int32 B) { return A < B; });
+
+	TArray<const UScriptStruct*> Result;
+	TypePriorities.GenerateKeyArray(Result);
+	return Result;
+}
diff --git a/Engine/Plugins/PCG/Source/PCGEditor/Private/DeltaViewportExtensions/PCGDeltaViewportExtensionHelpers.cpp b/Engine/Plugins/PCG/Source/PCGEditor/Private/DeltaViewportExtensions/PCGDeltaViewportExtensionHelpers.cpp
new file mode 100644
--- /dev/null
+++ b/Engine/Plugins/PCG/Source/PCGEditor/Private/DeltaViewportExtensions/PCGDeltaViewportExtensionHelpers.cpp
@@ -0,0 +1,20 @@
+// Copyright Epic Games, Inc. All Rights Reserved.
+
+#include "DeltaViewportExtensions/PCGDeltaViewportExtensionHelpers.h"
+
+#include "PCGComponent.h"
+#include "PCGGraph.h"
+#include "Graph/PCGSourceDataContainer.h"
+
+void PCG::DeltaViewportExtension::MarkComponentDirty(UPCGComponent* PCGComponent, FPCGSourceDataContainer* DataContainer)
+{
+	DataContainer->MarkDirty();
+
+	if (UPCGGraphInstance* GraphInstance = PCGComponent->GetGraphInstance())
+	{
+		if (UPCGGraph* Graph = GraphInstance->GetMutablePCGGraph())
+		{
+			Graph->ForceNotificationForEditor(EPCGChangeType::ExternalModification);
+		}
+	}
+}
diff --git a/Engine/Plugins/PCG/Source/PCGEditor/Private/DeltaViewportExtensions/PCGDeltaViewportExtensionHelpers.h b/Engine/Plugins/PCG/Source/PCGEditor/Private/DeltaViewportExtensions/PCGDeltaViewportExtensionHelpers.h
new file mode 100644
--- /dev/null
+++ b/Engine/Plugins/PCG/Source/PCGEditor/Private/DeltaViewportExtensions/PCGDeltaViewportExtensionHelpers.h
@@ -0,0 +1,12 @@
+// Copyright Epic Games, Inc. All Rights Reserved.
+
+#pragma once
+
+class UPCGComponent;
+struct FPCGSourceDataContainer;
+
+namespace PCG::DeltaViewportExtension
+{
+	/** Mark the PCG component and its source data container as dirty, then notify the graph editor. */
+	void MarkComponentDirty(UPCGComponent* PCGComponent, FPCGSourceDataContainer* DataContainer);
+}
diff --git a/Engine/Plugins/PCG/Source/PCGEditor/Private/DeltaViewportExtensions/PCGPointDeletionViewportExtension.cpp b/Engine/Plugins/PCG/Source/PCGEditor/Private/DeltaViewportExtensions/PCGPointDeletionViewportExtension.cpp
new file mode 100644
--- /dev/null
+++ b/Engine/Plugins/PCG/Source/PCGEditor/Private/DeltaViewportExtensions/PCGPointDeletionViewportExtension.cpp
@@ -0,0 +1,433 @@
+// Copyright Epic Games, Inc. All Rights Reserved.
+
+#include "DeltaViewportExtensions/PCGPointDeletionViewportExtension.h"
+
+#include "DeltaViewportExtensions/PCGDeltaViewportExtensionHelpers.h"
+#include "PCGComponent.h"
+
+#include "ScopedTransaction.h"
+#include "Styling/AppStyle.h"
+#include "Styling/CoreStyle.h"
+
+#include "SPCGManualEditPanel.h"
+#include "Widgets/SBoxPanel.h"
+#include "Widgets/Images/SImage.h"
+#include "Widgets/Input/SButton.h"
+#include "Widgets/Text/STextBlock.h"
+#include "Widgets/Views/SListView.h"
+
+#define LOCTEXT_NAMESPACE "PCGPointDeletionViewportExtension"
+
+namespace PCGPointDeletionDelta::Constants
+{
+	constexpr FLinearColor SelectedColor(1.0f, 0.4f, 0.4f, 0.5f);
+	constexpr FLinearColor NormalColor(0.65f, 0.05f, 0.05f, 0.3f);
+}
+
+TSharedRef<SWidget> FPCGPointDeletionViewportExtension::CreateWidget(FPCGDeltaViewportCallbacks Callbacks)
+{
+	CachedCallbacks = MoveTemp(Callbacks);
+	DeletionEntries.Empty();
+
+	TSharedRef<SVerticalBox> BaseBox = SNew(SVerticalBox)
+		// Delete / Restore Selected button
+		+ SVerticalBox::Slot()
+	   .AutoHeight()
+	   .Padding(0.0f, 4.0f, 0.0f, 0.0f)
+		[
+			SNew(SButton)
+		   .HAlign(HAlign_Center)
+		   .IsEnabled_Lambda([this]() -> bool
+			{
+				return CachedContext.bSelectionActive;
+			})
+		   .OnClicked_Lambda([this]() -> FReply
+			{
+				HandleDeletePressed();
+				return FReply::Handled();
+			})
+			[
+				SNew(STextBlock)
+			   .Text_Lambda([this]() -> FText
+				{
+					return IsSelectionDeletionDelta()
+							   ? LOCTEXT("RestoreSelectedButton", "Restore Selected")
+							   : LOCTEXT("DeleteSelectedButton", "Delete Selected");
+				})
+			   .Font(FCoreStyle::GetDefaultFontStyle("Regular", 8))
+			]
+		]
+		// Deletions list section
+		+ SVerticalBox::Slot()
+	   .AutoHeight()
+	   .Padding(0.0f, 8.0f, 0.0f, 2.0f)
+		[
+			SAssignNew(DeletionSection, SVerticalBox)
+		   .Visibility(EVisibility::Collapsed)
+			// Header row: "Deletions" + "Restore All"
+			+ SVerticalBox::Slot()
+		   .AutoHeight()
+		   .Padding(0.0f, 0.0f, 0.0f, 4.0f)
+			[
+				SNew(SHorizontalBox)
+				+ SHorizontalBox::Slot()
+			   .FillWidth(1.0f)
+			   .VAlign(VAlign_Center)
+				[
+					SNew(STextBlock)
+				   .Text(LOCTEXT("DeletionsHeader", "Deletions"))
+				   .Font(FCoreStyle::GetDefaultFontStyle("Bold", 8))
+				]
+				+ SHorizontalBox::Slot()
+			   .AutoWidth()
+			   .VAlign(VAlign_Center)
+				[
+					SNew(SButton)
+				   .OnClicked_Lambda([this]() -> FReply
+					{
+						if (CachedCallbacks.RevertAllRestorableDeltas)
+						{
+							CachedCallbacks.RevertAllRestorableDeltas();
+						}
+						RefreshLists(CachedContext);
+						return FReply::Handled();
+					})
+					[
+						SNew(STextBlock)
+					   .Text(LOCTEXT("RestoreAllButton", "Restore All"))
+					   .Font(FCoreStyle::GetDefaultFontStyle("Regular", 8))
+					]
+				]
+			]
+			// Deletion list view
+			+ SVerticalBox::Slot()
+		   .AutoHeight()
+		   .MaxHeight(150.0f)
+			[
+				SAssignNew(DeletionListView, SListView<TSharedPtr<FPCGPointDeletionDeltaEntry>>)
+			   .ListItemsSource(&DeletionEntries)
+			   .OnGenerateRow_Lambda([this](TSharedPtr<FPCGPointDeletionDeltaEntry> Entry, const TSharedRef<STableViewBase>& OwnerTable) -> TSharedRef<ITableRow>
+				{
+					return OnGenerateDeletionRow(Entry, OwnerTable);
+				})
+			   .SelectionMode(ESelectionMode::None)
+			]
+		];
+
+	return BaseBox;
+}
+
+void FPCGPointDeletionViewportExtension::RefreshLists(const FPCGDeltaViewportContext& Context)
+{
+	CachedContext = Context;
+
+	TArray<TSharedPtr<FPCGPointDeletionDeltaEntry>> NewEntries;
+
+	if (UPCGComponent* PCGComponent = Context.ActivePCGComponent.Get())
+	{
+		if (const FPCGSourceDataContainer* DataContainer = PCGComponent->GetExecutionState().GetSourceDataContainer())
+		{
+			const FConstSharedStruct SharedStruct = DataContainer->Get<FPCGDeltaCollection>(Context.ActiveStorageKey);
+			if (const FPCGDeltaCollection* Collection = SharedStruct.IsValid() ? SharedStruct.GetPtr<const FPCGDeltaCollection>() : nullptr)
+			{
+				Collection->ForEachDelta([&NewEntries](const FPCGDeltaKey& DeltaKey, const TInstancedStruct<FPCGDeltaBase>& DeltaValue) -> bool
+				{
+					if (const FPCGPointDeletionDelta* DeleteDelta = DeltaValue.GetPtr<FPCGPointDeletionDelta>())
+					{
+						TSharedPtr<FPCGPointDeletionDeltaEntry> Entry = MakeShared<FPCGPointDeletionDeltaEntry>();
+						Entry->Key = DeltaKey;
+						Entry->OriginalElementIndex = DeleteDelta->ElementIndex;
+						Entry->Location = DeleteDelta->OriginalTransform.GetLocation();
+						NewEntries.Add(Entry);
+					}
+
+					return true;
+				});
+			}
+		}
+	}
+
+	bool bListChanged = NewEntries.Num() != DeletionEntries.Num();
+	if (!bListChanged)
+	{
+		for (int32 i = 0; i < NewEntries.Num(); ++i)
+		{
+			if (NewEntries[i]->Key != DeletionEntries[i]->Key)
+			{
+				bListChanged = true;
+				break;
+			}
+
+			DeletionEntries[i]->Location = NewEntries[i]->Location;
+		}
+	}
+
+	if (bListChanged)
+	{
+		DeletionEntries = MoveTemp(NewEntries);
+
+		if (DeletionListView.IsValid())
+		{
+			DeletionListView->RequestListRefresh();
+		}
+	}
+
+	const bool bHasDeletions = DeletionEntries.Num() > 0;
+	if (DeletionSection.IsValid())
+	{
+		DeletionSection->SetVisibility(bHasDeletions ? EVisibility::SelfHitTestInvisible : EVisibility::Collapsed);
+	}
+}
+
+FLinearColor FPCGPointDeletionViewportExtension::GetDisplayColor(const bool bIsSelected) const
+{
+	return bIsSelected ? PCGPointDeletionDelta::Constants::SelectedColor : PCGPointDeletionDelta::Constants::NormalColor;
+}
+
+FTransform FPCGPointDeletionViewportExtension::GetDisplayTransform(const FTransform& SourceElementTransform, const FConstStructView DeltaStruct) const
+{
+	if (const FPCGPointDeletionDelta* DeleteDelta = DeltaStruct.GetPtr<const FPCGPointDeletionDelta>())
+	{
+		return DeleteDelta->OriginalTransform;
+	}
+
+	return SourceElementTransform;
+}
+
+bool FPCGPointDeletionViewportExtension::MatchesSourceElement(const FTransform& SourceElementTransform, const FConstStructView DeltaStruct, const double SpatialTolerance) const
+{
+	const FPCGPointDeletionDelta* DeleteDelta = DeltaStruct.GetPtr<const FPCGPointDeletionDelta>();
+	if (!DeleteDelta)
+	{
+		return false;
+	}
+
+	if (!CachedContext.NodeConfiguration)
+	{
+		return false;
+	}
+
+	const FPCGDeltaKey SourceKey = PCG::DataOverride::Keys::FPCGCompositeTransformDeltaKey(
+		SourceElementTransform,
+		SpatialTolerance,
+		CachedContext.NodeConfiguration->bSignaturePosition,
+		CachedContext.NodeConfiguration->bSignatureRotation,
+		CachedContext.NodeConfiguration->bSignatureScale);
+
+	const FPCGDeltaKey DeltaKey = PCG::DataOverride::Keys::FPCGCompositeTransformDeltaKey(
+		DeleteDelta->OriginalTransform,
+		SpatialTolerance,
+		CachedContext.NodeConfiguration->bSignaturePosition,
+		CachedContext.NodeConfiguration->bSignatureRotation,
+		CachedContext.NodeConfiguration->bSignatureScale);
+
+	return SourceKey == DeltaKey;
+}
+
+FTransform FPCGPointDeletionViewportExtension::GetOriginalTransform(const FConstStructView DeltaStruct) const
+{
+	if (const FPCGPointDeletionDelta* DeleteDelta = DeltaStruct.GetPtr<const FPCGPointDeletionDelta>())
+	{
+		return DeleteDelta->OriginalTransform;
+	}
+
+	return FTransform::Identity;
+}
+
+void FPCGPointDeletionViewportExtension::CollectRestorableKeys(const FPCGDeltaCollection& Collection, TArray<FPCGDeltaKey>& OutKeys) const
+{
+	Collection.ForEachDelta([&OutKeys](const FPCGDeltaKey& DeltaKey, const TInstancedStruct<FPCGDeltaBase>& DeltaValue) -> bool
+	{
+		if (DeltaValue.GetPtr<FPCGPointDeletionDelta>())
+		{
+			OutKeys.Add(DeltaKey);
+		}
+
+		return true;
+	});
+}
+
+void FPCGPointDeletionViewportExtension::HandleDeletePressed()
+{
+	if (!CachedContext.bSelectionActive || CachedContext.SelectionDeltaKey.Hash == 0)
+	{
+		RefreshLists(CachedContext);
+		return;
+	}
+
+	UPCGComponent* PCGComponent = CachedContext.ActivePCGComponent.Get();
+	if (!PCGComponent)
+	{
+		RefreshLists(CachedContext);
+		return;
+	}
+
+	FPCGSourceDataContainer* DataContainer = PCGComponent->GetExecutionState().GetSourceDataContainer();
+	if (!DataContainer)
+	{
+		RefreshLists(CachedContext);
+		return;
+	}
+
+	FSharedStruct SharedStruct = DataContainer->GetMutable<FPCGDeltaCollection>(CachedContext.ActiveStorageKey);
+	FPCGDeltaCollection* Collection = SharedStruct.IsValid() ? SharedStruct.GetPtr<FPCGDeltaCollection>() : nullptr;
+	if (!Collection)
+	{
+		RefreshLists(CachedContext);
+		return;
+	}
+
+	if (IsSelectionDeletionDelta())
+	{
+		// RevertDelta creates its own FScopedTransaction.
+		if (CachedCallbacks.RevertDelta)
+		{
+			CachedCallbacks.RevertDelta(CachedContext.SelectionDeltaKey);
+		}
+	}
+	else
+	{
+		FScopedTransaction Transaction(LOCTEXT("CreateDeletionDelta", "PCG Manual Edit: Create Deletion"));
+		PCGComponent->Modify();
+
+		FPCGPointDeletionDelta DeletionDelta;
+		DeletionDelta.OriginalTransform = CachedContext.SelectionTransform;
+		DeletionDelta.Bounds = PCGComponent->GetTotalBounds();
+		DeletionDelta.ElementIndex = CachedContext.OriginalElementIndex;
+
+		if (!CachedContext.NodeConfiguration)
+		{
+			return;
+		}
+
+		const FPCGDeltaKey NewKey = PCG::DataOverride::Keys::FPCGCompositeTransformDeltaKey(
+			CachedContext.SelectionTransform,
+			CachedContext.NodeConfiguration->SignatureTolerance,
+			CachedContext.NodeConfiguration->bSignaturePosition,
+			CachedContext.NodeConfiguration->bSignatureRotation,
+			CachedContext.NodeConfiguration->bSignatureScale,
+			FPCGPointDeletionDelta::GetDeltaNameStatic());
+
+		Collection->Add(NewKey, TInstancedStruct<FPCGPointDeletionDelta>::Make(MoveTemp(DeletionDelta)));
+
+		PCG::DeltaViewportExtension::MarkComponentDirty(PCGComponent, DataContainer);
+
+		if (CachedCallbacks.ClearSelection)
+		{
+			CachedCallbacks.ClearSelection();
+		}
+	}
+
+	RefreshLists(CachedContext);
+}
+
+TSharedRef<ITableRow> FPCGPointDeletionViewportExtension::OnGenerateDeletionRow(TSharedPtr<FPCGPointDeletionDeltaEntry> Entry, const TSharedRef<STableViewBase>& OwnerTable)
+{
+	FPCGDeltaKey EntryKey = Entry->Key;
+
+	return SNew(STableRow<TSharedPtr<FPCGPointDeletionDeltaEntry>>, OwnerTable)
+		[
+			SNew(SHorizontalBox)
+			+ SHorizontalBox::Slot()
+		   .AutoWidth()
+		   .VAlign(VAlign_Center)
+		   .Padding(4.0f, 1.0f, 2.0f, 1.0f)
+			[
+				SNew(STextBlock)
+			   .Text(FText::AsNumber(Entry->OriginalElementIndex))
+			   .Font(FCoreStyle::GetDefaultFontStyle("Regular", 7))
+			]
+			+ SHorizontalBox::Slot()
+		   .FillWidth(1.0f)
+		   .VAlign(VAlign_Center)
+		   .Padding(2.0f, 1.0f)
+			[
+				SNew(STextBlock)
+			   .Text(FText::Format(LOCTEXT("DeletionLocationFormat", "({0}, {1}, {2})"),
+								   FText::AsNumber(FMath::RoundToInt(Entry->Location.X)),
+								   FText::AsNumber(FMath::RoundToInt(Entry->Location.Y)),
+								   FText::AsNumber(FMath::RoundToInt(Entry->Location.Z))))
+			   .Font(FCoreStyle::GetDefaultFontStyle("Regular", 7))
+			]
+			+ SHorizontalBox::Slot()
+		   .AutoWidth()
+		   .VAlign(VAlign_Center)
+		   .Padding(2.0f, 0.0f)
+			[
+				SNew(SButton)
+			   .ButtonStyle(FAppStyle::Get(), "SimpleButton")
+			   .ContentPadding(FMargin(1, 0))
+			   .ToolTipText(LOCTEXT("SelectDeletionTooltip", "Select this deleted point"))
+			   .OnClicked_Lambda([this, EntryKey]() -> FReply
+				{
+					if (CachedCallbacks.SelectElement)
+					{
+						CachedCallbacks.SelectElement(EntryKey, INDEX_NONE);
+					}
+					return FReply::Handled();
+				})
+				[
+					SNew(SImage)
+				   .Image(FAppStyle::Get().GetBrush("Icons.Search"))
+				   .ColorAndOpacity(FSlateColor::UseForeground())
+				]
+			]
+			+ SHorizontalBox::Slot()
+		   .AutoWidth()
+		   .VAlign(VAlign_Center)
+		   .Padding(2.0f, 0.0f)
+			[
+				SNew(SButton)
+			   .ButtonStyle(FAppStyle::Get(), "SimpleButton")
+			   .ContentPadding(FMargin(1, 0))
+			   .ToolTipText(LOCTEXT("RestoreButtonTooltip", "Restore this deleted point"))
+			   .OnClicked_Lambda([this, EntryKey]() -> FReply
+				{
+					if (CachedCallbacks.RevertDelta)
+					{
+						CachedCallbacks.RevertDelta(EntryKey);
+					}
+					RefreshLists(CachedContext);
+
+					return FReply::Handled();
+				})
+				[
+					SNew(SImage)
+				   .Image(FAppStyle::Get().GetBrush("GenericCommands.Undo"))
+				   .ColorAndOpacity(FSlateColor::UseForeground())
+				]
+			]
+		];
+}
+
+bool FPCGPointDeletionViewportExtension::IsSelectionDeletionDelta() const
+{
+	if (!CachedContext.bSelectionActive || CachedContext.SelectionDeltaKey.Hash == 0)
+	{
+		return false;
+	}
+
+	UPCGComponent* PCGComponent = CachedContext.ActivePCGComponent.Get();
+	if (!PCGComponent)
+	{
+		return false;
+	}
+
+	const FPCGSourceDataContainer* DataContainer = PCGComponent->GetExecutionState().GetSourceDataContainer();
+	if (!DataContainer)
+	{
+		return false;
+	}
+
+	const FConstSharedStruct SharedStruct = DataContainer->Get<FPCGDeltaCollection>(CachedContext.ActiveStorageKey);
+	const FPCGDeltaCollection* Collection = SharedStruct.IsValid() ? SharedStruct.GetPtr<const FPCGDeltaCollection>() : nullptr;
+	if (!Collection)
+	{
+		return false;
+	}
+
+	const TInstancedStruct<FPCGDeltaBase>* Delta = Collection->Find(CachedContext.SelectionDeltaKey);
+	return Delta && Delta->GetPtr<FPCGPointDeletionDelta>() != nullptr;
+}
+
+#undef LOCTEXT_NAMESPACE
diff --git a/Engine/Plugins/PCG/Source/PCGEditor/Private/DeltaViewportExtensions/PCGPointDeletionViewportExtension.h b/Engine/Plugins/PCG/Source/PCGEditor/Private/DeltaViewportExtensions/PCGPointDeletionViewportExtension.h
new file mode 100644
--- /dev/null
+++ b/Engine/Plugins/PCG/Source/PCGEditor/Private/DeltaViewportExtensions/PCGPointDeletionViewportExtension.h
@@ -0,0 +1,54 @@
+// Copyright Epic Games, Inc. All Rights Reserved.
+
+#pragma once
+
+#include "DeltaViewportExtensions/PCGDeltaViewportExtension.h"
+
+#include "Graph/DataOverride/PCGDataOverrideHelpers.h"
+#include "Graph/DataOverride/PCGDataOverridePoints.h"
+
+#include "Widgets/Views/SListView.h"
+
+struct FPCGPointDeletionDeltaEntry
+{
+	FPCGDeltaKey Key;
+	int32 OriginalElementIndex = INDEX_NONE;
+	FVector Location = FVector::ZeroVector;
+};
+
+/** Viewport extension for FPCGPointDeletionDelta. */
+class FPCGPointDeletionViewportExtension : public IPCGDeltaViewportExtension
+{
+public:
+	// ~Begin IPCGDeltaViewportExtension interface
+	virtual FText GetDisplayName() const override { return NSLOCTEXT("PCGPointDeletionViewportExtension", "PointDeletionDisplayName", "Point Deletion"); }
+	virtual FName GetDeltaName() const override { return FPCGPointDeletionDelta::GetDeltaNameStatic(); }
+	virtual int32 GetSortPriority() const override { return 2; }
+	virtual TSharedRef<SWidget> CreateWidget(FPCGDeltaViewportCallbacks Callbacks) override;
+	virtual void UpdateContext(const FPCGDeltaViewportContext& Context) override { CachedContext = Context; }
+	virtual void RefreshLists(const FPCGDeltaViewportContext& Context) override;
+	virtual TArray<FPCGDeltaKeyBinding> GetKeyBindings() override { return {{EKeys::Delete, [this]() { HandleDeletePressed(); }}}; }
+	virtual FLinearColor GetDisplayColor(const bool bIsSelected) const override;
+	virtual FTransform GetDisplayTransform(const FTransform& SourceElementTransform, const FConstStructView DeltaStruct) const override;
+	virtual bool MatchesSourceElement(const FTransform& SourceElementTransform, const FConstStructView DeltaStruct, const double SpatialTolerance) const override;
+	virtual bool UsesTRSGizmo() const override { return false; }
+	virtual FTransform GetOriginalTransform(const FConstStructView DeltaStruct) const override;
+	virtual void CollectRestorableKeys(const FPCGDeltaCollection& Collection, TArray<FPCGDeltaKey>& OutKeys) const override;
+	// ~End IPCGDeltaViewportExtension interface
+
+private:
+	TSharedRef<ITableRow> OnGenerateDeletionRow(TSharedPtr<FPCGPointDeletionDeltaEntry> Entry, const TSharedRef<STableViewBase>& OwnerTable);
+
+	/** Returns true if the current selection points to a deletion delta in the collection. */
+	bool IsSelectionDeletionDelta() const;
+
+	/** Toggles deletion on the selected point: creates a deletion delta or restores an existing one. */
+	void HandleDeletePressed();
+
+	FPCGDeltaViewportContext CachedContext;
+	FPCGDeltaViewportCallbacks CachedCallbacks;
+
+	TArray<TSharedPtr<FPCGPointDeletionDeltaEntry>> DeletionEntries;
+	TSharedPtr<SVerticalBox> DeletionSection;
+	TSharedPtr<SListView<TSharedPtr<FPCGPointDeletionDeltaEntry>>> DeletionListView;
+};
diff --git a/Engine/Plugins/PCG/Source/PCGEditor/Private/DeltaViewportExtensions/PCGPointEditViewportExtension.cpp b/Engine/Plugins/PCG/Source/PCGEditor/Private/DeltaViewportExtensions/PCGPointEditViewportExtension.cpp
new file mode 100644
--- /dev/null
+++ b/Engine/Plugins/PCG/Source/PCGEditor/Private/DeltaViewportExtensions/PCGPointEditViewportExtension.cpp
@@ -0,0 +1,659 @@
+// Copyright Epic Games, Inc. All Rights Reserved.
+
+#include "DeltaViewportExtensions/PCGPointEditViewportExtension.h"
+#include "DeltaViewportExtensions/PCGDeltaViewportExtensionHelpers.h"
+#include "SPCGManualEditPanel.h"
+
+#include "PCGComponent.h"
+#include "Graph/DataOverride/PCGDataOverrideHelpers.h"
+#include "Graph/DataOverride/PCGDataOverridePoints.h"
+
+#include "ScopedTransaction.h"
+#include "Styling/AppStyle.h"
+#include "Styling/CoreStyle.h"
+#include "Widgets/SBoxPanel.h"
+#include "Widgets/Images/SImage.h"
+#include "Widgets/Input/SButton.h"
+#include "Widgets/Input/SCheckBox.h"
+#include "Widgets/Text/STextBlock.h"
+#include "Widgets/Views/SListView.h"
+
+#include "KismetPins/SVectorSlider.h"
+
+#define LOCTEXT_NAMESPACE "PCGPointEditViewportExtension"
+
+namespace PCGPointEditDelta::Constants
+{
+	constexpr FLinearColor SelectedColor(1.0f, 0.8f, 0.3f, 0.5f);
+	constexpr FLinearColor NormalColor(0.75f, 0.45f, 0.0f, 0.3f);
+}
+
+TSharedRef<SWidget> FPCGPointEditViewportExtension::CreateWidget(FPCGDeltaViewportCallbacks Callbacks)
+{
+	CachedCallbacks = MoveTemp(Callbacks);
+
+	return SNew(SVerticalBox)
+		// Override checkboxes — which transform components the edit delta writes back
+		+ SVerticalBox::Slot()
+		.AutoHeight()
+		.Padding(0.0f, 2.0f, 0.0f, 0.0f)
+		[
+			SNew(SHorizontalBox)
+			+ SHorizontalBox::Slot()
+			.AutoWidth()
+			.VAlign(VAlign_Center)
+			[
+				SNew(SBox)
+				.MinDesiredWidth(60.0f)
+				[
+					SNew(STextBlock)
+					.Text(LOCTEXT("OverrideLabel", "Override:"))
+					.Font(FCoreStyle::GetDefaultFontStyle("Regular", 8))
+				]
+			]
+			+ SHorizontalBox::Slot()
+			.AutoWidth()
+			.VAlign(VAlign_Center)
+			.Padding(0.0f, 0.0f, 4.0f, 0.0f)
+			[
+				SNew(SCheckBox)
+				.IsChecked_Lambda([this]() -> ECheckBoxState
+				{
+					return (bOverridePosition) ? ECheckBoxState::Checked : ECheckBoxState::Unchecked;
+				})
+				.OnCheckStateChanged_Lambda([this](const ECheckBoxState NewState)
+				{
+					bOverridePosition = (NewState == ECheckBoxState::Checked);
+					UpdateDeltaOverrideFlags();
+				})
+				[
+					SNew(STextBlock)
+					.Text(LOCTEXT("OverridePositionCheckbox", "Pos"))
+					.Font(FCoreStyle::GetDefaultFontStyle("Regular", 8))
+				]
+			]
+			+ SHorizontalBox::Slot()
+			.AutoWidth()
+			.VAlign(VAlign_Center)
+			.Padding(0.0f, 0.0f, 4.0f, 0.0f)
+			[
+				SNew(SCheckBox)
+				.IsChecked_Lambda([this]() -> ECheckBoxState
+				{
+					return (bOverrideRotation) ? ECheckBoxState::Checked : ECheckBoxState::Unchecked;
+				})
+				.OnCheckStateChanged_Lambda([this](const ECheckBoxState NewState)
+				{
+					bOverrideRotation = (NewState == ECheckBoxState::Checked);
+					UpdateDeltaOverrideFlags();
+				})
+				[
+					SNew(STextBlock)
+					.Text(LOCTEXT("OverrideRotationCheckbox", "Rot"))
+					.Font(FCoreStyle::GetDefaultFontStyle("Regular", 8))
+				]
+			]
+			+ SHorizontalBox::Slot()
+			.AutoWidth()
+			.VAlign(VAlign_Center)
+			.Padding(0.0f, 0.0f, 4.0f, 0.0f)
+			[
+				SNew(SCheckBox)
+				.IsChecked_Lambda([this]() -> ECheckBoxState
+				{
+					return (bOverrideScale) ? ECheckBoxState::Checked : ECheckBoxState::Unchecked;
+				})
+				.OnCheckStateChanged_Lambda([this](const ECheckBoxState NewState)
+				{
+					bOverrideScale = (NewState == ECheckBoxState::Checked);
+					UpdateDeltaOverrideFlags();
+				})
+				[
+					SNew(STextBlock)
+					.Text(LOCTEXT("OverrideScaleCheckbox", "Sca"))
+					.Font(FCoreStyle::GetDefaultFontStyle("Regular", 8))
+				]
+			]
+		]
+		// Edits list section
+		+ SVerticalBox::Slot()
+		.AutoHeight()
+		.Padding(0.0f, 8.0f, 0.0f, 2.0f)
+		[
+			SAssignNew(EditListSection, SVerticalBox)
+			.Visibility(EVisibility::Collapsed)
+			// Header row: "Edits" + "Remove All"
+			+ SVerticalBox::Slot()
+			.AutoHeight()
+			.Padding(0.0f, 0.0f, 0.0f, 4.0f)
+			[
+				SNew(SHorizontalBox)
+				+ SHorizontalBox::Slot()
+				.FillWidth(1.0f)
+				.VAlign(VAlign_Center)
+				[
+					SNew(STextBlock)
+					.Text(LOCTEXT("EditsHeader", "Edits"))
+					.Font(FCoreStyle::GetDefaultFontStyle("Bold", 8))
+				]
+				+ SHorizontalBox::Slot()
+				.AutoWidth()
+				.VAlign(VAlign_Center)
+				[
+					SNew(SButton)
+					.OnClicked_Lambda([this]() -> FReply
+					{
+						RemoveAllEdits();
+						RefreshLists(CachedContext);
+						return FReply::Handled();
+					})
+					[
+						SNew(STextBlock)
+						.Text(LOCTEXT("RemoveAllEditsButton", "Remove All"))
+						.Font(FCoreStyle::GetDefaultFontStyle("Regular", 8))
+					]
+				]
+			]
+			// Edit list view
+			+ SVerticalBox::Slot()
+			.AutoHeight()
+			.MaxHeight(150.0f)
+			[
+				SAssignNew(EditListView, SListView<TSharedPtr<FPCGPointEditDeltaEntry>>)
+				.ListItemsSource(&EditEntries)
+				.OnGenerateRow_Lambda([this](TSharedPtr<FPCGPointEditDeltaEntry> Entry, const TSharedRef<STableViewBase>& OwnerTable) -> TSharedRef<ITableRow>
+				{
+					return OnGenerateEditRow(Entry, OwnerTable);
+				})
+				.SelectionMode(ESelectionMode::None)
+			]
+		];
+}
+
+void FPCGPointEditViewportExtension::RefreshLists(const FPCGDeltaViewportContext& Context)
+{
+	CachedContext = Context;
+
+	TArray<TSharedPtr<FPCGPointEditDeltaEntry>> NewEntries;
+
+	if (UPCGComponent* PCGComponent = Context.ActivePCGComponent.Get())
+	{
+		if (FPCGSourceDataContainer* DataContainer = PCGComponent->GetExecutionState().GetSourceDataContainer())
+		{
+			const FConstSharedStruct SharedStruct = DataContainer->Get<FPCGDeltaCollection>(Context.ActiveStorageKey);
+			if (const FPCGDeltaCollection* Collection = SharedStruct.IsValid() ? SharedStruct.GetPtr<const FPCGDeltaCollection>() : nullptr)
+			{
+				Collection->ForEachDelta([&NewEntries](const FPCGDeltaKey& DeltaKey, const TInstancedStruct<FPCGDeltaBase>& DeltaValue) -> bool
+				{
+					if (const FPCGPointEditDelta* EditDelta = DeltaValue.GetPtr<FPCGPointEditDelta>())
+					{
+						TSharedPtr<FPCGPointEditDeltaEntry> Entry = MakeShared<FPCGPointEditDeltaEntry>();
+						Entry->Key = DeltaKey;
+						Entry->OriginalElementIndex = EditDelta->ElementIndex;
+						Entry->Location = EditDelta->TransformOverride.GetLocation();
+						NewEntries.Add(Entry);
+					}
+
+					return true;
+				});
+			}
+		}
+	}
+
+	bool bListChanged = NewEntries.Num() != EditEntries.Num();
+	if (!bListChanged)
+	{
+		for (int32 i = 0; i < NewEntries.Num(); ++i)
+		{
+			if (NewEntries[i]->Key != EditEntries[i]->Key
+				|| NewEntries[i]->OriginalElementIndex != EditEntries[i]->OriginalElementIndex
+				|| !NewEntries[i]->Location.Equals(EditEntries[i]->Location))
+			{
+				bListChanged = true;
+				break;
+			}
+		}
+	}
+
+	// @todo_pcg: Separate data diffing from UI updates. RefreshLists should return whether the list changed,
+	// and the caller should handle widget visibility and RequestListRefresh.
+	if (bListChanged)
+	{
+		EditEntries = MoveTemp(NewEntries);
+
+		if (EditListView.IsValid())
+		{
+			EditListView->RequestListRefresh();
+		}
+	}
+
+	const bool bHasEdits = EditEntries.Num() > 0;
+	if (EditListSection.IsValid())
+	{
+		EditListSection->SetVisibility(bHasEdits ? EVisibility::SelfHitTestInvisible : EVisibility::Collapsed);
+	}
+}
+
+FLinearColor FPCGPointEditViewportExtension::GetDisplayColor(const bool bIsSelected) const
+{
+	return bIsSelected ? PCGPointEditDelta::Constants::SelectedColor : PCGPointEditDelta::Constants::NormalColor;
+}
+
+FTransform FPCGPointEditViewportExtension::GetDisplayTransform(const FTransform& SourceElementTransform, const FConstStructView DeltaStruct) const
+{
+	if (const FPCGPointEditDelta* EditDelta = DeltaStruct.GetPtr<const FPCGPointEditDelta>())
+	{
+		return EditDelta->TransformOverride;
+	}
+
+	return SourceElementTransform;
+}
+
+bool FPCGPointEditViewportExtension::MatchesSourceElement(const FTransform& SourceElementTransform, const FConstStructView DeltaStruct, const double SpatialTolerance) const
+{
+	const FPCGPointEditDelta* EditDelta = DeltaStruct.GetPtr<const FPCGPointEditDelta>();
+	if (!EditDelta)
+	{
+		return false;
+	}
+
+	if (!CachedContext.NodeConfiguration)
+	{
+		return false;
+	}
+
+	const FPCGDeltaKey SourceKey = PCG::DataOverride::Keys::FPCGCompositeTransformDeltaKey(
+		SourceElementTransform,
+		SpatialTolerance,
+		CachedContext.NodeConfiguration->bSignaturePosition,
+		CachedContext.NodeConfiguration->bSignatureRotation,
+		CachedContext.NodeConfiguration->bSignatureScale);
+
+	const FPCGDeltaKey DeltaKey = PCG::DataOverride::Keys::FPCGCompositeTransformDeltaKey(
+		EditDelta->OriginalTransform,
+		SpatialTolerance,
+		CachedContext.NodeConfiguration->bSignaturePosition,
+		CachedContext.NodeConfiguration->bSignatureRotation,
+		CachedContext.NodeConfiguration->bSignatureScale);
+
+	return SourceKey == DeltaKey;
+}
+
+bool FPCGPointEditViewportExtension::ApplyGizmoTransform(const FTransform& NewTransform, TInstancedStruct<FPCGDeltaBase>& DeltaStruct, int32 ElementIndex) const
+{
+	if (FPCGPointEditDelta* EditDelta = DeltaStruct.GetMutablePtr<FPCGPointEditDelta>())
+	{
+		EditDelta->TransformOverride = NewTransform;
+
+		return true;
+	}
+
+	return false;
+}
+
+void FPCGPointEditViewportExtension::CreateNewDelta(
+	const FPCGDeltaKey& DeltaKey,
+	FPCGDeltaCollection& Collection,
+	const FTransform& OriginalTransform,
+	const FTransform& NewTransform,
+	const FPCGDeltaCreateContext& Context) const
+{
+	FPCGPointEditDelta NewDelta;
+	NewDelta.OriginalTransform = OriginalTransform;
+	NewDelta.TransformOverride = NewTransform;
+	NewDelta.Bounds = Context.ComponentBounds;
+	NewDelta.bOverridePosition = bOverridePosition;
+	NewDelta.bOverrideRotation = bOverrideRotation;
+	NewDelta.bOverrideScale = bOverrideScale;
+
+	NewDelta.ElementIndex = Context.OriginalElementIndex;
+
+	Collection.Add(DeltaKey, TInstancedStruct<FPCGPointEditDelta>::Make(MoveTemp(NewDelta)));
+}
+
+FTransform FPCGPointEditViewportExtension::GetOriginalTransform(const FConstStructView DeltaStruct) const
+{
+	if (const FPCGPointEditDelta* EditDelta = DeltaStruct.GetPtr<const FPCGPointEditDelta>())
+	{
+		return EditDelta->OriginalTransform;
+	}
+
+	return FTransform::Identity;
+}
+
+void FPCGPointEditViewportExtension::HandleDeletePressed()
+{
+	UPCGComponent* PCGComponent = CachedContext.ActivePCGComponent.Get();
+	if (!PCGComponent)
+	{
+		return;
+	}
+
+	FPCGSourceDataContainer* DataContainer = PCGComponent->GetExecutionState().GetSourceDataContainer();
+	if (!DataContainer)
+	{
+		return;
+	}
+
+	FSharedStruct SharedStruct = DataContainer->GetMutable<FPCGDeltaCollection>(CachedContext.ActiveStorageKey);
+	FPCGDeltaCollection* Collection = SharedStruct.IsValid() ? SharedStruct.GetPtr<FPCGDeltaCollection>() : nullptr;
+	if (!Collection)
+	{
+		return;
+	}
+
+	if (CachedContext.bSelectionActive && CachedContext.SelectionDeltaKey.Hash != 0)
+	{
+		if (!CachedContext.NodeConfiguration)
+		{
+			return;
+		}
+
+		FScopedTransaction Transaction(LOCTEXT("ConvertEditToDeletion", "PCG Manual Edit: Convert Edit to Deletion"));
+		PCGComponent->Modify();
+
+		TInstancedStruct<FPCGDeltaBase>* ExistingDelta = Collection->Find(CachedContext.SelectionDeltaKey);
+
+		FTransform OriginalTransform = CachedContext.SelectionTransform;
+		if (ExistingDelta)
+		{
+			const FConstStructView DeltaView(ExistingDelta->GetScriptStruct(), ExistingDelta->GetMemory());
+			OriginalTransform = GetOriginalTransform(DeltaView);
+			Collection->Remove(CachedContext.SelectionDeltaKey);
+		}
+
+		FPCGPointDeletionDelta DeletionDelta;
+		DeletionDelta.OriginalTransform = OriginalTransform;
+		DeletionDelta.Bounds = PCGComponent->GetTotalBounds();
+		DeletionDelta.ElementIndex = CachedContext.OriginalElementIndex;
+
+		const FPCGDeltaKey DeletionKey = PCG::DataOverride::Keys::FPCGCompositeTransformDeltaKey(
+			OriginalTransform,
+			CachedContext.NodeConfiguration->SignatureTolerance,
+			CachedContext.NodeConfiguration->bSignaturePosition,
+			CachedContext.NodeConfiguration->bSignatureRotation,
+			CachedContext.NodeConfiguration->bSignatureScale,
+			FPCGPointDeletionDelta::GetDeltaNameStatic());
+
+		Collection->Add(DeletionKey, TInstancedStruct<FPCGPointDeletionDelta>::Make(MoveTemp(DeletionDelta)));
+
+		PCG::DeltaViewportExtension::MarkComponentDirty(PCGComponent, DataContainer);
+
+		if (CachedCallbacks.ClearSelection)
+		{
+			CachedCallbacks.ClearSelection();
+		}
+	}
+
+	RefreshLists(CachedContext);
+}
+
+TSharedRef<ITableRow> FPCGPointEditViewportExtension::OnGenerateEditRow(TSharedPtr<FPCGPointEditDeltaEntry> Entry, const TSharedRef<STableViewBase>& OwnerTable)
+{
+	TWeakPtr<FPCGPointEditDeltaEntry> WeakEntry = Entry;
+	FPCGDeltaKey EntryKey = Entry->Key;
+
+	return SNew(STableRow<TSharedPtr<FPCGPointEditDeltaEntry>>, OwnerTable)
+		[
+			SNew(SHorizontalBox)
+			+ SHorizontalBox::Slot()
+			.AutoWidth()
+			.VAlign(VAlign_Center)
+			.Padding(4.0f, 1.0f, 2.0f, 1.0f)
+			[
+				SNew(STextBlock)
+				.Text(FText::AsNumber(Entry->OriginalElementIndex))
+				.Font(FCoreStyle::GetDefaultFontStyle("Regular", 8))
+			]
+			+ SHorizontalBox::Slot()
+			.FillWidth(1.0f)
+			.VAlign(VAlign_Center)
+			.Padding(2.0f, 1.0f)
+			[
+				SNew(SVectorSlider<double>, /*bIsRotator=*/false, /*InProperty=*/nullptr)
+				.VisibleText_0_Lambda([WeakEntry]() -> FString
+				{
+					TSharedPtr<FPCGPointEditDeltaEntry> Pinned = WeakEntry.Pin();
+					return Pinned ? FString::SanitizeFloat(Pinned->Location.X) : FString();
+				})
+				.VisibleText_1_Lambda([WeakEntry]() -> FString
+				{
+					TSharedPtr<FPCGPointEditDeltaEntry> Pinned = WeakEntry.Pin();
+					return Pinned ? FString::SanitizeFloat(Pinned->Location.Y) : FString();
+				})
+				.VisibleText_2_Lambda([WeakEntry]() -> FString
+				{
+					TSharedPtr<FPCGPointEditDeltaEntry> Pinned = WeakEntry.Pin();
+					return Pinned ? FString::SanitizeFloat(Pinned->Location.Z) : FString();
+				})
+				.OnNumericCommitted_Box_0(SVectorSlider<double>::FOnNumericValueCommitted::CreateLambda(
+					[this, EntryKey](const double V, ETextCommit::Type) { OnEditComponentCommitted(EntryKey, 0, V); }))
+				.OnNumericCommitted_Box_1(SVectorSlider<double>::FOnNumericValueCommitted::CreateLambda(
+					[this, EntryKey](const double V, ETextCommit::Type) { OnEditComponentCommitted(EntryKey, 1, V); }))
+				.OnNumericCommitted_Box_2(SVectorSlider<double>::FOnNumericValueCommitted::CreateLambda(
+					[this, EntryKey](const double V, ETextCommit::Type) { OnEditComponentCommitted(EntryKey, 2, V); }))
+			]
+			+ SHorizontalBox::Slot()
+			.AutoWidth()
+			.VAlign(VAlign_Center)
+			.Padding(2.0f, 0.0f)
+			[
+				SNew(SButton)
+				.ButtonStyle(FAppStyle::Get(), "SimpleButton")
+				.ContentPadding(FMargin(1, 0))
+				.ToolTipText(LOCTEXT("SelectEditTooltip", "Select this edited point"))
+				.OnClicked_Lambda([this, EntryKey]() -> FReply
+				{
+					if (CachedCallbacks.SelectElement)
+					{
+						CachedCallbacks.SelectElement(EntryKey, INDEX_NONE);
+					}
+					return FReply::Handled();
+				})
+				[
+					SNew(SImage)
+					.Image(FAppStyle::Get().GetBrush("Icons.Search"))
+					.ColorAndOpacity(FSlateColor::UseForeground())
+				]
+			]
+			+ SHorizontalBox::Slot()
+			.AutoWidth()
+			.VAlign(VAlign_Center)
+			.Padding(2.0f, 0.0f)
+			[
+				SNew(SButton)
+				.ButtonStyle(FAppStyle::Get(), "SimpleButton")
+				.ContentPadding(FMargin(1, 0))
+				.ToolTipText(LOCTEXT("RemoveEditTooltip", "Remove this edit override"))
+				.OnClicked_Lambda([this, EntryKey]() -> FReply
+				{
+					RemoveEditEntry(EntryKey);
+					RefreshLists(CachedContext);
+					return FReply::Handled();
+				})
+				[
+					SNew(SImage)
+					.Image(FAppStyle::Get().GetBrush("Icons.Delete"))
+					.ColorAndOpacity(FSlateColor::UseForeground())
+				]
+			]
+		];
+}
+
+void FPCGPointEditViewportExtension::OnEditComponentCommitted(const FPCGDeltaKey& EntryKey, const int32 AxisIndex, const double NewValue)
+{
+	for (const TSharedPtr<FPCGPointEditDeltaEntry>& Entry : EditEntries)
+	{
+		if (Entry && Entry->Key == EntryKey)
+		{
+			Entry->Location[AxisIndex] = NewValue;
+			UpdateEditDeltaLocation(EntryKey, Entry->Location);
+			return;
+		}
+	}
+}
+
+void FPCGPointEditViewportExtension::UpdateEditDeltaLocation(const FPCGDeltaKey& DeltaKey, const FVector& NewLocation)
+{
+	UPCGComponent* PCGComponent = CachedContext.ActivePCGComponent.Get();
+	if (!PCGComponent)
+	{
+		return;
+	}
+
+	FPCGSourceDataContainer* DataContainer = PCGComponent->GetExecutionState().GetSourceDataContainer();
+	if (!DataContainer)
+	{
+		return;
+	}
+
+	FSharedStruct SharedStruct = DataContainer->GetMutable<FPCGDeltaCollection>(CachedContext.ActiveStorageKey);
+	FPCGDeltaCollection* Collection = SharedStruct.IsValid() ? SharedStruct.GetPtr<FPCGDeltaCollection>() : nullptr;
+	if (!Collection)
+	{
+		return;
+	}
+
+	TInstancedStruct<FPCGDeltaBase>* DeltaStruct = Collection->Find(DeltaKey);
+	if (!DeltaStruct)
+	{
+		return;
+	}
+
+	FScopedTransaction Transaction(LOCTEXT("UpdateEditDeltaLocation", "PCG Manual Edit: Update Edit Location"));
+	PCGComponent->Modify();
+
+	if (FPCGPointEditDelta* EditDelta = DeltaStruct->GetMutablePtr<FPCGPointEditDelta>())
+	{
+		EditDelta->TransformOverride.SetLocation(NewLocation);
+	}
+
+	PCG::DeltaViewportExtension::MarkComponentDirty(PCGComponent, DataContainer);
+
+	if (CachedCallbacks.SelectElement)
+	{
+		CachedCallbacks.SelectElement(DeltaKey, INDEX_NONE);
+	}
+}
+
+void FPCGPointEditViewportExtension::RemoveEditEntry(const FPCGDeltaKey& EntryKey)
+{
+	UPCGComponent* PCGComponent = CachedContext.ActivePCGComponent.Get();
+	if (!PCGComponent)
+	{
+		return;
+	}
+
+	FPCGSourceDataContainer* DataContainer = PCGComponent->GetExecutionState().GetSourceDataContainer();
+	if (!DataContainer)
+	{
+		return;
+	}
+
+	FSharedStruct SharedStruct = DataContainer->GetMutable<FPCGDeltaCollection>(CachedContext.ActiveStorageKey);
+	FPCGDeltaCollection* Collection = SharedStruct.IsValid() ? SharedStruct.GetPtr<FPCGDeltaCollection>() : nullptr;
+	if (!Collection)
+	{
+		return;
+	}
+
+	FScopedTransaction Transaction(LOCTEXT("RemoveEditDelta", "PCG Manual Edit: Remove Edit"));
+	PCGComponent->Modify();
+
+	Collection->Remove(EntryKey);
+	PCG::DeltaViewportExtension::MarkComponentDirty(PCGComponent, DataContainer);
+}
+
+void FPCGPointEditViewportExtension::RemoveAllEdits()
+{
+	UPCGComponent* PCGComponent = CachedContext.ActivePCGComponent.Get();
+	if (!PCGComponent)
+	{
+		return;
+	}
+
+	FPCGSourceDataContainer* DataContainer = PCGComponent->GetExecutionState().GetSourceDataContainer();
+	if (!DataContainer)
+	{
+		return;
+	}
+
+	FSharedStruct SharedStruct = DataContainer->GetMutable<FPCGDeltaCollection>(CachedContext.ActiveStorageKey);
+	FPCGDeltaCollection* Collection = SharedStruct.IsValid() ? SharedStruct.GetPtr<FPCGDeltaCollection>() : nullptr;
+	if (!Collection)
+	{
+		return;
+	}
+
+	FScopedTransaction Transaction(LOCTEXT("RemoveAllEditDeltas", "PCG Manual Edit: Remove All Edits"));
+	PCGComponent->Modify();
+
+	TArray<FPCGDeltaKey> EditKeys;
+	Collection->ForEachDelta([&EditKeys](const FPCGDeltaKey& DeltaKey, const TInstancedStruct<FPCGDeltaBase>& DeltaValue) -> bool
+	{
+		if (DeltaValue.GetPtr<FPCGPointEditDelta>())
+		{
+			EditKeys.Add(DeltaKey);
+		}
+		return true;
+	});
+
+	for (const FPCGDeltaKey& Key : EditKeys)
+	{
+		Collection->Remove(Key);
+	}
+
+	PCG::DeltaViewportExtension::MarkComponentDirty(PCGComponent, DataContainer);
+}
+
+void FPCGPointEditViewportExtension::UpdateDeltaOverrideFlags()
+{
+	UPCGComponent* PCGComponent = CachedContext.ActivePCGComponent.Get();
+	if (!PCGComponent)
+	{
+		return;
+	}
+
+	FPCGSourceDataContainer* DataContainer = PCGComponent->GetExecutionState().GetSourceDataContainer();
+	if (!DataContainer)
+	{
+		return;
+	}
+
+	FSharedStruct SharedStruct = DataContainer->GetMutable<FPCGDeltaCollection>(CachedContext.ActiveStorageKey);
+	FPCGDeltaCollection* Collection = SharedStruct.IsValid() ? SharedStruct.GetPtr<FPCGDeltaCollection>() : nullptr;
+	if (!Collection)
+	{
+		return;
+	}
+
+	FScopedTransaction Transaction(LOCTEXT("UpdateEditDeltaOverrideFlags", "PCG Manual Edit: Update Override Flags"));
+	PCGComponent->Modify();
+
+	const bool bPosition = bOverridePosition;
+	const bool bRotation = bOverrideRotation;
+	const bool bScale = bOverrideScale;
+
+	bool bModified = false;
+	Collection->ForEachDelta([&bModified, bPosition, bRotation, bScale](const FPCGDeltaKey& DeltaKey, TInstancedStruct<FPCGDeltaBase>& DeltaValue) -> bool
+	{
+		if (FPCGPointEditDelta* EditDelta = DeltaValue.GetMutablePtr<FPCGPointEditDelta>())
+		{
+			if (EditDelta->bOverridePosition != bPosition || EditDelta->bOverrideRotation != bRotation || EditDelta->bOverrideScale != bScale)
+			{
+				EditDelta->bOverridePosition = bPosition;
+				EditDelta->bOverrideRotation = bRotation;
+				EditDelta->bOverrideScale = bScale;
+				bModified = true;
+			}
+		}
+		return true;
+	});
+
+	if (bModified)
+	{
+		PCG::DeltaViewportExtension::MarkComponentDirty(PCGComponent, DataContainer);
+	}
+}
+
+#undef LOCTEXT_NAMESPACE
diff --git a/Engine/Plugins/PCG/Source/PCGEditor/Private/DeltaViewportExtensions/PCGPointEditViewportExtension.h b/Engine/Plugins/PCG/Source/PCGEditor/Private/DeltaViewportExtensions/PCGPointEditViewportExtension.h
new file mode 100644
--- /dev/null
+++ b/Engine/Plugins/PCG/Source/PCGEditor/Private/DeltaViewportExtensions/PCGPointEditViewportExtension.h
@@ -0,0 +1,62 @@
+// Copyright Epic Games, Inc. All Rights Reserved.
+
+#pragma once
+
+#include "DeltaViewportExtensions/PCGDeltaViewportExtension.h"
+
+#include "Graph/DataOverride/PCGDataOverridePoints.h"
+
+class ITableRow;
+class STableViewBase;
+class SVerticalBox;
+template <typename ItemType> class SListView;
+
+struct FPCGPointEditDeltaEntry
+{
+	FPCGDeltaKey Key;
+	int32 OriginalElementIndex = INDEX_NONE;
+	FVector Location = FVector::ZeroVector;
+};
+
+/** Viewport extension for FPCGPointEditDelta: edits list + viewport drawing/gizmo policy. */
+class FPCGPointEditViewportExtension : public IPCGDeltaViewportExtension
+{
+public:
+	// ~Begin IPCGDeltaViewportExtension interface
+	virtual FText GetDisplayName() const override { return NSLOCTEXT("PCGPointEditViewportExtension", "PointEditDisplayName", "Point Edit"); }
+	virtual FName GetDeltaName() const override { return FPCGPointEditDelta::GetDeltaNameStatic(); }
+	virtual int32 GetSortPriority() const override { return 0; }
+	virtual TSharedRef<SWidget> CreateWidget(FPCGDeltaViewportCallbacks Callbacks) override;
+	virtual void UpdateContext(const FPCGDeltaViewportContext& Context) override { CachedContext = Context; }
+	virtual void RefreshLists(const FPCGDeltaViewportContext& Context) override;
+	virtual TArray<FPCGDeltaKeyBinding> GetKeyBindings() override { return {{EKeys::Delete, [this]() { HandleDeletePressed(); }}}; }
+	virtual FLinearColor GetDisplayColor(const bool bIsSelected) const override;
+	virtual FTransform GetDisplayTransform(const FTransform& SourceElementTransform, const FConstStructView DeltaStruct) const override;
+	virtual bool MatchesSourceElement(const FTransform& SourceElementTransform, const FConstStructView DeltaStruct, const double SpatialTolerance) const override;
+	virtual bool ApplyGizmoTransform(const FTransform& NewTransform, TInstancedStruct<FPCGDeltaBase>& DeltaStruct, int32 ElementIndex) const override;
+	virtual void CreateNewDelta(const FPCGDeltaKey& DeltaKey, FPCGDeltaCollection& Collection,
+		const FTransform& OriginalTransform, const FTransform& NewTransform, const FPCGDeltaCreateContext& Context) const override;
+	virtual FTransform GetOriginalTransform(const FConstStructView DeltaStruct) const override;
+	// ~End IPCGDeltaViewportExtension interface
+
+private:
+	void HandleDeletePressed();
+	TSharedRef<ITableRow> OnGenerateEditRow(TSharedPtr<FPCGPointEditDeltaEntry> Entry, const TSharedRef<STableViewBase>& OwnerTable);
+	void OnEditComponentCommitted(const FPCGDeltaKey& EntryKey, const int32 AxisIndex, const double NewValue);
+	void UpdateEditDeltaLocation(const FPCGDeltaKey& DeltaKey, const FVector& NewLocation);
+	void RemoveEditEntry(const FPCGDeltaKey& EntryKey);
+	void RemoveAllEdits();
+	void UpdateDeltaOverrideFlags();
+
+	/** Which transform components the edit delta writes back at apply time. */
+	bool bOverridePosition = true;
+	bool bOverrideRotation = true;
+	bool bOverrideScale = true;
+
+	FPCGDeltaViewportContext CachedContext;
+	FPCGDeltaViewportCallbacks CachedCallbacks;
+
+	TArray<TSharedPtr<FPCGPointEditDeltaEntry>> EditEntries;
+	TSharedPtr<SVerticalBox> EditListSection;
+	TSharedPtr<SListView<TSharedPtr<FPCGPointEditDeltaEntry>>> EditListView;
+};
diff --git a/Engine/Plugins/PCG/Source/PCGEditor/Private/DeltaViewportExtensions/PCGPointInsertionViewportExtension.cpp b/Engine/Plugins/PCG/Source/PCGEditor/Private/DeltaViewportExtensions/PCGPointInsertionViewportExtension.cpp
new file mode 100644
--- /dev/null
+++ b/Engine/Plugins/PCG/Source/PCGEditor/Private/DeltaViewportExtensions/PCGPointInsertionViewportExtension.cpp
@@ -0,0 +1,578 @@
+// Copyright Epic Games, Inc. All Rights Reserved.
+
+#include "DeltaViewportExtensions/PCGPointInsertionViewportExtension.h"
+#include "DeltaViewportExtensions/PCGDeltaViewportExtensionHelpers.h"
+#include "PCGComponent.h"
+#include "Graph/DataOverride/PCGDataOverrideHelpers.h"
+
+#include "ScopedTransaction.h"
+#include "Styling/AppStyle.h"
+#include "Styling/CoreStyle.h"
+#include "Widgets/SBoxPanel.h"
+#include "Widgets/Images/SImage.h"
+#include "Widgets/Input/SButton.h"
+#include "Widgets/Text/STextBlock.h"
+#include "Widgets/Views/SListView.h"
+
+#include "KismetPins/SVectorSlider.h"
+
+#define LOCTEXT_NAMESPACE "PCGPointInsertionViewportExtension"
+
+namespace PCGPointInsertionDelta::Constants
+{
+	constexpr FLinearColor SelectedColor(1.0f, 0.8f, 0.3f, 0.5f);
+	constexpr FLinearColor NormalColor(0.75f, 0.45f, 0.0f, 0.3f);
+}
+
+TSharedRef<SWidget> FPCGPointInsertionViewportExtension::CreateWidget(FPCGDeltaViewportCallbacks Callbacks)
+{
+	CachedCallbacks = MoveTemp(Callbacks);
+	InsertionEntries.Empty();
+
+	TSharedRef<SVerticalBox> BaseBox = SNew(SVerticalBox)
+		// Insert Point button
+		+ SVerticalBox::Slot()
+	   .AutoHeight()
+	   .Padding(0.0f, 4.0f, 0.0f, 0.0f)
+		[
+			SNew(SButton)
+		   .HAlign(HAlign_Center)
+		   .OnClicked(FOnClicked::CreateRaw(this, &FPCGPointInsertionViewportExtension::OnInsertPointClicked))
+			[
+				SNew(STextBlock)
+			   .Text(LOCTEXT("InsertPointButton", "Insert Point"))
+			   .Font(FCoreStyle::GetDefaultFontStyle("Regular", 8))
+			]
+		]
+		// Insertions list section
+		+ SVerticalBox::Slot()
+	   .AutoHeight()
+	   .Padding(0.0f, 8.0f, 0.0f, 2.0f)
+		[
+			SAssignNew(InsertionListSection, SVerticalBox)
+		   .Visibility(EVisibility::Collapsed)
+			// Header row: "Insertions" + "Remove All"
+			+ SVerticalBox::Slot()
+		   .AutoHeight()
+		   .Padding(0.0f, 0.0f, 0.0f, 4.0f)
+			[
+				SNew(SHorizontalBox)
+				+ SHorizontalBox::Slot()
+			   .FillWidth(1.0f)
+			   .VAlign(VAlign_Center)
+				[
+					SNew(STextBlock)
+				   .Text(LOCTEXT("InsertionsHeader", "Insertions"))
+				   .Font(FCoreStyle::GetDefaultFontStyle("Bold", 8))
+				]
+				+ SHorizontalBox::Slot()
+			   .AutoWidth()
+			   .VAlign(VAlign_Center)
+				[
+					SNew(SButton)
+				   .OnClicked(FOnClicked::CreateRaw(this, &FPCGPointInsertionViewportExtension::OnRemoveAllClicked))
+					[
+						SNew(STextBlock)
+					   .Text(LOCTEXT("RemoveAllInsertionsButton", "Remove All"))
+					   .Font(FCoreStyle::GetDefaultFontStyle("Regular", 8))
+					]
+				]
+			]
+			// Insertion list view
+			+ SVerticalBox::Slot()
+		   .AutoHeight()
+		   .MaxHeight(150.0f)
+			[
+				SAssignNew(InsertionListView, SListView<TSharedPtr<FPCGPointInsertionDeltaEntry>>)
+			   .ListItemsSource(&InsertionEntries)
+			   .OnGenerateRow_Lambda([this](TSharedPtr<FPCGPointInsertionDeltaEntry> Entry, const TSharedRef<STableViewBase>& OwnerTable) -> TSharedRef<ITableRow>
+				{
+					return OnGenerateInsertionRow(Entry, OwnerTable);
+				})
+			   .SelectionMode(ESelectionMode::None)
+			]
+		];
+
+	return BaseBox;
+}
+
+void FPCGPointInsertionViewportExtension::RefreshLists(const FPCGDeltaViewportContext& Context)
+{
+	CachedContext = Context;
+
+	// Build new entries into a temporary array and compare against existing to avoid unnecessary list rebuilds
+	// that destroy widget focus (SVectorSlider text fields lose focus on RequestListRefresh).
+	TArray<TSharedPtr<FPCGPointInsertionDeltaEntry>> NewEntries;
+
+	if (UPCGComponent* PCGComponent = Context.ActivePCGComponent.Get())
+	{
+		if (FPCGSourceDataContainer* DataContainer = PCGComponent->GetExecutionState().GetSourceDataContainer())
+		{
+			FConstSharedStruct SharedStruct = DataContainer->Get<FPCGDeltaCollection>(Context.ActiveStorageKey);
+			if (const FPCGDeltaCollection* Collection = SharedStruct.IsValid() ? SharedStruct.GetPtr<const FPCGDeltaCollection>() : nullptr)
+			{
+				Collection->ForEachDelta([&NewEntries](const FPCGDeltaKey& DeltaKey, const TInstancedStruct<FPCGDeltaBase>& DeltaValue) -> bool
+				{
+					if (const FPCGPointInsertionDelta* InsertDelta = DeltaValue.GetPtr<FPCGPointInsertionDelta>())
+					{
+						for (int32 i = 0; i < InsertDelta->InsertedPoints.Num(); ++i)
+						{
+							TSharedPtr<FPCGPointInsertionDeltaEntry> Entry = MakeShared<FPCGPointInsertionDeltaEntry>();
+							Entry->DeltaKey = DeltaKey;
+							Entry->Index = i;
+							Entry->Location = InsertDelta->InsertedPoints[i].Transform.GetLocation();
+							NewEntries.Add(Entry);
+						}
+					}
+					return true;
+				});
+			}
+		}
+	}
+
+	// Only rebuild the list if the entry count or keys changed. Location updates are reflected via VisibleText lambdas.
+	bool bListChanged = NewEntries.Num() != InsertionEntries.Num();
+	if (!bListChanged)
+	{
+		for (int32 i = 0; i < NewEntries.Num(); ++i)
+		{
+			if (NewEntries[i]->DeltaKey != InsertionEntries[i]->DeltaKey || NewEntries[i]->Index != InsertionEntries[i]->Index)
+			{
+				bListChanged = true;
+				break;
+			}
+
+			// Update cached location in-place so VisibleText lambdas reflect the latest values.
+			InsertionEntries[i]->Location = NewEntries[i]->Location;
+		}
+	}
+
+	if (bListChanged)
+	{
+		InsertionEntries = MoveTemp(NewEntries);
+
+		if (InsertionListView.IsValid())
+		{
+			InsertionListView->RequestListRefresh();
+		}
+	}
+
+	const bool bHasInsertions = InsertionEntries.Num() > 0;
+	if (InsertionListSection.IsValid())
+	{
+		InsertionListSection->SetVisibility(bHasInsertions ? EVisibility::SelfHitTestInvisible : EVisibility::Collapsed);
+	}
+}
+
+FLinearColor FPCGPointInsertionViewportExtension::GetDisplayColor(const bool bIsSelected) const
+{
+	return bIsSelected ? PCGPointInsertionDelta::Constants::SelectedColor : PCGPointInsertionDelta::Constants::NormalColor;
+}
+
+void FPCGPointInsertionViewportExtension::GetInsertedElements(const FPCGDeltaCollection& Collection, TArray<FInsertedViewportElement>& OutElements) const
+{
+	Collection.ForEachDelta([&OutElements](const FPCGDeltaKey& DeltaKey, const TInstancedStruct<FPCGDeltaBase>& DeltaValue) -> bool
+	{
+		if (const FPCGPointInsertionDelta* InsertDelta = DeltaValue.GetPtr<FPCGPointInsertionDelta>())
+		{
+			for (int32 Index = 0; Index < InsertDelta->InsertedPoints.Num(); ++Index)
+			{
+				OutElements.Add({DeltaKey, Index, InsertDelta->InsertedPoints[Index].Transform});
+			}
+		}
+		return true;
+	});
+}
+
+int32 FPCGPointInsertionViewportExtension::OnAltDrag(FPCGDeltaCollection& Collection, const FTransform& SourceTransform, FPCGDeltaKey& OutDeltaKey) const
+{
+	FPCGPointInsertionDelta* InsertionDelta = FindOrCreateInsertionDelta(&Collection, &OutDeltaKey);
+	if (!InsertionDelta)
+	{
+		return INDEX_NONE;
+	}
+
+	FPCGPoint NewPoint;
+	NewPoint.Transform = SourceTransform;
+	InsertionDelta->InsertedPoints.Add(NewPoint);
+
+	return InsertionDelta->InsertedPoints.Num() - 1;
+}
+
+bool FPCGPointInsertionViewportExtension::ApplyGizmoTransform(const FTransform& NewTransform, TInstancedStruct<FPCGDeltaBase>& DeltaStruct, const int32 ElementIndex) const
+{
+	FPCGPointInsertionDelta* InsertDelta = DeltaStruct.GetMutablePtr<FPCGPointInsertionDelta>();
+	if (!InsertDelta)
+	{
+		return false;
+	}
+
+	if (InsertDelta->InsertedPoints.IsValidIndex(ElementIndex))
+	{
+		InsertDelta->InsertedPoints[ElementIndex].Transform = NewTransform;
+		return true;
+	}
+
+	return false;
+}
+
+FPCGDeltaCollection* FPCGPointInsertionViewportExtension::GetOrCreateMutableCollection(FPCGSourceDataContainer* DataContainer)
+{
+	FSharedStruct SharedStruct = DataContainer->GetMutable<FPCGDeltaCollection>(CachedContext.ActiveStorageKey);
+	FPCGDeltaCollection* Collection = SharedStruct.IsValid() ? SharedStruct.GetPtr<FPCGDeltaCollection>() : nullptr;
+
+	if (!Collection)
+	{
+		FPCGDeltaCollection NewCollection;
+		NewCollection.Settings.KeyPolicy = EPCGDataOverrideKeyPolicy::Spatial;
+		NewCollection.Settings.KeyTarget = EPCGSpatialKeyTarget::Position;
+		NewCollection.Settings.SpatialTolerance = PCG::DataOverride::Constants::SpatialToleranceDefault;
+		DataContainer->Store(CachedContext.ActiveStorageKey, NewCollection);
+
+		SharedStruct = DataContainer->GetMutable<FPCGDeltaCollection>(CachedContext.ActiveStorageKey);
+		Collection = SharedStruct.IsValid() ? SharedStruct.GetPtr<FPCGDeltaCollection>() : nullptr;
+	}
+
+	return Collection;
+}
+
+FPCGPointInsertionDelta* FPCGPointInsertionViewportExtension::FindOrCreateInsertionDelta(FPCGDeltaCollection* Collection, FPCGDeltaKey* OutKey)
+{
+	FPCGPointInsertionDelta* InsertionDelta = nullptr;
+
+	Collection->ForEachDelta([&InsertionDelta, OutKey](const FPCGDeltaKey& IterKey, TInstancedStruct<FPCGDeltaBase>& DeltaValue) -> bool
+	{
+		InsertionDelta = DeltaValue.GetMutablePtr<FPCGPointInsertionDelta>();
+		if (InsertionDelta && OutKey)
+		{
+			*OutKey = IterKey;
+		}
+
+		return !InsertionDelta;
+	});
+
+	if (!InsertionDelta)
+	{
+		const FPCGDeltaKey NewKey(PCGInsertionConstants::DefaultInsertionKeyHash, FPCGPointInsertionDelta::GetDeltaNameStatic());
+		TInstancedStruct<FPCGDeltaBase>& AddedRef = Collection->Add_GetRef(NewKey, TInstancedStruct<FPCGPointInsertionDelta>::Make(FPCGPointInsertionDelta()));
+		InsertionDelta = AddedRef.GetMutablePtr<FPCGPointInsertionDelta>();
+		if (OutKey)
+		{
+			*OutKey = NewKey;
+		}
+	}
+
+	return InsertionDelta;
+}
+
+FReply FPCGPointInsertionViewportExtension::OnInsertPointClicked()
+{
+	AddInsertedElement();
+	RefreshLists(CachedContext);
+
+	return FReply::Handled();
+}
+
+FReply FPCGPointInsertionViewportExtension::OnRemoveAllClicked()
+{
+	RemoveAllInsertedElements();
+	RefreshLists(CachedContext);
+
+	return FReply::Handled();
+}
+
+void FPCGPointInsertionViewportExtension::HandleDeletePressed()
+{
+	if (CachedContext.bSelectionActive && CachedContext.SelectedElementIndex != INDEX_NONE)
+	{
+		RemoveInsertedElement(CachedContext.SelectedElementIndex);
+
+		if (CachedCallbacks.ClearSelection)
+		{
+			CachedCallbacks.ClearSelection();
+		}
+	}
+
+	RefreshLists(CachedContext);
+}
+
+void FPCGPointInsertionViewportExtension::AddInsertedElement()
+{
+	UPCGComponent* PCGComponent = CachedContext.ActivePCGComponent.Get();
+	if (!PCGComponent)
+	{
+		return;
+	}
+
+	FPCGSourceDataContainer* DataContainer = PCGComponent->GetExecutionState().GetSourceDataContainer();
+	if (!DataContainer)
+	{
+		return;
+	}
+
+	FScopedTransaction Transaction(LOCTEXT("AddInsertedElement", "PCG Manual Edit: Add Inserted Point"));
+	PCGComponent->Modify();
+
+	FPCGDeltaCollection* Collection = GetOrCreateMutableCollection(DataContainer);
+	if (!Collection)
+	{
+		return;
+	}
+
+	FPCGDeltaKey InsertionDeltaKey;
+	FPCGPointInsertionDelta* InsertionDelta = FindOrCreateInsertionDelta(Collection, &InsertionDeltaKey);
+	if (!InsertionDelta)
+	{
+		return;
+	}
+
+	FTransform NewTransform = FTransform::Identity;
+	if (CachedContext.bSelectionActive)
+	{
+		NewTransform = CachedContext.SelectionTransform;
+	}
+	else if (AActor* Owner = PCGComponent->GetOwner())
+	{
+		NewTransform = Owner->GetActorTransform();
+	}
+	FPCGPoint NewPoint;
+	NewPoint.Transform = NewTransform;
+	InsertionDelta->InsertedPoints.Add(NewPoint);
+
+	const int32 NewElementIndex = InsertionDelta->InsertedPoints.Num() - 1;
+
+	PCG::DeltaViewportExtension::MarkComponentDirty(PCGComponent, DataContainer);
+
+	if (CachedCallbacks.SelectElement)
+	{
+		CachedCallbacks.SelectElement(InsertionDeltaKey, NewElementIndex);
+	}
+}
+
+void FPCGPointInsertionViewportExtension::UpdateInsertedElementLocation(const FPCGDeltaKey& DeltaKey, const int32 Index, const FVector& NewLocation)
+{
+	UPCGComponent* PCGComponent = CachedContext.ActivePCGComponent.Get();
+	if (!PCGComponent)
+	{
+		return;
+	}
+
+	FPCGSourceDataContainer* DataContainer = PCGComponent->GetExecutionState().GetSourceDataContainer();
+	if (!DataContainer)
+	{
+		return;
+	}
+
+	FSharedStruct SharedStruct = DataContainer->GetMutable<FPCGDeltaCollection>(CachedContext.ActiveStorageKey);
+	FPCGDeltaCollection* Collection = SharedStruct.IsValid() ? SharedStruct.GetPtr<FPCGDeltaCollection>() : nullptr;
+	if (!Collection)
+	{
+		return;
+	}
+
+	FScopedTransaction Transaction(LOCTEXT("UpdateInsertedElementLocation", "PCG Manual Edit: Update Inserted Point Location"));
+	PCGComponent->Modify();
+
+	if (TInstancedStruct<FPCGDeltaBase>* DeltaStruct = Collection->Find(DeltaKey))
+	{
+		if (FPCGPointInsertionDelta* InsertDelta = DeltaStruct->GetMutablePtr<FPCGPointInsertionDelta>())
+		{
+			if (InsertDelta->InsertedPoints.IsValidIndex(Index))
+			{
+				InsertDelta->InsertedPoints[Index].Transform.SetLocation(NewLocation);
+			}
+		}
+	}
+
+	PCG::DeltaViewportExtension::MarkComponentDirty(PCGComponent, DataContainer);
+
+	if (CachedCallbacks.SelectElement)
+	{
+		CachedCallbacks.SelectElement(DeltaKey, Index);
+	}
+}
+
+void FPCGPointInsertionViewportExtension::RemoveInsertedElement(int32 Index)
+{
+	UPCGComponent* PCGComponent = CachedContext.ActivePCGComponent.Get();
+	if (!PCGComponent)
+	{
+		return;
+	}
+
+	FPCGSourceDataContainer* DataContainer = PCGComponent->GetExecutionState().GetSourceDataContainer();
+	if (!DataContainer)
+	{
+		return;
+	}
+
+	FSharedStruct SharedStruct = DataContainer->GetMutable<FPCGDeltaCollection>(CachedContext.ActiveStorageKey);
+	FPCGDeltaCollection* Collection = SharedStruct.IsValid() ? SharedStruct.GetPtr<FPCGDeltaCollection>() : nullptr;
+	if (!Collection)
+	{
+		return;
+	}
+
+	FScopedTransaction Transaction(LOCTEXT("RemoveInsertedElement", "PCG Manual Edit: Remove Inserted Point"));
+	PCGComponent->Modify();
+
+	Collection->ForEachDelta([Index](const FPCGDeltaKey& Key, TInstancedStruct<FPCGDeltaBase>& DeltaValue) -> bool
+	{
+		if (FPCGPointInsertionDelta* InsertDelta = DeltaValue.GetMutablePtr<FPCGPointInsertionDelta>())
+		{
+			if (InsertDelta->InsertedPoints.IsValidIndex(Index))
+			{
+				InsertDelta->InsertedPoints.RemoveAt(Index);
+			}
+
+			return false;
+		}
+
+		return true;
+	});
+
+	PCG::DeltaViewportExtension::MarkComponentDirty(PCGComponent, DataContainer);
+}
+
+void FPCGPointInsertionViewportExtension::RemoveAllInsertedElements()
+{
+	UPCGComponent* PCGComponent = CachedContext.ActivePCGComponent.Get();
+	if (!PCGComponent)
+	{
+		return;
+	}
+
+	FPCGSourceDataContainer* DataContainer = PCGComponent->GetExecutionState().GetSourceDataContainer();
+	if (!DataContainer)
+	{
+		return;
+	}
+
+	FSharedStruct SharedStruct = DataContainer->GetMutable<FPCGDeltaCollection>(CachedContext.ActiveStorageKey);
+	FPCGDeltaCollection* Collection = SharedStruct.IsValid() ? SharedStruct.GetPtr<FPCGDeltaCollection>() : nullptr;
+	if (!Collection)
+	{
+		return;
+	}
+
+	FScopedTransaction Transaction(LOCTEXT("RemoveAllInsertedElements", "PCG Manual Edit: Remove All Inserted Points"));
+	PCGComponent->Modify();
+
+	TArray<FPCGDeltaKey> InsertionKeys;
+	Collection->ForEachDelta([&InsertionKeys](const FPCGDeltaKey& DeltaKey, const TInstancedStruct<FPCGDeltaBase>& DeltaValue) -> bool
+	{
+		if (DeltaValue.GetPtr<FPCGPointInsertionDelta>())
+		{
+			InsertionKeys.Add(DeltaKey);
+		}
+		return true;
+	});
+
+	for (const FPCGDeltaKey& Key : InsertionKeys)
+	{
+		Collection->Remove(Key);
+	}
+
+	PCG::DeltaViewportExtension::MarkComponentDirty(PCGComponent, DataContainer);
+}
+
+void FPCGPointInsertionViewportExtension::OnInsertionComponentCommitted(const int32 EntryIndex, const int32 AxisIndex, const double NewValue)
+{
+	for (const TSharedPtr<FPCGPointInsertionDeltaEntry>& Entry : InsertionEntries)
+	{
+		if (Entry && Entry->Index == EntryIndex)
+		{
+			Entry->Location[AxisIndex] = NewValue;
+			UpdateInsertedElementLocation(Entry->DeltaKey, EntryIndex, Entry->Location);
+			return;
+		}
+	}
+}
+
+TSharedRef<ITableRow> FPCGPointInsertionViewportExtension::OnGenerateInsertionRow(TSharedPtr<FPCGPointInsertionDeltaEntry> Entry, const TSharedRef<STableViewBase>& OwnerTable)
+{
+	TWeakPtr<FPCGPointInsertionDeltaEntry> WeakEntry = Entry;
+	int32 EntryIndex = Entry->Index;
+
+	return SNew(STableRow<TSharedPtr<FPCGPointInsertionDeltaEntry>>, OwnerTable)
+		[
+			SNew(SHorizontalBox)
+			+ SHorizontalBox::Slot()
+		   .FillWidth(1.0f)
+		   .VAlign(VAlign_Center)
+		   .Padding(4.0f, 1.0f)
+			[
+				SNew(SVectorSlider<double>, /*bIsRotator=*/false, /*InProperty=*/nullptr)
+			   .VisibleText_0_Lambda([WeakEntry]() -> FString
+				{
+					TSharedPtr<FPCGPointInsertionDeltaEntry> Pinned = WeakEntry.Pin();
+					return Pinned ? FString::SanitizeFloat(Pinned->Location.X) : FString();
+				})
+			   .VisibleText_1_Lambda([WeakEntry]() -> FString
+				{
+					TSharedPtr<FPCGPointInsertionDeltaEntry> Pinned = WeakEntry.Pin();
+					return Pinned ? FString::SanitizeFloat(Pinned->Location.Y) : FString();
+				})
+			   .VisibleText_2_Lambda([WeakEntry]() -> FString
+				{
+					TSharedPtr<FPCGPointInsertionDeltaEntry> Pinned = WeakEntry.Pin();
+					return Pinned ? FString::SanitizeFloat(Pinned->Location.Z) : FString();
+				})
+			   .OnNumericCommitted_Box_0(SVectorSlider<double>::FOnNumericValueCommitted::CreateLambda(
+					[this, EntryIndex](const double V, ETextCommit::Type) { OnInsertionComponentCommitted(EntryIndex, 0, V); }))
+			   .OnNumericCommitted_Box_1(SVectorSlider<double>::FOnNumericValueCommitted::CreateLambda(
+					[this, EntryIndex](const double V, ETextCommit::Type) { OnInsertionComponentCommitted(EntryIndex, 1, V); }))
+			   .OnNumericCommitted_Box_2(SVectorSlider<double>::FOnNumericValueCommitted::CreateLambda(
+					[this, EntryIndex](const double V, ETextCommit::Type) { OnInsertionComponentCommitted(EntryIndex, 2, V); }))
+			]
+			+ SHorizontalBox::Slot()
+		   .AutoWidth()
+		   .VAlign(VAlign_Center)
+		   .Padding(2.0f, 0.0f)
+			[
+				SNew(SButton)
+			   .ButtonStyle(FAppStyle::Get(), "SimpleButton")
+			   .ContentPadding(FMargin(1, 0))
+			   .ToolTipText(LOCTEXT("SelectInsertionTooltip", "Select this inserted point"))
+			   .OnClicked_Lambda([this, WeakEntry]() -> FReply
+				{
+					if (CachedCallbacks.SelectElement)
+					{
+						if (TSharedPtr<FPCGPointInsertionDeltaEntry> Pinned = WeakEntry.Pin())
+						{
+							CachedCallbacks.SelectElement(Pinned->DeltaKey, Pinned->Index);
+						}
+					}
+					return FReply::Handled();
+				})
+				[
+					SNew(SImage)
+				   .Image(FAppStyle::Get().GetBrush("Icons.Search"))
+				   .ColorAndOpacity(FSlateColor::UseForeground())
+				]
+			]
+			+ SHorizontalBox::Slot()
+		   .AutoWidth()
+		   .VAlign(VAlign_Center)
+		   .Padding(2.0f, 0.0f)
+			[
+				SNew(SButton)
+			   .ButtonStyle(FAppStyle::Get(), "SimpleButton")
+			   .ContentPadding(FMargin(1, 0))
+			   .ToolTipText(LOCTEXT("RemoveInsertionTooltip", "Remove this inserted point"))
+			   .OnClicked_Lambda([this, EntryIndex]() -> FReply
+				{
+					RemoveInsertedElement(EntryIndex);
+					RefreshLists(CachedContext);
+					return FReply::Handled();
+				})
+				[
+					SNew(SImage)
+				   .Image(FAppStyle::Get().GetBrush("Icons.Delete"))
+				   .ColorAndOpacity(FSlateColor::UseForeground())
+				]
+			]
+		];
+}
+
+#undef LOCTEXT_NAMESPACE
diff --git a/Engine/Plugins/PCG/Source/PCGEditor/Private/DeltaViewportExtensions/PCGPointInsertionViewportExtension.h b/Engine/Plugins/PCG/Source/PCGEditor/Private/DeltaViewportExtensions/PCGPointInsertionViewportExtension.h
new file mode 100644
--- /dev/null
+++ b/Engine/Plugins/PCG/Source/PCGEditor/Private/DeltaViewportExtensions/PCGPointInsertionViewportExtension.h
@@ -0,0 +1,69 @@
+// Copyright Epic Games, Inc. All Rights Reserved.
+
+#pragma once
+
+#include "DeltaViewportExtensions/PCGDeltaViewportExtension.h"
+
+#include "Graph/DataOverride/PCGDataOverridePoints.h"
+
+class ITableRow;
+class STableViewBase;
+class SVerticalBox;
+template <typename ItemType> class SListView;
+
+struct FPCGPointInsertionDeltaEntry
+{
+	FPCGDeltaKey DeltaKey;
+	int32 Index = INDEX_NONE;
+	FVector Location = FVector::ZeroVector;
+};
+
+namespace PCGInsertionConstants
+{
+	// @todo_pcg: Find an appropriate hash value for the insertion delta key. Currently an arbitrary non-zero constant.
+	constexpr FXxHash64 DefaultInsertionKeyHash{1};
+}
+
+/** Viewport extension for FPCGPointInsertionDelta: Insert button + insertions list + viewport policy. */
+class FPCGPointInsertionViewportExtension : public IPCGDeltaViewportExtension
+{
+public:
+	// ~Begin IPCGDeltaViewportExtension interface
+	virtual FText GetDisplayName() const override { return NSLOCTEXT("PCGPointInsertionViewportExtension", "PointInsertionDisplayName", "Point Insertion"); }
+	virtual FName GetDeltaName() const override { return FPCGPointInsertionDelta::GetDeltaNameStatic(); }
+	virtual int32 GetSortPriority() const override { return 1; }
+	virtual TSharedRef<SWidget> CreateWidget(FPCGDeltaViewportCallbacks Callbacks) override;
+	virtual void UpdateContext(const FPCGDeltaViewportContext& Context) override { CachedContext = Context; }
+	virtual void RefreshLists(const FPCGDeltaViewportContext& Context) override;
+	virtual TArray<FPCGDeltaKeyBinding> GetKeyBindings() override { return {{EKeys::Delete, [this]() { HandleDeletePressed(); }}}; }
+	virtual FLinearColor GetDisplayColor(const bool bIsSelected) const override;
+	virtual void GetInsertedElements(const FPCGDeltaCollection& Collection, TArray<FInsertedViewportElement>& OutElements) const override;
+	virtual int32 OnAltDrag(FPCGDeltaCollection& Collection, const FTransform& SourceTransform, FPCGDeltaKey& OutDeltaKey) const override;
+	virtual bool ApplyGizmoTransform(const FTransform& NewTransform, TInstancedStruct<FPCGDeltaBase>& DeltaStruct, const int32 ElementIndex) const override;
+	// ~End IPCGDeltaViewportExtension interface
+
+private:
+	/** Get or create the mutable delta collection from the cached context. */
+	FPCGDeltaCollection* GetOrCreateMutableCollection(FPCGSourceDataContainer* DataContainer);
+
+	/** Find or create the single insertion delta within the collection. */
+	static FPCGPointInsertionDelta* FindOrCreateInsertionDelta(FPCGDeltaCollection* Collection, FPCGDeltaKey* OutKey = nullptr);
+
+	FReply OnInsertPointClicked();
+	FReply OnRemoveAllClicked();
+
+	void HandleDeletePressed();
+	void AddInsertedElement();
+	void UpdateInsertedElementLocation(const FPCGDeltaKey& DeltaKey, const int32 Index, const FVector& NewLocation);
+	void RemoveInsertedElement(int32 Index);
+	void RemoveAllInsertedElements();
+	void OnInsertionComponentCommitted(const int32 EntryIndex, const int32 AxisIndex, const double NewValue);
+	TSharedRef<ITableRow> OnGenerateInsertionRow(TSharedPtr<FPCGPointInsertionDeltaEntry> Entry, const TSharedRef<STableViewBase>& OwnerTable);
+
+	FPCGDeltaViewportContext CachedContext;
+	FPCGDeltaViewportCallbacks CachedCallbacks;
+
+	TArray<TSharedPtr<FPCGPointInsertionDeltaEntry>> InsertionEntries;
+	TSharedPtr<SVerticalBox> InsertionListSection;
+	TSharedPtr<SListView<TSharedPtr<FPCGPointInsertionDeltaEntry>>> InsertionListView;
+};
diff --git a/Engine/Plugins/PCG/Source/PCGEditor/Private/DeltaVisualizations/PCGDeltaBaseVisualization.cpp b/Engine/Plugins/PCG/Source/PCGEditor/Private/DeltaVisualizations/PCGDeltaBaseVisualization.cpp
new file mode 100644
--- /dev/null
+++ b/Engine/Plugins/PCG/Source/PCGEditor/Private/DeltaVisualizations/PCGDeltaBaseVisualization.cpp
@@ -0,0 +1,17 @@
+// Copyright Epic Games, Inc. All Rights Reserved.
+
+#include "DeltaVisualizations/PCGDeltaBaseVisualization.h"
+
+#include "DeltaVisualizations/PCGDeltaVisualizationHelpers.h"
+
+TArray<FPCGDeltaVisualizerColumnInfo> FPCGDeltaBaseVisualization::GetColumnInfos() const
+{
+	TArray<FPCGDeltaVisualizerColumnInfo> Columns;
+	PCGDeltaVisualization::Helpers::AppendBaseDeltaColumns(Columns);
+	return Columns;
+}
+
+FText FPCGDeltaBaseVisualization::GetCellText(const FName ColumnId, const FPCGDeltaKey& DeltaKey, const FConstStructView Delta) const
+{
+	return PCGDeltaVisualization::Helpers::GetBaseDeltaCellText(ColumnId, DeltaKey, Delta);
+}
diff --git a/Engine/Plugins/PCG/Source/PCGEditor/Private/DeltaVisualizations/PCGDeltaBaseVisualization.h b/Engine/Plugins/PCG/Source/PCGEditor/Private/DeltaVisualizations/PCGDeltaBaseVisualization.h
new file mode 100644
--- /dev/null
+++ b/Engine/Plugins/PCG/Source/PCGEditor/Private/DeltaVisualizations/PCGDeltaBaseVisualization.h
@@ -0,0 +1,13 @@
+// Copyright Epic Games, Inc. All Rights Reserved.
+
+#pragma once
+
+#include "DeltaVisualizations/PCGDeltaVisualization.h"
+
+/** Default visualizer registered to FPCGDeltaBase. Shows common delta metadata columns. */
+class FPCGDeltaBaseVisualization : public IPCGDeltaVisualization
+{
+public:
+	virtual TArray<FPCGDeltaVisualizerColumnInfo> GetColumnInfos() const override;
+	virtual FText GetCellText(FName ColumnId, const FPCGDeltaKey& DeltaKey, FConstStructView Delta) const override;
+};
diff --git a/Engine/Plugins/PCG/Source/PCGEditor/Private/DeltaVisualizations/PCGDeltaPointVisualization.cpp b/Engine/Plugins/PCG/Source/PCGEditor/Private/DeltaVisualizations/PCGDeltaPointVisualization.cpp
new file mode 100644
--- /dev/null
+++ b/Engine/Plugins/PCG/Source/PCGEditor/Private/DeltaVisualizations/PCGDeltaPointVisualization.cpp
@@ -0,0 +1,312 @@
+// Copyright Epic Games, Inc. All Rights Reserved.
+
+#include "DeltaVisualizations/PCGDeltaPointVisualization.h"
+
+#include "DeltaVisualizations/PCGDeltaVisualizationHelpers.h"
+#include "Graph/DataOverride/PCGDataOverride.h"
+#include "Graph/DataOverride/PCGDataOverridePoints.h"
+
+#define LOCTEXT_NAMESPACE "PCGDeltaVisualization"
+
+namespace PCGDeltaPoints::Constants
+{
+	static constexpr float TransformComponentColumnWidth = 120.f;
+	static constexpr EPCGTableVisualizerCellAlignment TransformComponentColumnAlignment = EPCGTableVisualizerCellAlignment::Right;
+} // namespace PCGDeltaPoints::Constants
+
+// SHeaderRow column identifiers for point delta visualizers.
+namespace PCGDeltaVisualizationColumns
+{
+	static const FName ElementIndexId = TEXT("ElementIndex");
+
+	// Per-component transform columns — Previous (OriginalTransform)
+	static const FName PreviousPosX = TEXT("Previous.Position.X");
+	static const FName PreviousPosY = TEXT("Previous.Position.Y");
+	static const FName PreviousPosZ = TEXT("Previous.Position.Z");
+	static const FName PreviousRotRoll = TEXT("Previous.Rotation.Roll");
+	static const FName PreviousRotPitch = TEXT("Previous.Rotation.Pitch");
+	static const FName PreviousRotYaw = TEXT("Previous.Rotation.Yaw");
+	static const FName PreviousScaleX = TEXT("Previous.Scale.X");
+	static const FName PreviousScaleY = TEXT("Previous.Scale.Y");
+	static const FName PreviousScaleZ = TEXT("Previous.Scale.Z");
+
+	// Per-component transform columns — Override (TransformOverride)
+	static const FName OverridePosX = TEXT("Override.Position.X");
+	static const FName OverridePosY = TEXT("Override.Position.Y");
+	static const FName OverridePosZ = TEXT("Override.Position.Z");
+	static const FName OverrideRotRoll = TEXT("Override.Rotation.Roll");
+	static const FName OverrideRotPitch = TEXT("Override.Rotation.Pitch");
+	static const FName OverrideRotYaw = TEXT("Override.Rotation.Yaw");
+	static const FName OverrideScaleX = TEXT("Override.Scale.X");
+	static const FName OverrideScaleY = TEXT("Override.Scale.Y");
+	static const FName OverrideScaleZ = TEXT("Override.Scale.Z");
+
+	// Point edit delta columns
+	static const FName OverrideFlagsId = TEXT("OverrideFlags");
+
+	// Point insertion delta columns
+	static const FName InsertionSignatureId = TEXT("InsertionSignature");
+	static const FName InsertedPointsId = TEXT("InsertedPoints");
+
+	// Spacer (visual end-cap padding)
+	static const FName SpacerId = TEXT("Spacer");
+} // namespace PCGDeltaVisualizationColumns
+
+namespace PCGDeltaPointVisualizationHelpers
+{
+	/** Bundles the 9 column IDs for one decomposed FTransform. */
+	struct FTransformColumnIds
+	{
+		FName PosX, PosY, PosZ;
+		FName RotRoll, RotPitch, RotYaw;
+		FName ScaleX, ScaleY, ScaleZ;
+	};
+
+	static const FTransformColumnIds PreviousTransformColumnIds =
+	{
+		PCGDeltaVisualizationColumns::PreviousPosX, PCGDeltaVisualizationColumns::PreviousPosY, PCGDeltaVisualizationColumns::PreviousPosZ,
+		PCGDeltaVisualizationColumns::PreviousRotRoll, PCGDeltaVisualizationColumns::PreviousRotPitch, PCGDeltaVisualizationColumns::PreviousRotYaw,
+		PCGDeltaVisualizationColumns::PreviousScaleX, PCGDeltaVisualizationColumns::PreviousScaleY, PCGDeltaVisualizationColumns::PreviousScaleZ
+	};
+
+	static const FTransformColumnIds OverrideTransformColumnIds =
+	{
+		PCGDeltaVisualizationColumns::OverridePosX, PCGDeltaVisualizationColumns::OverridePosY, PCGDeltaVisualizationColumns::OverridePosZ,
+		PCGDeltaVisualizationColumns::OverrideRotRoll, PCGDeltaVisualizationColumns::OverrideRotPitch, PCGDeltaVisualizationColumns::OverrideRotYaw,
+		PCGDeltaVisualizationColumns::OverrideScaleX, PCGDeltaVisualizationColumns::OverrideScaleY, PCGDeltaVisualizationColumns::OverrideScaleZ
+	};
+
+	/** Appends 9 per-component transform columns (Position X/Y/Z, Rotation Roll/Pitch/Yaw, Scale X/Y/Z). */
+	static void AppendTransformColumns(TArray<FPCGDeltaVisualizerColumnInfo>& OutColumns, const FText& GroupLabel, const FText& TooltipBase, const FTransformColumnIds& Ids)
+	{
+		using namespace PCGDeltaPoints::Constants;
+
+		auto CreateLabel = [&GroupLabel](const FText& Component) { return FText::Format(INVTEXT("{0}.{1}"), GroupLabel, Component); };
+		auto CreateToolTip = [&TooltipBase](const FText& Component) { return FText::Format(INVTEXT("{0} - {1}"), TooltipBase, Component); };
+
+		const FText PosX = INVTEXT("$Position.X");
+		const FText PosY = INVTEXT("$Position.Y");
+		const FText PosZ = INVTEXT("$Position.Z");
+		const FText RotRoll = INVTEXT("$Rotation.Roll");
+		const FText RotPitch = INVTEXT("$Rotation.Pitch");
+		const FText RotYaw = INVTEXT("$Rotation.Yaw");
+		const FText ScaleX = INVTEXT("$Scale.X");
+		const FText ScaleY = INVTEXT("$Scale.Y");
+		const FText ScaleZ = INVTEXT("$Scale.Z");
+
+		OutColumns.Add({Ids.PosX, CreateLabel(PosX), CreateToolTip(PosX), TransformComponentColumnWidth, TransformComponentColumnAlignment});
+		OutColumns.Add({Ids.PosY, CreateLabel(PosY), CreateToolTip(PosY), TransformComponentColumnWidth, TransformComponentColumnAlignment});
+		OutColumns.Add({Ids.PosZ, CreateLabel(PosZ), CreateToolTip(PosZ), TransformComponentColumnWidth, TransformComponentColumnAlignment});
+		OutColumns.Add({Ids.RotRoll, CreateLabel(RotRoll), CreateToolTip(RotRoll), TransformComponentColumnWidth, TransformComponentColumnAlignment});
+		OutColumns.Add({Ids.RotPitch, CreateLabel(RotPitch), CreateToolTip(RotPitch), TransformComponentColumnWidth, TransformComponentColumnAlignment});
+		OutColumns.Add({Ids.RotYaw, CreateLabel(RotYaw), CreateToolTip(RotYaw), TransformComponentColumnWidth, TransformComponentColumnAlignment});
+		OutColumns.Add({Ids.ScaleX, CreateLabel(ScaleX), CreateToolTip(ScaleX), TransformComponentColumnWidth, TransformComponentColumnAlignment});
+		OutColumns.Add({Ids.ScaleY, CreateLabel(ScaleY), CreateToolTip(ScaleY), TransformComponentColumnWidth, TransformComponentColumnAlignment});
+		OutColumns.Add({Ids.ScaleZ, CreateLabel(ScaleZ), CreateToolTip(ScaleZ), TransformComponentColumnWidth, TransformComponentColumnAlignment});
+	}
+
+	/** Returns cell text for a per-component transform column, or empty text if the column is not handled. */
+	static FText GetTransformCellText(const FName ColumnId, const FTransform& Transform, const FTransformColumnIds& Ids)
+	{
+		const FVector Translation = Transform.GetTranslation();
+		if (ColumnId == Ids.PosX) { return FText::AsNumber(Translation.X); }
+		if (ColumnId == Ids.PosY) { return FText::AsNumber(Translation.Y); }
+		if (ColumnId == Ids.PosZ) { return FText::AsNumber(Translation.Z); }
+
+		const FRotator Rotation = Transform.Rotator();
+		if (ColumnId == Ids.RotRoll) { return FText::AsNumber(Rotation.Roll); }
+		if (ColumnId == Ids.RotPitch) { return FText::AsNumber(Rotation.Pitch); }
+		if (ColumnId == Ids.RotYaw) { return FText::AsNumber(Rotation.Yaw); }
+
+		const FVector Scale = Transform.GetScale3D();
+		if (ColumnId == Ids.ScaleX) { return FText::AsNumber(Scale.X); }
+		if (ColumnId == Ids.ScaleY) { return FText::AsNumber(Scale.Y); }
+		if (ColumnId == Ids.ScaleZ) { return FText::AsNumber(Scale.Z); }
+
+		return FText::GetEmpty();
+	}
+
+	/** Appends columns common to point-based deltas (Element Index, Previous Transform components). */
+	static void AppendPointDeltaBaseColumns(TArray<FPCGDeltaVisualizerColumnInfo>& OutColumns)
+	{
+		using namespace PCGDeltaVisualizationColumns;
+
+		OutColumns.Add({ElementIndexId, LOCTEXT("ElementIndexLabel", "Element Index"), LOCTEXT("ElementIndexTooltip", "Original element index (debug only)"), 95.f, EPCGTableVisualizerCellAlignment::Center});
+		AppendTransformColumns(OutColumns, LOCTEXT("PreviousGroupLabel_Base", "Previous"), LOCTEXT("PreviousTransformTooltip_Base", "The original transform of the point"), PreviousTransformColumnIds);
+	}
+
+	/** Returns cell text for point delta base columns. Returns empty text if the column is not handled. */
+	static FText GetPointDeltaBaseCellText(const FName ColumnId, const FConstStructView Delta)
+	{
+		using namespace PCGDeltaVisualizationColumns;
+
+		if (ColumnId == ElementIndexId)
+		{
+			// ElementIndex is on both FPCGPointEditDelta and FPCGPointDeletionDelta but not on FPCGPointDeltaBase itself.
+			// Check each concrete type.
+			if (const FPCGPointEditDelta* EditDelta = Delta.GetPtr<const FPCGPointEditDelta>())
+			{
+				return FText::AsNumber(EditDelta->ElementIndex);
+			}
+			else if (const FPCGPointDeletionDelta* DeletionDelta = Delta.GetPtr<const FPCGPointDeletionDelta>())
+			{
+				return FText::AsNumber(DeletionDelta->ElementIndex);
+			}
+
+			return FText::GetEmpty();
+		}
+
+		if (const FPCGPointDeltaBase* PointDelta = Delta.GetPtr<const FPCGPointDeltaBase>())
+		{
+			FText TransformText = GetTransformCellText(ColumnId, PointDelta->OriginalTransform, PreviousTransformColumnIds);
+			if (!TransformText.IsEmpty())
+			{
+				return TransformText;
+			}
+		}
+
+		return FText::GetEmpty();
+	}
+} // namespace PCGDeltaPointVisualizationHelpers
+
+// --- FPCGPointEditDeltaVisualization ---
+
+TArray<FPCGDeltaVisualizerColumnInfo> FPCGPointEditDeltaVisualization::GetColumnInfos() const
+{
+	using namespace PCGDeltaVisualizationColumns;
+
+	TArray<FPCGDeltaVisualizerColumnInfo> Columns;
+
+	// Build columns directly in desired display order rather than using shared append helpers.
+	Columns.Add({ElementIndexId, LOCTEXT("PreviousIndexLabel", "Previous Index"), LOCTEXT("PreviousIndexTooltip", "Original element index before overrides"), 85.f, EPCGTableVisualizerCellAlignment::Center});
+	Columns.Add({PCGDeltaVisualization::Constants::DeltaTypeId, LOCTEXT("EditDeltaTypeLabel", "Delta Type"), LOCTEXT("EditDeltaTypeTooltip", "Name of the delta type"), 100.f, EPCGTableVisualizerCellAlignment::Left});
+	Columns.Add({OverrideFlagsId, LOCTEXT("SignatureLabel", "Signature"), LOCTEXT("SignatureTooltip", "Which components are overridden (P=Position, R=Rotation, S=Scale)"), 100.f, EPCGTableVisualizerCellAlignment::Center});
+	PCGDeltaPointVisualizationHelpers::AppendTransformColumns(Columns, LOCTEXT("PreviousGroupLabel_Edit", "Previous"), LOCTEXT("PreviousTransformTooltip_Edit", "The original transform of the point"), PCGDeltaPointVisualizationHelpers::PreviousTransformColumnIds);
+	PCGDeltaPointVisualizationHelpers::AppendTransformColumns(Columns, LOCTEXT("OverrideGroupLabel_Edit", "Override"), LOCTEXT("OverrideTransformTooltip_Edit", "The overridden transform value"), PCGDeltaPointVisualizationHelpers::OverrideTransformColumnIds);
+	Columns.Add({SpacerId, FText::GetEmpty(), FText::GetEmpty(), 16.f, EPCGTableVisualizerCellAlignment::Left});
+
+	return Columns;
+}
+
+// Returns the text for any column this visualizer supports, or empty if unhandled.
+FText FPCGPointEditDeltaVisualization::GetCellText(const FName ColumnId, const FPCGDeltaKey& DeltaKey, const FConstStructView Delta) const
+{
+	using namespace PCGDeltaVisualizationColumns;
+
+	// Try base columns first.
+	FText BaseText = PCGDeltaVisualization::Helpers::GetBaseDeltaCellText(ColumnId, DeltaKey, Delta);
+	if (!BaseText.IsEmpty())
+	{
+		return BaseText;
+	}
+
+	// Try point delta base columns.
+	FText PointBaseText = PCGDeltaPointVisualizationHelpers::GetPointDeltaBaseCellText(ColumnId, Delta);
+	if (!PointBaseText.IsEmpty())
+	{
+		return PointBaseText;
+	}
+
+	// Edit-specific columns.
+	const FPCGPointEditDelta* EditDelta = Delta.GetPtr<const FPCGPointEditDelta>();
+	if (!EditDelta)
+	{
+		return FText::GetEmpty();
+	}
+
+	FText OverrideTransformText = PCGDeltaPointVisualizationHelpers::GetTransformCellText(ColumnId, EditDelta->TransformOverride, PCGDeltaPointVisualizationHelpers::OverrideTransformColumnIds);
+	if (!OverrideTransformText.IsEmpty())
+	{
+		return OverrideTransformText;
+	}
+
+	// Compact signature string (P=Position, R=Rotation, S=Scale)
+	if (ColumnId == OverrideFlagsId)
+	{
+		FString Flags;
+		if (EditDelta->bOverridePosition) { Flags += TEXT("P"); }
+		if (EditDelta->bOverrideRotation) { Flags += TEXT("R"); }
+		if (EditDelta->bOverrideScale) { Flags += TEXT("S"); }
+		return FText::FromString(Flags);
+	}
+
+	return FText::GetEmpty();
+}
+
+// --- FPCGPointDeletionDeltaVisualization ---
+
+TArray<FPCGDeltaVisualizerColumnInfo> FPCGPointDeletionDeltaVisualization::GetColumnInfos() const
+{
+	TArray<FPCGDeltaVisualizerColumnInfo> Columns;
+	PCGDeltaVisualization::Helpers::AppendBaseDeltaColumns(Columns);
+	PCGDeltaPointVisualizationHelpers::AppendPointDeltaBaseColumns(Columns);
+	return Columns;
+}
+
+FText FPCGPointDeletionDeltaVisualization::GetCellText(const FName ColumnId, const FPCGDeltaKey& DeltaKey, const FConstStructView Delta) const
+{
+	FText BaseText = PCGDeltaVisualization::Helpers::GetBaseDeltaCellText(ColumnId, DeltaKey, Delta);
+	if (!BaseText.IsEmpty())
+	{
+		return BaseText;
+	}
+
+	return PCGDeltaPointVisualizationHelpers::GetPointDeltaBaseCellText(ColumnId, Delta);
+}
+
+// --- FPCGPointInsertionDeltaVisualization ---
+
+TArray<FPCGDeltaVisualizerColumnInfo> FPCGPointInsertionDeltaVisualization::GetColumnInfos() const
+{
+	using namespace PCGDeltaVisualizationColumns;
+
+	TArray<FPCGDeltaVisualizerColumnInfo> Columns;
+
+	Columns.Add({PCGDeltaVisualization::Constants::DeltaTypeId, LOCTEXT("InsertDeltaTypeLabel", "Delta Type"), LOCTEXT("InsertDeltaTypeTooltip", "Name of the delta type"), 100.f, EPCGTableVisualizerCellAlignment::Left});
+	Columns.Add({InsertionSignatureId, LOCTEXT("InsertionSignatureLabel", "Signature"), LOCTEXT("InsertionSignatureTooltip", "Identifies this delta as a point insertion"), 100.f, EPCGTableVisualizerCellAlignment::Center});
+	Columns.Add({InsertedPointsId, LOCTEXT("InsertedPointsLabel", "Inserted Points"), LOCTEXT("InsertedPointsTooltip", "Number of points to be inserted"), 105.f, EPCGTableVisualizerCellAlignment::Center});
+	PCGDeltaPointVisualizationHelpers::AppendTransformColumns(Columns, LOCTEXT("OverrideGroupLabel_Insert", "Override"), LOCTEXT("OverrideTransformTooltip_Insert", "The transform of the inserted point"), PCGDeltaPointVisualizationHelpers::OverrideTransformColumnIds);
+	Columns.Add({SpacerId, FText::GetEmpty(), FText::GetEmpty(), 16.f, EPCGTableVisualizerCellAlignment::Left});
+
+	return Columns;
+}
+
+FText FPCGPointInsertionDeltaVisualization::GetCellText(const FName ColumnId, const FPCGDeltaKey& DeltaKey, const FConstStructView Delta) const
+{
+	using namespace PCGDeltaVisualizationColumns;
+
+	FText BaseText = PCGDeltaVisualization::Helpers::GetBaseDeltaCellText(ColumnId, DeltaKey, Delta);
+	if (!BaseText.IsEmpty())
+	{
+		return BaseText;
+	}
+
+	const FPCGPointInsertionDelta* InsertionDelta = Delta.GetPtr<const FPCGPointInsertionDelta>();
+	if (!InsertionDelta)
+	{
+		return FText::GetEmpty();
+	}
+
+	if (ColumnId == InsertionSignatureId)
+	{
+		return LOCTEXT("InsertionSignatureValue", "PointInsert");
+	}
+
+	if (ColumnId == InsertedPointsId)
+	{
+		return FText::AsNumber(InsertionDelta->InsertedPoints.Num());
+	}
+
+	if (!InsertionDelta->InsertedPoints.IsEmpty())
+	{
+		FText OverrideTransformText = PCGDeltaPointVisualizationHelpers::GetTransformCellText(ColumnId, InsertionDelta->InsertedPoints[0].Transform, PCGDeltaPointVisualizationHelpers::OverrideTransformColumnIds);
+		if (!OverrideTransformText.IsEmpty())
+		{
+			return OverrideTransformText;
+		}
+	}
+
+	return FText::GetEmpty();
+}
+
+#undef LOCTEXT_NAMESPACE
diff --git a/Engine/Plugins/PCG/Source/PCGEditor/Private/DeltaVisualizations/PCGDeltaPointVisualization.h b/Engine/Plugins/PCG/Source/PCGEditor/Private/DeltaVisualizations/PCGDeltaPointVisualization.h
new file mode 100644
--- /dev/null
+++ b/Engine/Plugins/PCG/Source/PCGEditor/Private/DeltaVisualizations/PCGDeltaPointVisualization.h
@@ -0,0 +1,29 @@
+// Copyright Epic Games, Inc. All Rights Reserved.
+
+#pragma once
+
+#include "DeltaVisualizations/PCGDeltaVisualization.h"
+
+/** Visualizer for FPCGPointEditDelta. Extends base columns with transform override info. */
+class FPCGPointEditDeltaVisualization : public IPCGDeltaVisualization
+{
+public:
+	virtual TArray<FPCGDeltaVisualizerColumnInfo> GetColumnInfos() const override;
+	virtual FText GetCellText(FName ColumnId, const FPCGDeltaKey& DeltaKey, FConstStructView Delta) const override;
+};
+
+/** Visualizer for FPCGPointDeletionDelta. Extends base columns with original transform. */
+class FPCGPointDeletionDeltaVisualization : public IPCGDeltaVisualization
+{
+public:
+	virtual TArray<FPCGDeltaVisualizerColumnInfo> GetColumnInfos() const override;
+	virtual FText GetCellText(FName ColumnId, const FPCGDeltaKey& DeltaKey, FConstStructView Delta) const override;
+};
+
+/** Visualizer for FPCGPointInsertionDelta. Extends base columns with inserted point count. */
+class FPCGPointInsertionDeltaVisualization : public IPCGDeltaVisualization
+{
+public:
+	virtual TArray<FPCGDeltaVisualizerColumnInfo> GetColumnInfos() const override;
+	virtual FText GetCellText(FName ColumnId, const FPCGDeltaKey& DeltaKey, FConstStructView Delta) const override;
+};
diff --git a/Engine/Plugins/PCG/Source/PCGEditor/Private/DeltaVisualizations/PCGDeltaVisualizationHelpers.cpp b/Engine/Plugins/PCG/Source/PCGEditor/Private/DeltaVisualizations/PCGDeltaVisualizationHelpers.cpp
new file mode 100644
--- /dev/null
+++ b/Engine/Plugins/PCG/Source/PCGEditor/Private/DeltaVisualizations/PCGDeltaVisualizationHelpers.cpp
@@ -0,0 +1,31 @@
+// Copyright Epic Games, Inc. All Rights Reserved.
+
+#include "DeltaVisualizations/PCGDeltaVisualizationHelpers.h"
+
+#include "Graph/DataOverride/PCGDataOverride.h"
+
+#define LOCTEXT_NAMESPACE "PCGDeltaVisualization"
+
+void PCGDeltaVisualization::Helpers::AppendBaseDeltaColumns(TArray<FPCGDeltaVisualizerColumnInfo>& OutColumns)
+{
+	using namespace PCGDeltaVisualization::Constants;
+
+	OutColumns.Add({DeltaTypeId, LOCTEXT("DeltaTypeLabel", "Delta Type"), LOCTEXT("DeltaTypeTooltip", "Name of the delta type"), 100.f, EPCGTableVisualizerCellAlignment::Left});
+}
+
+FText PCGDeltaVisualization::Helpers::GetBaseDeltaCellText(const FName ColumnId, const FPCGDeltaKey& DeltaKey, const FConstStructView Delta)
+{
+	using namespace PCGDeltaVisualization::Constants;
+
+	if (ColumnId == DeltaTypeId)
+	{
+		if (const FPCGDeltaBase* DeltaBase = Delta.GetPtr<const FPCGDeltaBase>())
+		{
+			return FText::FromName(DeltaBase->GetDeltaName());
+		}
+	}
+
+	return FText::GetEmpty();
+}
+
+#undef LOCTEXT_NAMESPACE
diff --git a/Engine/Plugins/PCG/Source/PCGEditor/Private/DeltaVisualizations/PCGDeltaVisualizationHelpers.h b/Engine/Plugins/PCG/Source/PCGEditor/Private/DeltaVisualizations/PCGDeltaVisualizationHelpers.h
new file mode 100644
--- /dev/null
+++ b/Engine/Plugins/PCG/Source/PCGEditor/Private/DeltaVisualizations/PCGDeltaVisualizationHelpers.h
@@ -0,0 +1,16 @@
+// Copyright Epic Games, Inc. All Rights Reserved.
+
+#pragma once
+
+#include "DeltaVisualizations/PCGDeltaVisualization.h"
+
+struct FPCGDeltaKey;
+
+namespace PCGDeltaVisualization::Helpers
+{
+	/** Appends the common base delta columns (Delta Type). */
+	void AppendBaseDeltaColumns(TArray<FPCGDeltaVisualizerColumnInfo>& OutColumns);
+
+	/** Returns cell text for the common base delta columns. Returns empty text if the column is not handled. */
+	FText GetBaseDeltaCellText(FName ColumnId, const FPCGDeltaKey& DeltaKey, FConstStructView Delta);
+} // namespace PCGDeltaVisualization::Helpers
diff --git a/Engine/Plugins/PCG/Source/PCGEditor/Private/DeltaVisualizations/PCGDeltaVisualizationRegistry.cpp b/Engine/Plugins/PCG/Source/PCGEditor/Private/DeltaVisualizations/PCGDeltaVisualizationRegistry.cpp
new file mode 100644
--- /dev/null
+++ b/Engine/Plugins/PCG/Source/PCGEditor/Private/DeltaVisualizations/PCGDeltaVisualizationRegistry.cpp
@@ -0,0 +1,49 @@
+// Copyright Epic Games, Inc. All Rights Reserved.
+
+#include "DeltaVisualizations/PCGDeltaVisualizationRegistry.h"
+
+#include "Graph/DataOverride/PCGDataOverride.h"
+
+void FPCGDeltaVisualizationRegistry::RegisterDeltaVisualization(const UScriptStruct* DeltaStruct, TUniquePtr<const IPCGDeltaVisualization> Visualization)
+{
+	if (ensureMsgf(!ExternalRegistry.Find(DeltaStruct), TEXT("Cannot register multiple delta visualizations for the same struct type.")))
+	{
+		ExternalRegistry.Add(DeltaStruct, MoveTemp(Visualization));
+	}
+}
+
+void FPCGDeltaVisualizationRegistry::UnregisterDeltaVisualization(const UScriptStruct* DeltaStruct)
+{
+	ExternalRegistry.Remove(DeltaStruct);
+}
+
+const IPCGDeltaVisualization* FPCGDeltaVisualizationRegistry::GetDeltaVisualization(const UScriptStruct* DeltaStruct) const
+{
+	const UScriptStruct* CurrentStruct = DeltaStruct;
+	const UScriptStruct* DeltaBaseStruct = FPCGDeltaBase::StaticStruct();
+
+	// Walk the struct hierarchy from the concrete type up through FPCGDeltaBase (inclusive).
+	while (CurrentStruct)
+	{
+		// Always try external implementations first so users can override built-in visualizations.
+		if (const TUniquePtr<const IPCGDeltaVisualization>* ExternalVisualizationPtr = ExternalRegistry.Find(CurrentStruct))
+		{
+			return ExternalVisualizationPtr->Get();
+		}
+
+		if (const TUniquePtr<const IPCGDeltaVisualization>* InternalVisualizationPtr = InternalRegistry.Find(CurrentStruct))
+		{
+			return InternalVisualizationPtr->Get();
+		}
+
+		// Stop after checking FPCGDeltaBase.
+		if (CurrentStruct == DeltaBaseStruct)
+		{
+			break;
+		}
+
+		CurrentStruct = Cast<UScriptStruct>(CurrentStruct->GetSuperStruct());
+	}
+
+	return nullptr;
+}
diff --git a/Engine/Plugins/PCG/Source/PCGEditor/Private/Nodes/PCGEditorGraphNodeBase.cpp b/Engine/Plugins/PCG/Source/PCGEditor/Private/Nodes/PCGEditorGraphNodeBase.cpp
--- a/Engine/Plugins/PCG/Source/PCGEditor/Private/Nodes/PCGEditorGraphNodeBase.cpp
+++ b/Engine/Plugins/PCG/Source/PCGEditor/Private/Nodes/PCGEditorGraphNodeBase.cpp
@@ -873,6 +873,26 @@ EPCGChangeType UPCGEditorGraphNodeBase::UpdateGPUVisualization(const IPCGGraphEx
 	return ChangeType;
 }
 
+EPCGChangeType UPCGEditorGraphNodeBase::UpdateOverrideVisualization(const IPCGGraphExecutionSource* InSourceBeingDebugged, const FPCGStack* InStackBeingInspected)
+{
+	EPCGChangeType ChangeType = EPCGChangeType::None;
+
+	const UPCGNode* Node = GetPCGNode();
+
+	const bool bOverridesApplied = Node
+		&& InStackBeingInspected
+		&& InSourceBeingDebugged
+		&& InSourceBeingDebugged->GetExecutionState().GetInspection().NodeAppliedDataOverrides(Node, *InStackBeingInspected);
+
+	if (bOverridesApplied != GetAppliedDataOverrides())
+	{
+		SetAppliedDataOverrides(bOverridesApplied);
+		ChangeType |= EPCGChangeType::Cosmetic;
+	}
+
+	return ChangeType;
+}
+
 bool UPCGEditorGraphNodeBase::CanUserAddRemoveDynamicInputPins() const
 {
 	const UPCGSettings* Settings = PCGNode ? PCGNode->GetSettings() : nullptr;
@@ -1560,7 +1580,7 @@ void UPCGEditorGraphNodeBase::CreatePins(const TArray<UPCGPin*>& InInputPins, co
 
 		UEdGraphPin* Pin = CreatePin(EEdGraphPinDirection::EGPD_Input, GetPinType(InputPin), InputPin->Properties.Label);
 		Pin->PinFriendlyName = GetPinFriendlyName(InputPin);
-		Pin->bAdvancedView = InputPin->Properties.IsAdvancedPin();
+		Pin->bAdvancedView = InputPin->Properties.IsAdvancedPin() && !IsPinDefaultValueActivated(InputPin->Properties.Label);
 		bHasAdvancedPin |= Pin->bAdvancedView;
 	}
 
@@ -1573,7 +1593,7 @@ void UPCGEditorGraphNodeBase::CreatePins(const TArray<UPCGPin*>& InInputPins, co
 
 		UEdGraphPin* Pin = CreatePin(EEdGraphPinDirection::EGPD_Output, GetPinType(OutputPin), OutputPin->Properties.Label);
 		Pin->PinFriendlyName = GetPinFriendlyName(OutputPin);
-		Pin->bAdvancedView = OutputPin->Properties.IsAdvancedPin();
+		Pin->bAdvancedView = OutputPin->Properties.IsAdvancedPin() && !IsPinDefaultValueActivated(OutputPin->Properties.Label);
 		bHasAdvancedPin |= Pin->bAdvancedView;
 	}
 
diff --git a/Engine/Plugins/PCG/Source/PCGEditor/Private/Nodes/SPCGEditorGraphNode.cpp b/Engine/Plugins/PCG/Source/PCGEditor/Private/Nodes/SPCGEditorGraphNode.cpp
--- a/Engine/Plugins/PCG/Source/PCGEditor/Private/Nodes/SPCGEditorGraphNode.cpp
+++ b/Engine/Plugins/PCG/Source/PCGEditor/Private/Nodes/SPCGEditorGraphNode.cpp
@@ -399,12 +399,30 @@ void SPCGEditorGraphNode::GetOverlayBrushes(bool bSelected, const FVector2f& Wid
 		{
 			AddOverlayBrush(TEXT("PCG.NodeOverlay.Debug"));
 		}
+
+		if (const UPCGSettingsInterface* SettingsInterface = PCGNode->GetSettingsInterface())
+		{
+			if (SettingsInterface->bMarkedForManualEditing)
+			{
+				AddOverlayBrush(TEXT("PCG.NodeOverlay.MarkedForEditing"));
+			}
+
+			if (SettingsInterface->bTemporaryManualEditingEnabled)
+			{
+				AddOverlayBrush(TEXT("PCG.NodeOverlay.ManualEditing"));
+			}
+		}
 	}
 
 	if (PCGEditorGraphNode->GetInspected() && UsesInspectBrush())
 	{
 		AddOverlayBrush(TEXT("PCG.NodeOverlay.Inspect"));
 	}
+
+	if (PCGEditorGraphNode->GetAppliedDataOverrides())
+	{
+		AddOverlayBrush(TEXT("PCG.NodeOverlay.HasOverrides"));
+	}
 }
 
 void SPCGEditorGraphNode::OnNodeChanged()
diff --git a/Engine/Plugins/PCG/Source/PCGEditor/Private/PCGComponentVisualizer.cpp b/Engine/Plugins/PCG/Source/PCGEditor/Private/PCGComponentVisualizer.cpp
--- a/Engine/Plugins/PCG/Source/PCGEditor/Private/PCGComponentVisualizer.cpp
+++ b/Engine/Plugins/PCG/Source/PCGEditor/Private/PCGComponentVisualizer.cpp
@@ -3,89 +3,196 @@
 #include "PCGComponentVisualizer.h"
 
 #include "PCGComponent.h"
+#include "PCGDebug.h"
 #include "PCGEditorModule.h"
+#include "PCGGraph.h"
+#include "PCGNode.h"
+#include "Data/PCGBasePointData.h"
+#include "DeltaViewportExtensions/PCGDeltaViewportExtension.h"
+#include "DeltaViewportExtensions/PCGDeltaViewportExtensionHelpers.h"
 #include "Graph/DataOverride/PCGDataOverrideHelpers.h"
-#include "Graph/DataOverride/PCGDataOverridePoints.h"
 
 #include "CanvasItem.h"
 #include "CanvasTypes.h"
 #include "EditorModeManager.h"
 #include "EditorModes.h"
 #include "EditorViewportClient.h"
-#include "PCGGraph.h"
 #include "PrimitiveDrawInterface.h"
+#include "ScopedTransaction.h"
+#include "StaticMeshResources.h"
 #include "BaseGizmos/TransformGizmoUtil.h"
-#include "Data/PCGBasePointData.h"
 #include "Editor/UnrealEdEngine.h"
-#include "Materials/Material.h"
-#include "Materials/MaterialRenderProxy.h"
+#include "Engine/StaticMesh.h"
+#include "Modules/ModuleManager.h"
 #include "Tools/EdModeInteractiveToolsContext.h"
 
+#include "SPCGManualEditPanel.h"
+
 IMPLEMENT_HIT_PROXY(HPCGComponentVisualizerHitProxy, HComponentVisProxy);
-IMPLEMENT_HIT_PROXY(HPCGManualEditHitProxy, HComponentVisProxy);
 
 #define LOCTEXT_NAMESPACE "PCGComponentVisualizer"
 
 namespace PCG::ComponentVisualizer
 {
+	/** Result of looking up an element's delta state in a collection. */
+	struct FDeltaStateResult
+	{
+		FPCGDeltaKey Key = {};
+		IPCGDeltaViewportExtension* Extension = nullptr;
+		FConstStructView DeltaView;
+	};
+
+	/** Cached mesh geometry for wireframe drawing. Persists across frames on the visualizer. */
+	struct FCachedMeshWireframe
+	{
+		const UStaticMesh* Mesh = nullptr;
+		TArray<uint32> Indices;
+		const FPositionVertexBuffer* PositionBuffer = nullptr;
+
+		bool IsValid() const { return Mesh && PositionBuffer && !Indices.IsEmpty(); }
+
+		/** Returns the triangle count of the mesh's LOD 0, or 0 if the mesh has no valid render data. */
+		static int32 GetTriangleCount(const UStaticMesh* InMesh)
+		{
+			if (!InMesh)
+			{
+				return 0;
+			}
+
+			const FStaticMeshRenderData* RenderData = InMesh->GetRenderData();
+			if (!RenderData || RenderData->LODResources.IsEmpty())
+			{
+				return 0;
+			}
+
+			return RenderData->LODResources[0].GetNumTriangles();
+		}
+
+		bool TryCache(const UStaticMesh* InMesh)
+		{
+			if (InMesh == Mesh)
+			{
+				return IsValid();
+			}
+
+			Reset();
+			Mesh = InMesh;
+
+			if (!Mesh)
+			{
+				return false;
+			}
+
+			const FStaticMeshRenderData* RenderData = Mesh->GetRenderData();
+			if (!RenderData || RenderData->LODResources.IsEmpty())
+			{
+				return false;
+			}
+
+			const FStaticMeshLODResources& LOD = RenderData->LODResources[0];
+			if (LOD.IndexBuffer.GetNumIndices() == 0)
+			{
+				return false;
+			}
+
+			LOD.IndexBuffer.GetCopy(Indices);
+			PositionBuffer = &LOD.VertexBuffers.PositionVertexBuffer;
+			return IsValid();
+		}
+
+		void Reset()
+		{
+			Mesh = nullptr;
+			Indices.Empty();
+			PositionBuffer = nullptr;
+		}
+	};
+
 	namespace Defaults
 	{
-		const FVector ManualEditCubeExtent(100, 100, 50);
-		constexpr FLinearColor ManualEditDeltaCubeColor(1.0f, 0.0f, 0.5f, 0.4f);
-		constexpr FLinearColor ManualEditDefaultCubeColor(0.0f, 1.0f, 0.5f, 0.4f);
-		constexpr FLinearColor ManualEditSelectedCubeColor(1.0f, 1.0f, 0.0f, 0.4f);
+		constexpr FLinearColor ManualEditNoDeltaColor(0.0f, 0.55f, 0.45f, 0.3f);
+		constexpr FLinearColor ManualEditNoDeltaSelectedColor(0.4f, 1.0f, 0.9f, 0.5f);
 
-		constexpr const TCHAR* ManualEditOnText = TEXT("Manual Edit Elements : ON");
-		constexpr const TCHAR* ManualEditOffText = TEXT("Manual Edit Elements : OFF");
+		FLinearColor GetNoDeltaColor(const bool bIsSelected)
+		{
+			return bIsSelected ? ManualEditNoDeltaSelectedColor : ManualEditNoDeltaColor;
+		}
 	}
 
 	namespace Helpers
 	{
-		FPCGDeltaKey GetDeltaKeyFromCollection(const FPCGDeltaCollection& Collection, const FTransform& Transform)
+		/** Looks up an element's delta in the collection via the extension registry. */
+		FDeltaStateResult FindDeltaFromCollection(const FPCGDeltaCollection& Collection, const FTransform& Transform, const FPCGDeltaViewportExtensionRegistry& Registry)
 		{
-			FPCGDeltaKey LocationKey = PCG::DataOverride::Keys::FPCGLocationDeltaKey(Transform.GetLocation());
-			FPCGDeltaKey FoundDeltaKey = {};
-			Collection.ForEachDelta([&FoundDeltaKey, LocationKey](const FPCGDeltaKey& DeltaKey, const TInstancedStruct<FPCGDeltaBase>& DeltaValue)
+			FDeltaStateResult Result;
+
+			// @todo_pcg: Consider caching a DeltaKey -> Extension + View mapping to avoid resolving struct types on every lookup.
+			Collection.ForEachDelta([&Result, &Transform, &Registry, &Collection](const FPCGDeltaKey& DeltaKey, const TInstancedStruct<FPCGDeltaBase>& DeltaValue) -> bool
 			{
-				if (const FPCGPointEditDelta* EditDelta = DeltaValue.GetPtr<FPCGPointEditDelta>())
+				const UScriptStruct* ScriptStruct = DeltaValue.GetScriptStruct();
+				IPCGDeltaViewportExtension* Extension = Registry.GetExtension(ScriptStruct);
+				if (!Extension)
 				{
-					if (LocationKey == PCG::DataOverride::Keys::FPCGLocationDeltaKey(EditDelta->TransformOverride.GetLocation()))
-					{
-						FoundDeltaKey = DeltaKey;
-						return false;
-					}
+					return true;
+				}
+
+				// Create a read-only view so the extension can inspect the delta without a concrete type cast.
+				const FConstStructView View(ScriptStruct, DeltaValue.GetMemory());
+				if (Extension->MatchesSourceElement(Transform, View, Collection.Settings.SpatialTolerance))
+				{
+					Result.Key = DeltaKey;
+					Result.Extension = Extension;
+					Result.DeltaView = View;
+					return false;
 				}
+
 				return true;
 			});
 
-			return FoundDeltaKey;
+			return Result;
 		}
 
-		uint64 ComputePinKey(const FPCGStack& PinStack)
+		void DrawMeshWireframe(FPrimitiveDrawInterface* PDI, const FTransform& Transform, const FCachedMeshWireframe& CachedDebugMesh, const FLinearColor& Color)
 		{
-			// Right now, the pin key for a given stack should be computed with a node key first.
-			// @todo_pcg: This will need to change if/when the Data Override API changes.
-			check(PinStack.GetStackFrames().Num() > 1);
-			FPCGStack NodeStack = PinStack;
-			NodeStack.PopFrame();
-			const FXxHash64 NodeKey = DataOverride::Helpers::ComputeDataHash(NodeStack);
-			const FPCGStackFrame& PinFrame = PinStack.GetStackFrames().Last();
-			const UPCGPin* TopPin = PinFrame.GetObject_AnyThread<UPCGPin>();
-			check(TopPin);
-			return DataOverride::Helpers::ComputeDataHashWithPin(NodeKey, TopPin->Properties.Label).Hash;
+			for (int32 i = 0; i + 2 < CachedDebugMesh.Indices.Num(); i += 3)
+			{
+				const FVector3f V0 = CachedDebugMesh.PositionBuffer->VertexPosition(CachedDebugMesh.Indices[i]);
+				const FVector3f V1 = CachedDebugMesh.PositionBuffer->VertexPosition(CachedDebugMesh.Indices[i + 1]);
+				const FVector3f V2 = CachedDebugMesh.PositionBuffer->VertexPosition(CachedDebugMesh.Indices[i + 2]);
+
+				const FVector P0 = Transform.TransformPosition(FVector(V0));
+				const FVector P1 = Transform.TransformPosition(FVector(V1));
+				const FVector P2 = Transform.TransformPosition(FVector(V2));
+
+				PDI->DrawLine(P0, P1, Color, SDPG_World);
+				PDI->DrawLine(P1, P2, Color, SDPG_World);
+				PDI->DrawLine(P2, P0, Color, SDPG_World);
+			}
 		}
+
 	}
 } // PCG::ComponentVisualizer
 
-HPCGComponentVisualizerHitProxy::HPCGComponentVisualizerHitProxy(const UActorComponent* InComponent, const FPCGStack& InPinStack, int InDataIndex, int InPointIndex)
+// @todo_pcg: Revisit making this a user preference instead of a CVar
+static TAutoConsoleVariable<int32> CVarPCGManualEditMeshComplexityThreshold(
+	TEXT("pcg.Editor.ManualEdit.MeshComplexityThreshold"),
+	1000,
+	TEXT("Maximum triangle count for rendering custom debug meshes as wireframe in viewport editing. Meshes exceeding this fall back to the default point mesh."),
+	ECVF_Default);
+
+HPCGComponentVisualizerHitProxy::HPCGComponentVisualizerHitProxy(const UActorComponent* InComponent, const FPCGStack& InPinStack, const int32 InDataIndex, const int32 InPointIndex)
 	: HComponentVisProxy(InComponent)
 	, PinStack(InPinStack)
 	, DataIndex(InDataIndex)
 	, PointIndex(InPointIndex)
 {}
 
-HPCGManualEditHitProxy::HPCGManualEditHitProxy(const UActorComponent* InComponent)
+HPCGComponentVisualizerHitProxy::HPCGComponentVisualizerHitProxy(const UActorComponent* InComponent, const FPCGStack& InPinStack, const FPCGDeltaKey& InInsertedDeltaKey, const int32 InInsertedElementIndex)
 	: HComponentVisProxy(InComponent)
+	, PinStack(InPinStack)
+	, bIsInsertedElement(true)
+	, InsertedElementIndex(InInsertedElementIndex)
+	, InsertedDeltaKey(InInsertedDeltaKey)
 {}
 
 /** Define commands for the PCG component visualizer */
@@ -103,12 +210,12 @@ public:
 
 	virtual void RegisterCommands() override
 	{
-		UI_COMMAND(DeleteAllDeltas, "Remove All Deltas", "TODO.", EUserInterfaceActionType::Button, FInputChord());
-		UI_COMMAND(DeleteSelectedDeltas, "Remove Selected Deltas", "TODO.", EUserInterfaceActionType::Button, FInputChord());
+		UI_COMMAND(RemoveAllOverrides, "Remove All Overrides", "Removes all data overrides from the active collection.", EUserInterfaceActionType::Button, FInputChord());
+		UI_COMMAND(RemoveSelectedOverride, "Remove Selected Override", "Removes the currently selected data override.", EUserInterfaceActionType::Button, FInputChord());
 	}
 
-	TSharedPtr<FUICommandInfo> DeleteAllDeltas;
-	TSharedPtr<FUICommandInfo> DeleteSelectedDeltas;
+	TSharedPtr<FUICommandInfo> RemoveAllOverrides;
+	TSharedPtr<FUICommandInfo> RemoveSelectedOverride;
 };
 
 FPCGComponentVisualizer::FPCGComponentVisualizer()
@@ -116,6 +223,8 @@ FPCGComponentVisualizer::FPCGComponentVisualizer()
 	FPCGComponentVisualizerCommands::Register();
 
 	ComponentVisualizerActions = MakeShareable(new FUICommandList);
+
+	CachedWireframe = MakeUnique<PCG::ComponentVisualizer::FCachedMeshWireframe>();
 }
 
 
@@ -131,10 +240,10 @@ void FPCGComponentVisualizer::OnRegister()
 	const FPCGComponentVisualizerCommands& Commands = FPCGComponentVisualizerCommands::Get();
 
 	ComponentVisualizerActions->MapAction(
-		Commands.DeleteAllDeltas,
+		Commands.RemoveAllOverrides,
 		FExecuteAction::CreateSP(this, &FPCGComponentVisualizer::RemoveAllDeltas));
 	ComponentVisualizerActions->MapAction(
-		Commands.DeleteSelectedDeltas,
+		Commands.RemoveSelectedOverride,
 		FExecuteAction::CreateSP(this, &FPCGComponentVisualizer::RemoveSelectedDeltas),
 		FCanExecuteAction::CreateSP(this, &FPCGComponentVisualizer::HasSelection));
 }
@@ -146,90 +255,195 @@ void FPCGComponentVisualizer::DrawVisualization(const UActorComponent* Component
 
 	if (const UPCGComponent* PCGComp = Cast<const UPCGComponent>(Component))
 	{
-		if (bManualEditMode)
+		const FPCGSourceDataContainer* DataContainer = PCGComp->GetExecutionState().GetSourceDataContainer();
+		if (!DataContainer)
 		{
-			const FPCGGraphExecutionInspection& Inspection = PCGComp->GetExecutionState().GetInspection();
+			return;
+		}
+
+		// Not ideal to const_cast here, but we need to modify the component.
+		LastPCGComponent = const_cast<UPCGComponent*>(PCGComp);
+
+		FPCGEditorModule* EditorModule = FModuleManager::GetModulePtr<FPCGEditorModule>("PCGEditor");
+		const FPCGDeltaViewportExtensionRegistry& ExtensionRegistry = FPCGEditorModule::GetConstDeltaViewportExtensionRegistry();
 
-			const FPCGSourceDataContainer* DataContainer = PCGComp->GetExecutionState().GetSourceDataContainer();
-			check(DataContainer);
+		const FPCGGraphExecutionInspection& Inspection = PCGComp->GetExecutionState().GetInspection();
+		TMap<TObjectKey<const UPCGNode>, TSet<FPCGGraphExecutionInspection::FNodeExecutedNotificationData>> NodeStacks = Inspection.GetExecutedNodeStacks();
 
-			// Not ideal to const_cast here, but we need to modify the component.
-			LastPCGComponent = const_cast<UPCGComponent*>(PCGComp);
+		// Find the first node with viewport editing active (one-at-a-time enforcement).
+		// A node is active if toggled (bTemporaryManualEditingEnabled) or the panel selected it.
+		// @todo_pcg: GetExecutedNodeStacks only returns top-level nodes. Nodes inside subgraphs are not visited here.
+		ActiveEditingNode.Reset();
+		CachedDeltaContext.Reset();
 
-			TMap<TObjectKey<const UPCGNode>, TSet<FPCGGraphExecutionInspection::FNodeExecutedNotificationData>> NodeStacks = Inspection.GetExecutedNodeStacks();
+		const UPCGNode* PanelActiveNode = nullptr;
+		if (EditorModule)
+		{
+			if (TSharedPtr<SPCGManualEditPanel> Panel = EditorModule->GetManualEditPanel())
+			{
+				PanelActiveNode = Panel->GetActiveNode();
+			}
+		}
 
-			// Find nodes with manual edit.
-			for (const auto& [ObjectKey, StackSet] : NodeStacks)
+		for (const TPair<TObjectKey<const UPCGNode>, TSet<FPCGGraphExecutionInspection::FNodeExecutedNotificationData>>& NodeStackPair : NodeStacks)
+		{
+			const UPCGNode* Node = NodeStackPair.Key.ResolveObjectPtr();
+			if (!Node || !Node->GetSettingsInterface())
 			{
-				const UPCGNode* Node = ObjectKey.ResolveObjectPtr();
-				if (!Node || !Node->GetSettingsInterface() || !Node->GetSettingsInterface()->bViewportEditingEnabled)
-				{
-					continue;
-				}
+				continue;
+			}
 
-				FPCGElementPtr Element = Node->GetSettings() ? Node->GetSettings()->GetElement() : nullptr;
-				if (!Element)
-				{
-					continue;
-				}
+			const UPCGSettingsInterface* NodeSettings = Node->GetSettingsInterface();
+			if (!NodeSettings->bTemporaryManualEditingEnabled && Node != PanelActiveNode)
+			{
+				continue;
+			}
 
-				// Only override pre or post execute.
-				EPCGDataOverridePhase Phase = Element->GetDataOverridePhase();
-				if (Phase != EPCGDataOverridePhase::PrepareData && Phase != EPCGDataOverridePhase::PostExecute)
+			FPCGElementPtr Element = Node->GetSettings() ? Node->GetSettings()->GetElement() : nullptr;
+			if (!Element)
+			{
+				continue;
+			}
+
+			// Only override pre or post execute.
+			EPCGDataOverridePhase Phase = Element->GetDataOverridePhase();
+			if (Phase != EPCGDataOverridePhase::PrepareData && Phase != EPCGDataOverridePhase::PostExecute)
+			{
+				continue;
+			}
+
+			// Const cast needed because the inspection map stores const keys.
+			ActiveEditingNode = const_cast<UPCGNode*>(Node);
+
+			// Resolve and cache the debug mesh wireframe geometry for this node.
+			// @todo_pcg: Allow selecting the debug mesh via a attribute selector for the selection proxy.
+			const FPCGDebugVisualizationSettings& DebugSettings = Node->GetSettingsInterface()->DebugSettings;
+			const UStaticMesh* DebugMesh = DebugSettings.PointMesh.Get();
+			if (!DebugMesh)
+			{
+				if (!CachedDefaultPointMesh.IsValid())
 				{
-					continue;
+					CachedDefaultPointMesh = Cast<UStaticMesh>(TSoftObjectPtr<UStaticMesh>(PCGDebugVisConstants::DefaultPointMesh).LoadSynchronous());
 				}
+				DebugMesh = CachedDefaultPointMesh.Get();
+			}
+
+			const int32 ComplexityThreshold = CVarPCGManualEditMeshComplexityThreshold.GetValueOnGameThread();
+			if (ComplexityThreshold > 0 && FCachedMeshWireframe::GetTriangleCount(DebugMesh) > ComplexityThreshold)
+			{
+				DebugMesh = CachedDefaultPointMesh.Get();
+			}
 
-				const TArray<TObjectPtr<UPCGPin>>& Pins = (Phase == EPCGDataOverridePhase::PrepareData) ? Node->GetInputPins() : Node->GetOutputPins();
-				if (Pins.Num() > 0)
+			if (!CachedWireframe->TryCache(DebugMesh))
+			{
+				break;
+			}
+
+			const TArray<TObjectPtr<UPCGPin>>& Pins = (Phase == EPCGDataOverridePhase::PrepareData) ? Node->GetInputPins() : Node->GetOutputPins();
+			if (Pins.Num() > 0)
+			{
+				for (const FPCGGraphExecutionInspection::FNodeExecutedNotificationData& Stack : NodeStackPair.Value)
 				{
-					for (const FPCGGraphExecutionInspection::FNodeExecutedNotificationData& Stack : StackSet)
+					const UPCGPin* PinUsed = Pins[0];
+					const FName PinLabel = PinUsed->Properties.Label;
+
+					FPCGSourceDataStorageKey StorageKey(PCG::DataOverride::Constants::DefaultOverrideLabel, PCG::DataOverride::Helpers::ComputePinOverrideKey(Stack.Stack, Node, PinLabel).Hash);
+
+					// Push context and selection state to the panel
+					if (EditorModule)
 					{
-						const UPCGPin* PinUsed = Pins[0];
+						if (TSharedPtr<SPCGManualEditPanel> Panel = EditorModule->GetManualEditPanel())
+						{
+							Panel->SetDeltaContext(LastPCGComponent.Get(), StorageKey);
+							if (EditTarget.IsSet())
+							{
+								Panel->SetSelectionState(true, EditTarget->DeltaKey, EditTarget->Transform, EditTarget->SelectedElementIndex, EditTarget->OriginalPointIndex);
+							}
+							else
+							{
+								Panel->SetSelectionState(false, FPCGDeltaKey{});
+							}
+						}
+					}
 
-						FPCGStack PinStack = Stack.Stack;
-						PinStack.PushFrame(Node);
-						PinStack.PushFrame(PinUsed);
-						TSharedPtr<FPCGInspectionData> InspectionData = Inspection.GetInspectionDataPtr(PinStack);
+					FPCGStack PinStack = Stack.Stack;
+					PinStack.PushFrame(Node);
+					PinStack.PushFrame(PinUsed);
 
-						// Compute the node key from the execution stack
-						FPCGSourceDataStorageKey StorageKey(PCG::DataOverride::Constants::DefaultOverrideLabel, Helpers::ComputePinKey(PinStack));
+					// Cache the collection lookup for action methods.
+					CachedDeltaContext.DataContainer = LastPCGComponent->GetExecutionState().GetSourceDataContainer();
+					CachedDeltaContext.StorageKey = StorageKey;
+					CachedDeltaContext.PinStack = PinStack;
 
-						// Retrieve the delta collection
-						FConstSharedStruct SharedStruct = DataContainer->Get<FPCGDeltaCollection>(StorageKey);
-						const FPCGDeltaCollection* Collection = SharedStruct.IsValid()
-																	? SharedStruct.GetPtr<const FPCGDeltaCollection>()
-																	: nullptr;
+					// Retrieve the delta collection
+					FConstSharedStruct SharedStruct = DataContainer->Get<FPCGDeltaCollection>(StorageKey);
+					const FPCGDeltaCollection* Collection = SharedStruct.IsValid() ? SharedStruct.GetPtr<const FPCGDeltaCollection>() : nullptr;
+					TSharedPtr<FPCGInspectionData> InspectionData = Inspection.GetInspectionDataPtr(PinStack);
 
-						if (InspectionData && InspectionData->Data)
+					if (InspectionData && InspectionData->Data)
+					{
+						const TArray<FPCGTaggedData>& DataArray = InspectionData->Data->GetAllInputs();
+						for (int DataIndex = 0; DataIndex < DataArray.Num(); DataIndex++)
 						{
-							const TArray<FPCGTaggedData>& DataArray = InspectionData->Data->GetAllInputs();
-							for (int DataIndex = 0; DataIndex < DataArray.Num(); DataIndex++)
+							const FPCGTaggedData& Input = DataArray[DataIndex];
+							// @todo_pcg: For now only point data is supported. On a future pass, this should be expanded to include other data types.
+							if (const UPCGBasePointData* BasePointData = Cast<UPCGBasePointData>(Input.Data))
 							{
-								const FPCGTaggedData& Input = DataArray[DataIndex];
-								if (const UPCGBasePointData* BasePointData = Cast<UPCGBasePointData>(Input.Data))
+								TConstPCGValueRange<FTransform> TransformRange = BasePointData->GetConstTransformValueRange();
+								for (int Index = 0; Index < BasePointData->GetNumPoints(); Index++)
+								{
+									const FTransform& SourceElementTransform = TransformRange[Index];
+
+									// Delta-type-aware lookup via extension registry
+									FDeltaStateResult StateResult;
+									if (Collection)
+									{
+										StateResult = Helpers::FindDeltaFromCollection(*Collection, SourceElementTransform, ExtensionRegistry);
+									}
+
+									const bool bHasDelta = StateResult.Key.Hash != 0;
+									const bool bIsSelectedDelta = bHasDelta && EditTarget && EditTarget->DeltaKey == StateResult.Key && EditTarget->Stack == PinStack;
+									const bool bIsSelectedSource = !bHasDelta && EditTarget && EditTarget->OriginalPointIndex == Index && EditTarget->Stack == PinStack;
+									const bool bIsSelected = bIsSelectedDelta || bIsSelectedSource;
+
+									// Extension provides the display transform (e.g. deletions show at original location)
+									const FTransform& Transform = StateResult.Extension ? StateResult.Extension->GetDisplayTransform(SourceElementTransform, StateResult.DeltaView) : SourceElementTransform;
+
+									FLinearColor ProxyColor = StateResult.Extension ? StateResult.Extension->GetDisplayColor(bIsSelected) : GetNoDeltaColor(bIsSelected);
+
+									PDI->SetHitProxy(new HPCGComponentVisualizerHitProxy(Component, PinStack, DataIndex, Index));
+									Helpers::DrawMeshWireframe(PDI, Transform, *CachedWireframe, ProxyColor);
+									PDI->SetHitProxy(nullptr);
+								}
+
+								// Render inserted elements provided by extensions
+								if (Collection)
 								{
-									TConstPCGValueRange<FTransform> TransformRange = BasePointData->GetConstTransformValueRange();
-									for (int Index = 0; Index < BasePointData->GetNumPoints(); Index++)
+									const FPCGDeltaViewportExtensionRegistry& Registry = ExtensionRegistry;
+									TArray<FInsertedViewportElement> InsertedElements;
+									for (const UScriptStruct* DeltaType : Registry.GetRegisteredDeltaTypes())
 									{
-										const FTransform& Transform = TransformRange[Index];
-										bool bIsDelta = false;
-										bool bIsSelected = false;
-										if (Collection)
+										IPCGDeltaViewportExtension* Extension = Registry.GetExtension(DeltaType);
+										if (!Extension)
 										{
-											FPCGDeltaKey Key = Helpers::GetDeltaKeyFromCollection(*Collection, Transform);
-											bIsDelta = Key.Hash != 0;
-											bIsSelected = bIsDelta && EditTarget && EditTarget->DeltaKey == Key && EditTarget->Stack == PinStack;
+											continue;
 										}
 
-										FLinearColor ProxyColor = bIsSelected ? ManualEditSelectedCubeColor : bIsDelta ? ManualEditDeltaCubeColor : ManualEditDefaultCubeColor;
+										InsertedElements.Reset();
+										Extension->GetInsertedElements(*Collection, InsertedElements);
+
+										for (const FInsertedViewportElement& Inserted : InsertedElements)
+										{
+											const bool bIsInsertedSelected = EditTarget && EditTarget->ActiveExtension == Extension
+												&& EditTarget->DeltaKey == Inserted.DeltaKey
+												&& EditTarget->SelectedElementIndex == Inserted.ElementIndex;
+
+											FLinearColor InsertedColor = Extension->GetDisplayColor(bIsInsertedSelected);
 
-										PDI->SetHitProxy(new HPCGComponentVisualizerHitProxy(Component, PinStack, DataIndex, Index));
-										DrawWireBox(PDI, Transform.ToMatrixWithScale(), FBox(-ManualEditCubeExtent, ManualEditCubeExtent), ProxyColor, SDPG_World);
-										FDynamicColoredMaterialRenderProxy* ColoredMaterialRenderProxy = new FDynamicColoredMaterialRenderProxy(GEngine->DebugEditorMaterial->GetRenderProxy(), ProxyColor);
-										PDI->RegisterDynamicResource(ColoredMaterialRenderProxy);
-										DrawBox(PDI, Transform.ToMatrixWithScale(), ManualEditCubeExtent, ColoredMaterialRenderProxy, SDPG_World);
-										PDI->SetHitProxy(nullptr);
+											PDI->SetHitProxy(new HPCGComponentVisualizerHitProxy(Component, PinStack, Inserted.DeltaKey, Inserted.ElementIndex));
+											Helpers::DrawMeshWireframe(PDI, Inserted.Transform, *CachedWireframe, InsertedColor);
+											PDI->SetHitProxy(nullptr);
+										}
 									}
 								}
 							}
@@ -237,12 +451,23 @@ void FPCGComponentVisualizer::DrawVisualization(const UActorComponent* Component
 					}
 				}
 			}
+
+			// Only process the first eligible node
+			break;
 		}
-		// while dragging, selection is not visited with the previous loop. So we draw it here so we can have visual feedback on dragging.
-		// No need to add proxy since
-		if (EditTarget)
+
+		// @todo_pcg: Gizmo sync with delta data for external changes (e.g. graph re-execution) is not yet implemented.
+		// The sync block needs a dedicated virtual method to read the current transform from the delta struct,
+		// since GetDisplayTransform is designed for the drawing loop and returns the wrong value here.
+
+		// While dragging, the selected element is not visited by the main loop. Draw it here for visual feedback.
+		if (EditTarget && EditTarget->bGizmoDragging && ActiveEditingNode.IsValid() && CachedWireframe->IsValid())
 		{
-			DrawWireBox(PDI, EditTarget->Transform.ToMatrixWithScale(), FBox(-ManualEditCubeExtent, ManualEditCubeExtent), ManualEditSelectedCubeColor, SDPG_World);
+			const FLinearColor DragColor = EditTarget->ActiveExtension
+				? EditTarget->ActiveExtension->GetDisplayColor(/*bIsSelected=*/true)
+				: GetNoDeltaColor(/*bIsSelected=*/true);
+
+			Helpers::DrawMeshWireframe(PDI, EditTarget->Transform, *CachedWireframe, DragColor);
 		}
 	}
 	FComponentVisualizer::DrawVisualization(Component, View, PDI);
@@ -251,34 +476,6 @@ void FPCGComponentVisualizer::DrawVisualization(const UActorComponent* Component
 void FPCGComponentVisualizer::DrawVisualizationHUD(const UActorComponent* Component, const FViewport* Viewport, const FSceneView* View,
 												   FCanvas* Canvas)
 {
-	using namespace PCG::ComponentVisualizer;
-
-	if (GEngine)
-	{
-		Canvas->SetHitProxy(new HPCGManualEditHitProxy(Component));
-		const TCHAR* ManualEditText = bManualEditMode ? Defaults::ManualEditOnText : Defaults::ManualEditOffText;
-		int32 XL;
-		int32 YL;
-		const UFont* Font = GEngine->GetLargeFont();
-		StringSize(Font, XL, YL, ManualEditText);
-
-		// @todo_pcg: There seems to be a bug in the canvas hit proxy. it only works at position 0,0.
-		constexpr float DrawPositionX = 0;
-		constexpr float DrawPositionY = 0;
-		// const FIntRect CanvasRect = Canvas->GetViewRect();
-		// const float DrawPositionX = FMath::FloorToFloat(static_cast<float>(CanvasRect.Min.X) + static_cast<float>(CanvasRect.Width() - XL) * 0.5f);
-		// const float DrawPositionY = static_cast<float>(CanvasRect.Min.Y) + 10;
-		FVector2D DrawPosition{DrawPositionX, DrawPositionY};
-
-		FCanvasTileItem Background(DrawPosition, FVector2D{static_cast<float>(XL), static_cast<float>(YL)}, FLinearColor::Yellow);
-		Canvas->DrawItem(Background, DrawPosition);
-
-		FCanvasTextItem TextItem(FVector2D{DrawPositionX, DrawPositionY}, FText::FromString(ManualEditText), Font, FLinearColor::White);
-		TextItem.EnableShadow(FLinearColor::Black);
-		Canvas->DrawItem(TextItem, DrawPosition);
-		Canvas->SetHitProxy(nullptr);
-	}
-
 	FComponentVisualizer::DrawVisualizationHUD(Component, Viewport, View, Canvas);
 }
 
@@ -288,17 +485,99 @@ bool FPCGComponentVisualizer::VisProxyHandleClick(FEditorViewportClient* InViewp
 
 	// We need to register the gizmo tool with the current mode tool so the gizmo will work.
 	//  Has to be done here since we need an interactive tool context.
-	UE::TransformGizmoUtil::RegisterTransformGizmoContextObject(InViewportClient->GetModeTools()->GetInteractiveToolsContext());
-
-	if (HitProxyCast<HPCGManualEditHitProxy>(VisProxy))
+	if (FEditorModeTools* ModeTools = InViewportClient->GetModeTools())
 	{
-		bManualEditMode = !bManualEditMode;
-		return true;
+		UE::TransformGizmoUtil::RegisterTransformGizmoContextObject(ModeTools->GetInteractiveToolsContext());
 	}
-	else if (HPCGComponentVisualizerHitProxy* HitProxy = HitProxyCast<HPCGComponentVisualizerHitProxy>(VisProxy))
+
+	FPCGEditorModule* EditorModule = FModuleManager::GetModulePtr<FPCGEditorModule>("PCGEditor");
+
+	if (HPCGComponentVisualizerHitProxy* HitProxy = HitProxyCast<HPCGComponentVisualizerHitProxy>(VisProxy))
 	{
-		const UPCGComponent* PCGComponent = Cast<UPCGComponent>(HitProxy->Component);
-		const FPCGGraphExecutionInspection& Inspection = PCGComponent->GetExecutionState().GetInspection();
+		UPCGComponent* MutablePCGComponent = const_cast<UPCGComponent*>(CastChecked<const UPCGComponent>(HitProxy->Component));
+
+		FPCGSourceDataContainer* DataContainer = MutablePCGComponent->GetExecutionState().GetSourceDataContainer();
+		check(DataContainer);
+
+		UModeManagerInteractiveToolsContext* ToolsContext = InViewportClient->GetModeTools() ? InViewportClient->GetModeTools()->GetInteractiveToolsContext() : nullptr;
+		UInteractiveGizmoManager* GizmoManager = ToolsContext ? ToolsContext->GizmoManager.Get() : nullptr;
+		if (!GizmoManager)
+		{
+			return false;
+		}
+
+		// Extract node and pin from the stored PinStack (execution stack + Node frame + Pin frame).
+		const TArray<FPCGStackFrame>& HitFrames = HitProxy->PinStack.GetStackFrames();
+		check(HitFrames.Num() >= 2);
+		const UPCGNode* HitNode = HitFrames[HitFrames.Num() - 2].GetObject_AnyThread<UPCGNode>();
+		const UPCGPin* HitPin = HitFrames.Last().GetObject_AnyThread<UPCGPin>();
+		if (!HitNode || !HitPin)
+		{
+			return false;
+		}
+
+		const FName HitPinLabel = HitPin->Properties.Label;
+		FPCGStack HitExecutionStack = HitProxy->PinStack;
+		HitExecutionStack.PopFrame(); // Pin
+		HitExecutionStack.PopFrame(); // Node
+		FPCGSourceDataStorageKey StorageKey(PCG::DataOverride::Constants::DefaultOverrideLabel, PCG::DataOverride::Helpers::ComputePinOverrideKey(HitExecutionStack, HitNode, HitPinLabel).Hash);
+
+		// @todo_pcg: Delegate click handling to extensions to decouple behavior from the visualizer.
+		// --- Handle inserted element click ---
+		if (HitProxy->bIsInsertedElement)
+		{
+			FSharedStruct SharedStruct = DataContainer->GetMutable<FPCGDeltaCollection>(StorageKey);
+			FPCGDeltaCollection* Collection = SharedStruct.IsValid() ? SharedStruct.GetPtr<FPCGDeltaCollection>() : nullptr;
+			if (!Collection)
+			{
+				return false;
+			}
+
+			TInstancedStruct<FPCGDeltaBase>* DeltaInstancedStruct = Collection->Find(HitProxy->InsertedDeltaKey);
+			if (!DeltaInstancedStruct)
+			{
+				return false;
+			}
+
+			const UScriptStruct* DeltaStructType = DeltaInstancedStruct->GetScriptStruct();
+			IPCGDeltaViewportExtension* Extension = FPCGEditorModule::GetMutableDeltaViewportExtensionRegistry().GetExtension(DeltaStructType);
+			if (!Extension)
+			{
+				return false;
+			}
+
+			FTransform InsertedTransform = FTransform::Identity;
+			const int32 TargetElementIndex = HitProxy->InsertedElementIndex;
+
+			TArray<FInsertedViewportElement> InsertedElements;
+			Extension->GetInsertedElements(*Collection, InsertedElements);
+			for (const FInsertedViewportElement& Inserted : InsertedElements)
+			{
+				if (Inserted.DeltaKey == HitProxy->InsertedDeltaKey && Inserted.ElementIndex == TargetElementIndex)
+				{
+					InsertedTransform = Inserted.Transform;
+					break;
+				}
+			}
+
+			EditTarget.Emplace(MutablePCGComponent, HitProxy->PinStack, StorageKey, HitProxy->InsertedDeltaKey, /*bWasInDelta=*/true, InsertedTransform);
+			EditTarget->OriginalTransform = InsertedTransform;
+			EditTarget->ActiveExtension = Extension;
+			EditTarget->DeltaStructType = DeltaStructType;
+			EditTarget->SelectedElementIndex = HitProxy->InsertedElementIndex;
+
+			ResetGizmo();
+			if (Extension->UsesTRSGizmo())
+			{
+				SetupGizmo(GizmoManager, InsertedTransform);
+			}
+
+			RefreshKeyBindings(Extension);
+			return true;
+		}
+
+		// --- Handle regular point click ---
+		const FPCGGraphExecutionInspection& Inspection = MutablePCGComponent->GetExecutionState().GetInspection();
 		TSharedPtr<FPCGInspectionData> InspectionData = Inspection.GetInspectionDataPtr(HitProxy->PinStack);
 		if (!InspectionData || !InspectionData->Data)
 		{
@@ -311,131 +590,303 @@ bool FPCGComponentVisualizer::VisProxyHandleClick(FEditorViewportClient* InViewp
 			return false;
 		}
 
+		// @todo_pcg: For now only point data is supported. On a future pass, this should be expanded to include other data types.
 		const FPCGTaggedData& Input = DataArray[HitProxy->DataIndex];
 		if (const UPCGBasePointData* BasePointData = Cast<UPCGBasePointData>(Input.Data))
 		{
-			if (UInteractiveGizmoManager* GizmoManager = InViewportClient->GetModeTools()->GetInteractiveToolsContext()->GizmoManager)
+			FPCGPointTransform::ConstValueRange TransformRange = BasePointData->GetConstTransformValueRange();
+			if (!ensure(TransformRange.IsValidIndex(HitProxy->PointIndex)))
 			{
-				// Not ideal to const_cast here, but we need to modify the PCGComp.
-				UPCGComponent* MutablePCGComponent = const_cast<UPCGComponent*>(PCGComponent);
+				return false;
+			}
 
-				FPCGSourceDataContainer* DataContainer = MutablePCGComponent->GetExecutionState().GetSourceDataContainer();
-				check(DataContainer);
+			const FTransform& CurrentTransform = TransformRange[HitProxy->PointIndex];
+			FPCGDeltaCollection* Collection = nullptr;
+			FPCGDeltaKey DeltaKey = {};
+			bool bWasInDelta = false;
+			FDeltaStateResult StateResult;
 
-				FPCGSourceDataStorageKey StorageKey(PCG::DataOverride::Constants::DefaultOverrideLabel, Helpers::ComputePinKey(HitProxy->PinStack));
+			// Check if the point is already a product of a delta.
+			{
+				FSharedStruct SharedStruct = DataContainer->GetMutable<FPCGDeltaCollection>(StorageKey);
+				Collection = SharedStruct.IsValid() ? SharedStruct.GetPtr<FPCGDeltaCollection>() : nullptr;
+				if (Collection)
+				{
+					StateResult = Helpers::FindDeltaFromCollection(*Collection, CurrentTransform, FPCGEditorModule::GetConstDeltaViewportExtensionRegistry());
+					DeltaKey = StateResult.Key;
+					bWasInDelta = (DeltaKey.Hash != 0);
+				}
+			}
 
-				FPCGPointTransform::ConstValueRange TransformRange = BasePointData->GetConstTransformValueRange();
-				if (!ensure(TransformRange.IsValidIndex(HitProxy->PointIndex)))
+			// Duplicate the clicked element as an inserted point on Alt+click
+			if (Click.IsAltDown())
+			{
+				if (!Collection)
 				{
-					return false;
+					FPCGDeltaCollection NewCollection;
+					NewCollection.Settings.KeyPolicy = EPCGDataOverrideKeyPolicy::Spatial;
+					NewCollection.Settings.KeyTarget = EPCGSpatialKeyTarget::Position;
+					NewCollection.Settings.SpatialTolerance = PCG::DataOverride::Constants::SpatialToleranceDefault;
+					DataContainer->Store(StorageKey, NewCollection);
+
+					FSharedStruct SharedStruct = DataContainer->GetMutable<FPCGDeltaCollection>(StorageKey);
+					Collection = SharedStruct.IsValid() ? SharedStruct.GetPtr<FPCGDeltaCollection>() : nullptr;
 				}
+				check(Collection);
 
-				const FTransform& CurTransform = TransformRange[HitProxy->PointIndex];
-				FPCGDeltaCollection* Collection = nullptr;
-				FPCGDeltaKey DeltaKey = {};
-				bool bWasInDelta = false;
+				// Use the original transform for the duplicate (not post-apply zero-scale for deletions)
+				const FTransform& SourceTransform = (StateResult.Extension && StateResult.DeltaView.IsValid())
+					? StateResult.Extension->GetDisplayTransform(CurrentTransform, StateResult.DeltaView)
+					: CurrentTransform;
 
-				// Check if the point is already a product of a delta.
+				// Delegate to the active extension's OnAltDrag handler
+				IPCGDeltaViewportExtension* AltDragExtension = nullptr;
+				if (EditorModule)
 				{
-					FSharedStruct SharedStruct = DataContainer->GetMutable<FPCGDeltaCollection>(StorageKey);
-					Collection = SharedStruct.IsValid()
-									 ? SharedStruct.GetPtr<FPCGDeltaCollection>()
-									 : nullptr;
-					if (Collection)
+					if (TSharedPtr<SPCGManualEditPanel> Panel = EditorModule->GetManualEditPanel())
 					{
-						DeltaKey = Helpers::GetDeltaKeyFromCollection(*Collection, CurTransform);
-						bWasInDelta = (DeltaKey.Hash != 0);
+						if (TSharedPtr<const FPCGManualEditNodeConfiguration> NodeConfiguration = Panel->GetConfigurationForActiveNode())
+						{
+							AltDragExtension = FPCGEditorModule::GetMutableDeltaViewportExtensionRegistry().GetExtension(NodeConfiguration->DeltaType);
+						}
 					}
 				}
 
-				// It was not part of a delta...
-				if (!bWasInDelta)
+				if (!AltDragExtension)
 				{
-					// Create collection if there was not already one.
-					if (!Collection)
-					{
-						// Build a delta collection with spatial key settings
-						FPCGDeltaCollection NewCollection;
-						NewCollection.Settings.KeyPolicy = EPCGDataOverrideKeyPolicy::Spatial;
-						NewCollection.Settings.KeyTarget = EPCGSpatialKeyTarget::Position;
-						NewCollection.Settings.SpatialTolerance = 10.0;
-
-						DataContainer->Store(StorageKey, NewCollection);
-
-						FSharedStruct SharedStruct = DataContainer->GetMutable<FPCGDeltaCollection>(StorageKey);
-						Collection = SharedStruct.IsValid()
-										 ? SharedStruct.GetPtr<FPCGDeltaCollection>()
-										 : nullptr;
-					}
+					return false;
+				}
 
-					// Add current point in the delta, with the same transform as the point (it will be modified by the user later)
-					check(Collection);
-					FPCGPointEditDelta EditDelta;
-					EditDelta.OriginalTransform = CurTransform;
-					EditDelta.TransformOverride = CurTransform;
-					EditDelta.Bounds = PCGComponent->GetTotalBounds();
+				FScopedTransaction Transaction(LOCTEXT("InsertPointDelta", "PCG Manual Edit: Insert Point"));
+				MutablePCGComponent->Modify();
 
-					// Key the delta by spatial position and delta type label
-					DeltaKey = PCG::DataOverride::Keys::FPCGLocationDeltaKey(
-						CurTransform.GetLocation(),
-						Collection->Settings.SpatialTolerance,
-						FPCGPointEditDelta::GetDeltaNameStatic());
+				FPCGDeltaKey InsertionDeltaKey;
+				int32 NewElementIndex = AltDragExtension->OnAltDrag(*Collection, SourceTransform, InsertionDeltaKey);
+				if (NewElementIndex == INDEX_NONE)
+				{
+					return false;
+				}
+
+				EditTarget.Emplace(MutablePCGComponent, HitProxy->PinStack, StorageKey, InsertionDeltaKey, /*bWasInDelta=*/true, SourceTransform);
+				EditTarget->OriginalTransform = SourceTransform;
+				EditTarget->ActiveExtension = AltDragExtension;
+				const TInstancedStruct<FPCGDeltaBase>* FoundDelta = Collection->Find(InsertionDeltaKey);
+				EditTarget->DeltaStructType = FoundDelta ? FoundDelta->GetScriptStruct() : nullptr;
+				EditTarget->SelectedElementIndex = NewElementIndex;
+				EditTarget->bModified = true;
+
+				ResetGizmo();
 
-					Collection->Add(DeltaKey, TInstancedStruct<FPCGPointEditDelta>::Make(EditDelta));
+				if (AltDragExtension->UsesTRSGizmo())
+				{
+					SetupGizmo(GizmoManager, SourceTransform);
 				}
 
-				//Now that the point is part of the delta, set up the current selection
-				EditTarget.Emplace(MutablePCGComponent, HitProxy->PinStack, StorageKey, DeltaKey, bWasInDelta, CurTransform);
+				PCG::DeltaViewportExtension::MarkComponentDirty(MutablePCGComponent, DataContainer);
+
+				RefreshKeyBindings(AltDragExtension);
+
+				return true;
+			}
+
+			// Clicking an element with an existing delta that doesn't use a TRS gizmo
+			if (bWasInDelta && StateResult.Extension && !StateResult.Extension->UsesTRSGizmo())
+			{
+				const FTransform DisplayTransform = StateResult.Extension->GetDisplayTransform(CurrentTransform, StateResult.DeltaView);
+				EditTarget.Emplace(MutablePCGComponent, HitProxy->PinStack, StorageKey, DeltaKey, /*bWasInDelta=*/true, DisplayTransform);
+				EditTarget->OriginalTransform = DisplayTransform;
+				EditTarget->ActiveExtension = StateResult.Extension;
+				EditTarget->DeltaStructType = StateResult.DeltaView.GetScriptStruct();
+				EditTarget->OriginalPointIndex = HitProxy->PointIndex;
 
 				ResetGizmo();
-				ETransformGizmoSubElements TransformElements = ETransformGizmoSubElements::FullTranslateRotateScale;
-				Gizmo = UE::TransformGizmoUtil::CreateCustomTransformGizmo(GizmoManager, TransformElements, this, FString());
-				UTransformProxy* Proxy = NewObject<UTransformProxy>(Gizmo);
-				Proxy->SetTransform(CurTransform);
-				Gizmo->SetActiveTarget(Proxy, GizmoManager);
-				Gizmo->ReinitializeGizmoTransform(CurTransform);
-				Gizmo->SetVisibility(true);
+				RefreshKeyBindings(StateResult.Extension);
+
+				return true;
+			}
 
+			// Resolve extension for deferred delta creation
+			IPCGDeltaViewportExtension* DeferredExtension = nullptr;
+			const UScriptStruct* DeferredDeltaType = nullptr;
 
-				Proxy->OnTransformChanged.AddLambda([this](UTransformProxy* InProxy, const FTransform& NewTransform)
+			// It was not part of a delta...
+			if (!bWasInDelta)
+			{
+				// Create collection if there was not already one.
+				if (!Collection)
 				{
-					SetSelectedTransform(NewTransform);
-				});
+					FPCGDeltaCollection NewCollection;
+					NewCollection.Settings.KeyPolicy = EPCGDataOverrideKeyPolicy::Spatial;
+					NewCollection.Settings.KeyTarget = EPCGSpatialKeyTarget::Position;
+					NewCollection.Settings.SpatialTolerance = PCG::DataOverride::Constants::SpatialToleranceDefault;
 
-				return true;
+					DataContainer->Store(StorageKey, NewCollection);
+
+					FSharedStruct SharedStruct = DataContainer->GetMutable<FPCGDeltaCollection>(StorageKey);
+					Collection = SharedStruct.IsValid() ? SharedStruct.GetPtr<FPCGDeltaCollection>() : nullptr;
+				}
+
+				check(Collection);
+
+				// Defer delta creation until an explicit user action (gizmo drag, Delete key, Insert button).
+				// Resolve the extension from the panel's selected delta type for key computation.
+				FName DeltaName = NAME_None;
+				if (EditorModule)
+				{
+					if (TSharedPtr<SPCGManualEditPanel> Panel = EditorModule->GetManualEditPanel())
+					{
+						if (TSharedPtr<const FPCGManualEditNodeConfiguration> NodeConfiguration = Panel->GetConfigurationForActiveNode())
+						{
+							DeferredExtension = FPCGEditorModule::GetMutableDeltaViewportExtensionRegistry().GetExtension(NodeConfiguration->DeltaType);
+							DeferredDeltaType = NodeConfiguration->DeltaType;
+							if (DeferredExtension)
+							{
+								DeltaName = DeferredExtension->GetDeltaName();
+							}
+						}
+					}
+				}
+
+				// Use composite key with signature settings from the panel
+				bool bSignaturePosition = true;
+				bool bSignatureRotation = false;
+				bool bSignatureScale = false;
+				double SignatureTolerance = Collection->Settings.SpatialTolerance;
+				if (EditorModule)
+				{
+					if (TSharedPtr<SPCGManualEditPanel> Panel = EditorModule->GetManualEditPanel())
+					{
+						if (TSharedPtr<const FPCGManualEditNodeConfiguration> NodeConfiguration = Panel->GetConfigurationForActiveNode())
+						{
+							bSignaturePosition = NodeConfiguration->bSignaturePosition;
+							bSignatureRotation = NodeConfiguration->bSignatureRotation;
+							bSignatureScale = NodeConfiguration->bSignatureScale;
+							SignatureTolerance = NodeConfiguration->SignatureTolerance;
+						}
+					}
+				}
+
+				DeltaKey = PCG::DataOverride::Keys::FPCGCompositeTransformDeltaKey(CurrentTransform, SignatureTolerance, bSignaturePosition, bSignatureRotation, bSignatureScale, DeltaName);
+			}
+
+			EditTarget.Emplace(MutablePCGComponent, HitProxy->PinStack, StorageKey, DeltaKey, bWasInDelta, CurrentTransform);
+			EditTarget->OriginalTransform = CurrentTransform;
+			EditTarget->OriginalPointIndex = HitProxy->PointIndex;
+
+			if (bWasInDelta && StateResult.Extension)
+			{
+				EditTarget->ActiveExtension = StateResult.Extension;
+				EditTarget->DeltaStructType = StateResult.DeltaView.GetScriptStruct();
+			}
+			else if (!bWasInDelta)
+			{
+				EditTarget->bDeltaDeferred = true;
+				EditTarget->ActiveExtension = DeferredExtension;
+				EditTarget->DeltaStructType = DeferredDeltaType;
 			}
+
+			ResetGizmo();
+			if (!EditTarget->ActiveExtension || EditTarget->ActiveExtension->UsesTRSGizmo())
+			{
+				SetupGizmo(GizmoManager, CurrentTransform);
+			}
+
+			RefreshKeyBindings(EditTarget->ActiveExtension);
+
+			return true;
 		}
 	}
+
 	return false;
 }
 
 void FPCGComponentVisualizer::ResetGizmo()
 {
-	if (Gizmo)
+	if (Gizmo.IsValid())
 	{
 		Gizmo->SetVisibility(false);
 		Gizmo->ClearActiveTarget();
 
 		if (UInteractiveGizmoManager* GizmoManager = Gizmo->GetGizmoManager())
 		{
-			GizmoManager->DestroyGizmo(Gizmo);
+			GizmoManager->DestroyGizmo(Gizmo.Get());
 		}
-		Gizmo = nullptr;
+		Gizmo.Reset();
+	}
+}
+
+void FPCGComponentVisualizer::SetupGizmo(UInteractiveGizmoManager* GizmoManager, const FTransform& InitialTransform)
+{
+	ResetGizmo();
+
+	ETransformGizmoSubElements TransformElements = ETransformGizmoSubElements::FullTranslateRotateScale;
+	Gizmo = UE::TransformGizmoUtil::CreateCustomTransformGizmo(GizmoManager, TransformElements, this, FString());
+	if (!Gizmo.IsValid())
+	{
+		return;
+	}
+
+	UTransformProxy* Proxy = NewObject<UTransformProxy>(Gizmo.Get());
+	Proxy->SetTransform(InitialTransform);
+	Gizmo->SetActiveTarget(Proxy, GizmoManager);
+	Gizmo->ReinitializeGizmoTransform(InitialTransform);
+	Gizmo->SetVisibility(true);
+
+	Proxy->OnTransformChanged.AddLambda([this](UTransformProxy* InProxy, const FTransform& NewTransform)
+	{
+		OnGizmoDragged(NewTransform);
+	});
+
+	Proxy->OnEndTransformEdit.AddLambda([this](UTransformProxy* InProxy)
+	{
+		OnGizmoReleased(InProxy->GetTransform());
+	});
+}
+
+void FPCGComponentVisualizer::RefreshKeyBindings(IPCGDeltaViewportExtension* Extension)
+{
+	ActiveKeyBindings.Empty();
+
+	if (!Extension)
+	{
+		return;
 	}
+
+	ActiveKeyBindings = Extension->GetKeyBindings();
 }
 
 void FPCGComponentVisualizer::EndEditing()
 {
-	EndCurrentElementManualEditing();
+	ActiveKeyBindings.Empty();
+	ClearActiveSelection();
+
+	// Clear only transient manual editing flags, preserve persistent bMarkedForManualEditing
+	if (ActiveEditingNode.IsValid())
+	{
+		if (UPCGSettingsInterface* Settings = ActiveEditingNode->GetSettingsInterface())
+		{
+			Settings->bTemporaryManualEditingEnabled = false;
+		}
+
+		// Broadcast through the PCGNode to trigger a full node reconstruct. The editor graph node is not available here.
+		ActiveEditingNode->OnNodeChangedDelegate.Broadcast(ActiveEditingNode.Get(), EPCGChangeType::Cosmetic);
+	}
+	ActiveEditingNode.Reset();
+	CachedDeltaContext.Reset();
 
-	bManualEditMode = false;
 	LastPCGComponent.Reset();
 	FComponentVisualizer::EndEditing();
 }
 
 bool FPCGComponentVisualizer::GetWidgetLocation(const FEditorViewportClient* ViewportClient, FVector& OutLocation) const
 {
-	return FComponentVisualizer::GetWidgetLocation(ViewportClient, OutLocation);
+	if (EditTarget.IsSet())
+	{
+		OutLocation = EditTarget->Transform.GetLocation();
+		return true;
+	}
+
+	return false;
 }
 
 bool FPCGComponentVisualizer::GetCustomInputCoordinateSystem(const FEditorViewportClient* ViewportClient, FMatrix& OutMatrix) const
@@ -449,76 +900,271 @@ bool FPCGComponentVisualizer::HandleInputDelta(FEditorViewportClient* ViewportCl
 	return FComponentVisualizer::HandleInputDelta(ViewportClient, Viewport, DeltaTranslate, DeltaRotate, DeltaScale);
 }
 
-bool FPCGComponentVisualizer::HandleInputKey(FEditorViewportClient* ViewportClient, FViewport* Viewport, FKey Key, EInputEvent Event)
+bool FPCGComponentVisualizer::HandleInputKey(FEditorViewportClient* ViewportClient, FViewport* Viewport, const FKey Key, const EInputEvent Event)
 {
+	if (Event == IE_Pressed && EditTarget.IsSet())
+	{
+		for (const FPCGDeltaKeyBinding& Binding : ActiveKeyBindings)
+		{
+			if (Binding.Key == Key && Binding.Action)
+			{
+				Binding.Action();
+				return true;
+			}
+		}
+	}
+
 	return FComponentVisualizer::HandleInputKey(ViewportClient, Viewport, Key, Event);
 }
 
 TSharedPtr<SWidget> FPCGComponentVisualizer::GenerateContextMenu() const
 {
-	UE_LOG(LogPCGEditor, Warning, TEXT("FPCGComponentVisualizer::GenerateContextMenu"));
 	FMenuBuilder MenuBuilder(/*bInShouldCloseWindowAfterMenuSelection=*/true, ComponentVisualizerActions);
 
-	MenuBuilder.AddMenuEntry(FPCGComponentVisualizerCommands::Get().DeleteAllDeltas);
-	MenuBuilder.AddMenuEntry(FPCGComponentVisualizerCommands::Get().DeleteSelectedDeltas);
+	MenuBuilder.AddMenuEntry(FPCGComponentVisualizerCommands::Get().RemoveAllOverrides);
+	MenuBuilder.AddMenuEntry(FPCGComponentVisualizerCommands::Get().RemoveSelectedOverride);
 
 	TSharedPtr<SWidget> MenuWidget = MenuBuilder.MakeWidget();
 	return MenuWidget;
 }
 
-void FPCGComponentVisualizer::SetSelectedTransform(const FTransform& Transform)
+void FPCGComponentVisualizer::OnGizmoDragged(const FTransform& NewTransform)
 {
+	if (FEditState* Selection = EditTarget.GetPtrOrNull())
+	{
+		Selection->Transform = NewTransform;
+		Selection->bGizmoDragging = true;
+	}
+}
+
+void FPCGComponentVisualizer::OnGizmoReleased(const FTransform& FinalTransform)
+{
+	using namespace PCG::ComponentVisualizer;
+
 	if (!LastPCGComponent.IsValid())
 	{
 		return;
 	}
 
-	if (FEditState* Selection = EditTarget.GetPtrOrNull(); Selection && EditTarget->PCGComponent.IsValid())
+	FEditState* Selection = EditTarget.GetPtrOrNull();
+	if (!Selection || !EditTarget->PCGComponent.IsValid())
 	{
-		Selection->Transform = Transform;
+		return;
+	}
 
-		FPCGSourceDataContainer* DataContainer = EditTarget->PCGComponent->GetExecutionState().GetSourceDataContainer();
+	Selection->bGizmoDragging = false;
 
-		FSharedStruct SharedStruct = DataContainer->GetMutable<FPCGDeltaCollection>(EditTarget->StorageKey);
-		FPCGDeltaCollection* Collection = SharedStruct.IsValid()
-											  ? SharedStruct.GetPtr<FPCGDeltaCollection>()
-											  : nullptr;
+	Selection->Transform = FinalTransform;
 
-		TInstancedStruct<FPCGDeltaBase>* DeltaInstancedStruct = Collection ? Collection->Find(EditTarget->DeltaKey) : nullptr;
+	FScopedTransaction Transaction(LOCTEXT("MovePointDelta", "PCG Manual Edit: Move Point"));
 
-		if (FPCGPointEditDelta* EditDelta = DeltaInstancedStruct ? DeltaInstancedStruct->GetMutablePtr<FPCGPointEditDelta>() : nullptr)
+	FPCGSourceDataContainer* DataContainer = EditTarget->PCGComponent->GetExecutionState().GetSourceDataContainer();
+	if (!DataContainer)
+	{
+		return;
+	}
+
+	FSharedStruct SharedStruct = DataContainer->GetMutable<FPCGDeltaCollection>(EditTarget->StorageKey);
+	FPCGDeltaCollection* Collection = SharedStruct.IsValid()
+										  ? SharedStruct.GetPtr<FPCGDeltaCollection>()
+										  : nullptr;
+
+	// Lazy delta creation: first gizmo release materializes the deferred delta via the extension
+	if (Selection->bDeltaDeferred && Collection && Selection->ActiveExtension)
+	{
+		Selection->bDeltaDeferred = false;
+
+		FPCGDeltaCreateContext CreateContext;
+		CreateContext.ComponentBounds = EditTarget->PCGComponent->GetTotalBounds();
+		CreateContext.OriginalElementIndex = Selection->OriginalPointIndex;
+
+		Selection->ActiveExtension->CreateNewDelta(EditTarget->DeltaKey, *Collection, Selection->OriginalTransform, FinalTransform, CreateContext);
+	}
+
+	TInstancedStruct<FPCGDeltaBase>* DeltaInstancedStruct = Collection ? Collection->Find(EditTarget->DeltaKey) : nullptr;
+
+	if (DeltaInstancedStruct && Selection->ActiveExtension)
+	{
+		if (Selection->ActiveExtension->ApplyGizmoTransform(FinalTransform, *DeltaInstancedStruct, Selection->SelectedElementIndex))
 		{
-			EditDelta->TransformOverride = Transform;
 			EditTarget->bModified = true;
-			DataContainer->MarkDirty();
+			PCG::DeltaViewportExtension::MarkComponentDirty(LastPCGComponent.Get(), DataContainer);
+		}
+	}
+}
+
+void FPCGComponentVisualizer::RestoreDelta(const FPCGDeltaKey& Key)
+{
+	if (!CachedDeltaContext.DataContainer)
+	{
+		return;
+	}
 
-			// TODO: This is needed here but why???
-			LastPCGComponent->Modify();
-			LastPCGComponent->GetGraphInstance()->GetMutablePCGGraph()->ForceNotificationForEditor(EPCGChangeType::ExternalModification);
+	FSharedStruct SharedStruct = CachedDeltaContext.DataContainer->GetMutable<FPCGDeltaCollection>(CachedDeltaContext.StorageKey);
+	FPCGDeltaCollection* Collection = SharedStruct.IsValid() ? SharedStruct.GetPtr<FPCGDeltaCollection>() : nullptr;
+	if (!Collection)
+	{
+		return;
+	}
+
+	FScopedTransaction Transaction(LOCTEXT("RestoreDelta", "PCG Manual Edit: Restore Element"));
+
+	if (Collection->Remove(Key))
+	{
+		if (EditTarget.IsSet() && EditTarget->DeltaKey == Key)
+		{
+			EditTarget.Reset();
+			ResetGizmo();
 		}
+
+		PCG::DeltaViewportExtension::MarkComponentDirty(LastPCGComponent.Get(), CachedDeltaContext.DataContainer);
 	}
 }
 
+void FPCGComponentVisualizer::RestoreAllDeltas()
+{
+	if (!CachedDeltaContext.DataContainer)
+	{
+		return;
+	}
+
+	FSharedStruct SharedStruct = CachedDeltaContext.DataContainer->GetMutable<FPCGDeltaCollection>(CachedDeltaContext.StorageKey);
+	FPCGDeltaCollection* Collection = SharedStruct.IsValid() ? SharedStruct.GetPtr<FPCGDeltaCollection>() : nullptr;
+	if (!Collection)
+	{
+		return;
+	}
+
+	TArray<FPCGDeltaKey> RestorableKeys;
+	const FPCGDeltaViewportExtensionRegistry& Registry = FPCGEditorModule::GetConstDeltaViewportExtensionRegistry();
+	for (const UScriptStruct* DeltaType : Registry.GetRegisteredDeltaTypes())
+	{
+		if (IPCGDeltaViewportExtension* Extension = Registry.GetExtension(DeltaType))
+		{
+			Extension->CollectRestorableKeys(*Collection, RestorableKeys);
+		}
+	}
+
+	if (RestorableKeys.Num() > 0)
+	{
+		FScopedTransaction Transaction(LOCTEXT("RestoreAllElements", "PCG Manual Edit: Restore All Elements"));
+
+		for (const FPCGDeltaKey& Key : RestorableKeys)
+		{
+			Collection->Remove(Key);
+		}
+
+		// If the current selection was one of the removed deltas, clear it
+		if (EditTarget.IsSet())
+		{
+			for (const FPCGDeltaKey& Key : RestorableKeys)
+			{
+				if (EditTarget->DeltaKey == Key)
+				{
+					EditTarget.Reset();
+					ResetGizmo();
+					break;
+				}
+			}
+		}
+
+		PCG::DeltaViewportExtension::MarkComponentDirty(LastPCGComponent.Get(), CachedDeltaContext.DataContainer);
+	}
+}
+
+void FPCGComponentVisualizer::SelectElement(const FPCGDeltaKey& DeltaKey, int32 ElementIndex)
+{
+	if (!CachedDeltaContext.DataContainer)
+	{
+		return;
+	}
+
+	FSharedStruct SharedStruct = CachedDeltaContext.DataContainer->GetMutable<FPCGDeltaCollection>(CachedDeltaContext.StorageKey);
+	FPCGDeltaCollection* Collection = SharedStruct.IsValid() ? SharedStruct.GetPtr<FPCGDeltaCollection>() : nullptr;
+	if (!Collection)
+	{
+		return;
+	}
+
+	TInstancedStruct<FPCGDeltaBase>* DeltaStruct = Collection->Find(DeltaKey);
+	if (!DeltaStruct)
+	{
+		return;
+	}
+
+	const UScriptStruct* DeltaStructType = DeltaStruct->GetScriptStruct();
+	IPCGDeltaViewportExtension* Extension = FPCGEditorModule::GetMutableDeltaViewportExtensionRegistry().GetExtension(DeltaStructType);
+	if (!Extension)
+	{
+		return;
+	}
+
+	// Resolve the transform for the element
+	FTransform ElementTransform = FTransform::Identity;
+	const FConstStructView DeltaView(DeltaStructType, DeltaStruct->GetMemory());
+
+	if (ElementIndex != INDEX_NONE)
+	{
+		TArray<FInsertedViewportElement> InsertedElements;
+		Extension->GetInsertedElements(*Collection, InsertedElements);
+		for (const FInsertedViewportElement& Inserted : InsertedElements)
+		{
+			if (Inserted.DeltaKey == DeltaKey && Inserted.ElementIndex == ElementIndex)
+			{
+				ElementTransform = Inserted.Transform;
+				break;
+			}
+		}
+	}
+	else // Single-element delta without a specific element index. Read the transform directly from the delta.
+	{
+		ElementTransform = Extension->GetDisplayTransform(FTransform::Identity, DeltaView);
+	}
+
+	EditTarget.Emplace(LastPCGComponent.Get(), CachedDeltaContext.PinStack, CachedDeltaContext.StorageKey, DeltaKey, /*bWasInDelta=*/true, ElementTransform);
+	EditTarget->OriginalTransform = ElementTransform;
+	EditTarget->ActiveExtension = Extension;
+	EditTarget->DeltaStructType = DeltaStructType;
+	EditTarget->SelectedElementIndex = ElementIndex;
+
+	ResetGizmo();
+	UModeManagerInteractiveToolsContext* ToolsContext = GLevelEditorModeTools().GetInteractiveToolsContext();
+	UInteractiveGizmoManager* GizmoManager = ToolsContext ? ToolsContext->GizmoManager.Get() : nullptr;
+	if (Extension->UsesTRSGizmo() && GizmoManager)
+	{
+		SetupGizmo(GizmoManager, ElementTransform);
+	}
+
+	RefreshKeyBindings(Extension);
+}
+
 void FPCGComponentVisualizer::RemoveAllDeltas()
 {
-	if (LastPCGComponent.IsValid())
+	// @todo_pcg: Only removes deltas for the active node's storage key. If multiple nodes have been edited, enumerate
+	// all delta storage keys to clear them all.
+	if (LastPCGComponent.IsValid() && CachedDeltaContext.IsValid())
 	{
-		FPCGSourceDataContainer* DataContainer = LastPCGComponent->GetExecutionState().GetSourceDataContainer();
-		DataContainer->Empty();
-		DataContainer->MarkDirty();
+		FScopedTransaction Transaction(LOCTEXT("RemoveAllElements", "PCG Manual Edit: Remove All Elements"));
 
-		// TODO: This is needed here but why???
-		LastPCGComponent->Modify();
-		LastPCGComponent->GetGraphInstance()->GetMutablePCGGraph()->ForceNotificationForEditor(EPCGChangeType::ExternalModification);
+		CachedDeltaContext.DataContainer->Remove(CachedDeltaContext.StorageKey);
+		PCG::DeltaViewportExtension::MarkComponentDirty(LastPCGComponent.Get(), CachedDeltaContext.DataContainer);
+		CachedDeltaContext.Reset();
 	}
-	EndCurrentElementManualEditing();
+
+	ClearActiveSelection();
 }
 
 void FPCGComponentVisualizer::RemoveSelectedDeltas()
 {
 	if (LastPCGComponent.IsValid() && EditTarget.IsSet())
 	{
+		FScopedTransaction Transaction(LOCTEXT("RemoveSelectedElement", "PCG Manual Edit: Remove Selected Element"));
+
 		FPCGSourceDataContainer* DataContainer = EditTarget->PCGComponent->GetExecutionState().GetSourceDataContainer();
+		if (!DataContainer)
+		{
+			return;
+		}
 
 		FSharedStruct SharedStruct = DataContainer->GetMutable<FPCGDeltaCollection>(EditTarget->StorageKey);
 		if (FPCGDeltaCollection* Collection = SharedStruct.IsValid() ? SharedStruct.GetPtr<FPCGDeltaCollection>() : nullptr)
@@ -526,13 +1172,10 @@ void FPCGComponentVisualizer::RemoveSelectedDeltas()
 			Collection->Remove(EditTarget->DeltaKey);
 		}
 
-		DataContainer->MarkDirty();
 		EditTarget.Reset();
-		EndCurrentElementManualEditing();
+		ClearActiveSelection();
 
-		// TODO: This is needed here but why???
-		LastPCGComponent->Modify();
-		LastPCGComponent->GetGraphInstance()->GetMutablePCGGraph()->ForceNotificationForEditor(EPCGChangeType::ExternalModification);
+		PCG::DeltaViewportExtension::MarkComponentDirty(LastPCGComponent.Get(), DataContainer);
 	}
 }
 
@@ -541,12 +1184,14 @@ bool FPCGComponentVisualizer::HasSelection() const
 	return LastPCGComponent.IsValid() && EditTarget.IsSet();
 }
 
-void FPCGComponentVisualizer::EndCurrentElementManualEditing()
+void FPCGComponentVisualizer::ClearActiveSelection()
 {
 	if (EditTarget.IsSet())
 	{
-		if (!EditTarget->bWasInDelta && !EditTarget->bModified)
+		if (!EditTarget->bWasInDelta && !EditTarget->bModified && !EditTarget->bDeltaDeferred)
 		{
+			// Delta was created but never modified — remove it.
+			// If bDeltaDeferred is true, no delta was ever stored, so nothing to remove.
 			RemoveSelectedDeltas();
 		}
 	}
@@ -555,4 +1200,4 @@ void FPCGComponentVisualizer::EndCurrentElementManualEditing()
 	ResetGizmo();
 }
 
-#undef LOCTEXT_NAMESPACE
\ No newline at end of file
+#undef LOCTEXT_NAMESPACE
diff --git a/Engine/Plugins/PCG/Source/PCGEditor/Private/PCGEditor.cpp b/Engine/Plugins/PCG/Source/PCGEditor/Private/PCGEditor.cpp
--- a/Engine/Plugins/PCG/Source/PCGEditor/Private/PCGEditor.cpp
+++ b/Engine/Plugins/PCG/Source/PCGEditor/Private/PCGEditor.cpp
@@ -37,19 +37,19 @@
 #include "Schema/PCGEditorGraphSchema.h"
 #include "Schema/PCGEditorGraphSchemaActions.h"
 
+#include "Widgets/SPCGCodeEditor.h"
 #include "Widgets/SPCGEditorGraphActionWidget.h"
 #include "Widgets/SPCGEditorGraphAttributeListView.h"
 #include "Widgets/SPCGEditorGraphDebugObjectTree.h"
 #include "Widgets/SPCGEditorGraphDetailsView.h"
 #include "Widgets/SPCGEditorGraphDeterminism.h"
-#include "Widgets/SPCGEditorGraphFind.h"
 #include "Widgets/SPCGEditorGraphEmbeddedSubgraphsView.h"
+#include "Widgets/SPCGEditorGraphFind.h"
 #include "Widgets/SPCGEditorGraphLogView.h"
 #include "Widgets/SPCGEditorGraphNodePalette.h"
 #include "Widgets/SPCGEditorGraphParamsView.h"
 #include "Widgets/SPCGEditorGraphProfilingView.h"
 #include "Widgets/SPCGEditorGraphTitleBar.h"
-#include "Widgets/SPCGCodeEditor.h"
 #include "Widgets/AssetEditorViewport/SPCGEditorViewport.h"
 
 #include "AssetDefinition.h"
@@ -80,12 +80,14 @@
 #include "Misc/ITransaction.h"
 #include "Misc/MessageDialog.h"
 #include "Misc/TransactionObjectEvent.h"
+#include "Overlay/SDraggableBoxOverlay.h"
 #include "Preferences/UnrealEdOptions.h"
+#include "Widgets/SPCGEditorGraphDataOverrides.h"
 #include "Widgets/Docking/SDockTab.h"
 #include "Widgets/Input/SEditableTextBox.h"
 #include "Widgets/Notifications/SNotificationList.h"
-#include "WorkflowOrientedApp/WorkflowTabManager.h"
 #include "WorkflowOrientedApp/WorkflowTabFactory.h"
+#include "WorkflowOrientedApp/WorkflowTabManager.h"
 #include "WorkflowOrientedApp/WorkflowUObjectDocuments.h"
 
 #define LOCTEXT_NAMESPACE "PCGGraphEditor"
@@ -116,6 +118,7 @@ namespace FPCGEditor_private
 		PCGEditorTabs::ViewportID[2],
 		PCGEditorTabs::ViewportID[3],
 		PCGEditorTabs::EmbeddedSubgraphsID,
+		PCGEditorTabs::DataOverridesID
 	};
 
 	const FName AllRestrictedPanels[] =
@@ -251,6 +254,8 @@ void FPCGEditor::Initialize(const EToolkitMode::Type InMode, const TSharedPtr<cl
 
 	AddApplicationMode(FPCGEditorDefaultMode::StaticName(), MakeShared<FPCGEditorDefaultMode>(ThisPtr));
 	SetCurrentMode(FPCGEditorDefaultMode::StaticName());
+
+	NotifyModuleManualEditStateChanged();
 }
 
 UPCGGraph* FPCGEditor::GetMainGraph() const
@@ -452,6 +457,7 @@ void FPCGEditor::UpdateAfterInspectedStackChanged(const FPCGStack& FullStack)
 			EPCGChangeType ChangeType = PCGNode->UpdateErrorsAndWarnings();
 			ChangeType |= PCGNode->UpdateStructuralVisualization(Source, &FullStack);
 			ChangeType |= PCGNode->UpdateGPUVisualization(Source, &FullStack);
+			ChangeType |= PCGNode->UpdateOverrideVisualization(Source, &FullStack);
 
 			if (ChangeType != EPCGChangeType::None)
 			{
@@ -1024,6 +1030,18 @@ void FPCGEditor::BindCommands()
 		FExecuteAction::CreateSP(this, &FPCGEditor::OnDisableDebugOnAllNodes),
 		FCanExecuteAction::CreateSP(this, &FPCGEditor::CanToggleDebug));
 
+	GraphEditorCommands->MapAction(
+		PCGEditorCommands.ToggleViewportEditing,
+		FExecuteAction::CreateSP(this, &FPCGEditor::OnToggleViewportEditing),
+		FCanExecuteAction::CreateSP(this, &FPCGEditor::CanToggleViewportEditing),
+		FGetActionCheckState::CreateSP(this, &FPCGEditor::GetViewportEditingCheckState));
+
+	GraphEditorCommands->MapAction(
+		PCGEditorCommands.MarkForViewportEditing,
+		FExecuteAction::CreateSP(this, &FPCGEditor::OnMarkForViewportEditing),
+		FCanExecuteAction::CreateSP(this, &FPCGEditor::CanToggleViewportEditing),
+		FGetActionCheckState::CreateSP(this, &FPCGEditor::GetMarkForViewportEditingCheckState));
+
 	GraphEditorCommands->MapAction(
 		PCGEditorCommands.AddSourcePin,
 		FExecuteAction::CreateSP(this, &FPCGEditor::OnAddDynamicInputPin),
@@ -2313,6 +2331,11 @@ void FPCGEditor::SetNodeInspected(UPCGEditorGraphNodeBase* InspectedNode, bool b
 				AttributeListView->SetNodeBeingInspected(nullptr);
 			}
 		}
+
+		if (DataOverridesWidget.IsValid())
+		{
+			DataOverridesWidget->SetNodeBeingInspected(nullptr);
+		}
 	}
 	else
 	{
@@ -2335,6 +2358,8 @@ void FPCGEditor::SetNodeInspected(UPCGEditorGraphNodeBase* InspectedNode, bool b
 			}
 		}
 
+		GetOrCreateDataOverridesWidget()->SetNodeBeingInspected(InspectedNode);
+
 		if (InspectedNode && InspectedNodesAfter.Contains(InspectedNode))
 		{
 			InspectedNode->SetInspected(true);
@@ -2374,10 +2399,12 @@ void FPCGEditor::SetNodeInspected(UPCGEditorGraphNodeBase* InspectedNode, bool b
 			return false;
 		};
 
-		const bool bTabSummoned = (InvokeFirstTab(true) || InvokeFirstTab(false));
+		// If the Data Overrides tab is already active, don't switch focus to an ALV.
+		const bool bDataOverridesActive = TabManager->FindExistingLiveTab(FTabId(PCGEditorTabs::DataOverridesID)).IsValid();
+		const bool bTabSummoned = !bDataOverridesActive && (InvokeFirstTab(true) || InvokeFirstTab(false));
 
 		// Default to first if they are all locked
-		if (!bTabSummoned && IsPanelAvailable(PCGEditorTabs::AttributesID[0]))
+		if (!bTabSummoned && !bDataOverridesActive && IsPanelAvailable(PCGEditorTabs::AttributesID[0]))
 		{
 			GetTabManager()->TryInvokeTab(FTabId(PCGEditorTabs::AttributesID[0]));
 		}
@@ -2785,6 +2812,148 @@ ECheckBoxState FPCGEditor::GetDebugCheckState() const
 	return ECheckBoxState::Unchecked;
 }
 
+void FPCGEditor::OnToggleViewportEditing()
+{
+	UPCGEditorGraph* FocusedEditorGraph = GetFocusedEditorGraph();
+	const TSharedPtr<SGraphEditor> FocusedGraphWidget = GetFocusedGraphEditorWidget();
+	if (!FocusedEditorGraph || !FocusedGraphWidget)
+	{
+		return;
+	}
+
+	UEdGraphNode* SelectedGraphNode = FocusedGraphWidget->GetSingleSelectedNode();
+	UPCGEditorGraphNodeBase* SelectedEditorGraphNode = Cast<UPCGEditorGraphNodeBase>(SelectedGraphNode);
+	const UPCGNode* SelectedPCGNode = SelectedEditorGraphNode ? SelectedEditorGraphNode->GetPCGNode() : nullptr;
+	UPCGSettingsInterface* SelectedSettingsInterface = SelectedPCGNode ? SelectedPCGNode->GetSettingsInterface() : nullptr;
+
+	if (!SelectedSettingsInterface || !SelectedSettingsInterface->CanBeDebugged())
+	{
+		return;
+	}
+
+	const bool bManualEditingPreviouslyEnabled = SelectedSettingsInterface->bTemporaryManualEditingEnabled;
+
+	// Only one node can have viewport editing active at a time; clear the flag on all others when enabling a new one.
+	if (!bManualEditingPreviouslyEnabled)
+	{
+		for (UEdGraphNode* Node : FocusedEditorGraph->Nodes)
+		{
+			UPCGEditorGraphNodeBase* PCGEditorGraphNode = Cast<UPCGEditorGraphNodeBase>(Node);
+			const UPCGNode* PCGNode = PCGEditorGraphNode ? PCGEditorGraphNode->GetPCGNode() : nullptr;
+			UPCGSettingsInterface* PCGSettingsInterface = PCGNode ? PCGNode->GetSettingsInterface() : nullptr;
+			if (PCGSettingsInterface && PCGSettingsInterface != SelectedSettingsInterface && PCGSettingsInterface->bTemporaryManualEditingEnabled)
+			{
+				PCGSettingsInterface->bTemporaryManualEditingEnabled = false;
+				PCGEditorGraphNode->OnNodeChangedDelegate.ExecuteIfBound();
+			}
+		}
+	}
+
+	// Toggle the selected node
+	const bool bNewState = !bManualEditingPreviouslyEnabled;
+	SelectedSettingsInterface->bTemporaryManualEditingEnabled = bNewState;
+	SelectedEditorGraphNode->OnNodeChangedDelegate.ExecuteIfBound();
+
+	NotifyModuleManualEditStateChanged();
+}
+
+bool FPCGEditor::CanToggleViewportEditing() const
+{
+	const TSharedPtr<SGraphEditor> FocusedGraphWidget = GetFocusedGraphEditorWidget();
+	if (!FocusedGraphWidget)
+	{
+		return false;
+	}
+
+	const FGraphPanelSelectionSet& SelectedNodes = FocusedGraphWidget->GetSelectedNodes();
+	if (SelectedNodes.Num() != 1)
+	{
+		return false;
+	}
+
+	for (const UObject* NodeObject : SelectedNodes)
+	{
+		const UPCGEditorGraphNodeBase* PCGEditorGraphNode = Cast<UPCGEditorGraphNodeBase>(NodeObject);
+		const UPCGNode* PCGNode = PCGEditorGraphNode ? PCGEditorGraphNode->GetPCGNode() : nullptr;
+		const UPCGSettingsInterface* PCGSettingsInterface = PCGNode ? PCGNode->GetSettingsInterface() : nullptr;
+		if (PCGSettingsInterface && PCGSettingsInterface->CanBeDebugged())
+		{
+			return true;
+		}
+	}
+
+	return false;
+}
+
+ECheckBoxState FPCGEditor::GetViewportEditingCheckState() const
+{
+	const TSharedPtr<SGraphEditor> FocusedGraphWidget = GetFocusedGraphEditorWidget();
+	if (!FocusedGraphWidget)
+	{
+		return ECheckBoxState::Unchecked;
+	}
+
+	const UPCGEditorGraphNodeBase* PCGEditorGraphNode = Cast<UPCGEditorGraphNodeBase>(FocusedGraphWidget->GetSingleSelectedNode());
+	const UPCGNode* PCGNode = PCGEditorGraphNode ? PCGEditorGraphNode->GetPCGNode() : nullptr;
+	const UPCGSettingsInterface* PCGSettingsInterface = PCGNode ? PCGNode->GetSettingsInterface() : nullptr;
+
+	if (PCGSettingsInterface && PCGSettingsInterface->bTemporaryManualEditingEnabled)
+	{
+		return ECheckBoxState::Checked;
+	}
+
+	return ECheckBoxState::Unchecked;
+}
+
+void FPCGEditor::OnMarkForViewportEditing()
+{
+	const TSharedPtr<SGraphEditor> FocusedGraphWidget = GetFocusedGraphEditorWidget();
+	if (!FocusedGraphWidget)
+	{
+		return;
+	}
+
+	UEdGraphNode* SelectedGraphNode = FocusedGraphWidget->GetSingleSelectedNode();
+	UPCGEditorGraphNodeBase* SelectedEditorGraphNode = Cast<UPCGEditorGraphNodeBase>(SelectedGraphNode);
+	const UPCGNode* SelectedPCGNode = SelectedEditorGraphNode ? SelectedEditorGraphNode->GetPCGNode() : nullptr;
+	UPCGSettingsInterface* SelectedSettingsInterface = SelectedPCGNode ? SelectedPCGNode->GetSettingsInterface() : nullptr;
+
+	if (!SelectedSettingsInterface || !SelectedSettingsInterface->CanBeDebugged())
+	{
+		return;
+	}
+
+	{
+		FScopedTransaction Transaction(*FPCGEditorCommon::ContextIdentifier, LOCTEXT("PCGEditorMarkForViewportEditingTransactionMessage", "PCG Editor: Mark for Viewport Editing"), nullptr);
+
+		SelectedSettingsInterface->Modify();
+		SelectedSettingsInterface->bMarkedForManualEditing = !SelectedSettingsInterface->bMarkedForManualEditing;
+	}
+
+	SelectedEditorGraphNode->OnNodeChangedDelegate.ExecuteIfBound();
+	NotifyModuleManualEditStateChanged();
+}
+
+ECheckBoxState FPCGEditor::GetMarkForViewportEditingCheckState() const
+{
+	const TSharedPtr<SGraphEditor> FocusedGraphWidget = GetFocusedGraphEditorWidget();
+	if (!FocusedGraphWidget)
+	{
+		return ECheckBoxState::Unchecked;
+	}
+
+	const UPCGEditorGraphNodeBase* PCGEditorGraphNode = Cast<UPCGEditorGraphNodeBase>(FocusedGraphWidget->GetSingleSelectedNode());
+	const UPCGNode* PCGNode = PCGEditorGraphNode ? PCGEditorGraphNode->GetPCGNode() : nullptr;
+	const UPCGSettingsInterface* PCGSettingsInterface = PCGNode ? PCGNode->GetSettingsInterface() : nullptr;
+
+	if (PCGSettingsInterface && PCGSettingsInterface->bMarkedForManualEditing)
+	{
+		return ECheckBoxState::Checked;
+	}
+
+	return ECheckBoxState::Unchecked;
+}
+
 void FPCGEditor::SelectAllNodes()
 {
 	if (TSharedPtr<SGraphEditor> FocusedGraphWidget = GetFocusedGraphEditorWidget())
@@ -3280,10 +3449,18 @@ TSharedRef<SGraphEditor> FPCGEditor::CreateGraphEditorWidget(TSharedPtr<FTabInfo
 	return GraphWidget;
 }
 
+void FPCGEditor::NotifyModuleManualEditStateChanged()
+{
+	FModuleManager::GetModuleChecked<FPCGEditorModule>("PCGEditor").UpdateManualEditPanelVisibility();
+}
+
 void FPCGEditor::OnClose()
 {
 	FWorkflowCentricApplication::OnClose();
 
+	// Hide the manual edit panel now that this editor is closing.
+	NotifyModuleManualEditStateChanged();
+
 	if (FLevelEditorModule* LevelEditor = FModuleManager::GetModulePtr<FLevelEditorModule>(TEXT("LevelEditor")))
 	{
 		LevelEditor->OnMapChanged().RemoveAll(this);
@@ -3879,6 +4056,13 @@ TSharedPtr<FWorkflowTabFactory> FPCGEditor::CreateEmbeddedSubgraphsTabFactory()
 	return MakeShared<FPCGGenericTabFactory>(Params);
 }
 
+TSharedPtr<FWorkflowTabFactory> FPCGEditor::CreateDataOverridesTabFactory()
+{
+	FPCGGenericTabFactoryParams Params(PCGEditorTabs::DataOverridesID, SharedThis<FPCGEditor>(this), [this]() { return GetOrCreateDataOverridesWidget(); });
+	Params.Label = LOCTEXT("DataOverridesLabel", "Data Overrides");
+	return MakeShared<FPCGGenericTabFactory>(Params);
+}
+
 TSharedRef<SPCGEditorGraphEmbeddedSubgraphsView> FPCGEditor::GetOrCreateEmbeddedSubgraphsWidget()
 {
 	if (!EmbeddedSubgraphsWidget.IsValid())
@@ -3889,6 +4073,16 @@ TSharedRef<SPCGEditorGraphEmbeddedSubgraphsView> FPCGEditor::GetOrCreateEmbedded
 	return EmbeddedSubgraphsWidget.ToSharedRef();
 }
 
+TSharedRef<SPCGEditorGraphDataOverridesView> FPCGEditor::GetOrCreateDataOverridesWidget()
+{
+	if (!DataOverridesWidget.IsValid())
+	{
+		DataOverridesWidget = SNew(SPCGEditorGraphDataOverridesView, SharedThis(this));
+	}
+
+	return DataOverridesWidget.ToSharedRef();
+}
+
 void FPCGEditor::OnGraphEditorFocused(const TSharedRef<SGraphEditor>& InGraphEditor)
 {
 	FocusedGraphEditor = InGraphEditor;
@@ -4121,6 +4315,7 @@ void FPCGEditor::OnSourceGenerationDone(IPCGBaseSubsystem* Subsystem, IPCGGraphE
 			ChangeType |= PCGEditorGraphNode->UpdateErrorsAndWarnings();
 			ChangeType |= PCGEditorGraphNode->UpdateStructuralVisualization(PCGSourceBeingInspected, &StackBeingInspected);
 			ChangeType |= PCGEditorGraphNode->UpdateGPUVisualization(PCGSourceBeingInspected, &StackBeingInspected);
+			ChangeType |= PCGEditorGraphNode->UpdateOverrideVisualization(PCGSourceBeingInspected, &StackBeingInspected);
 
 			if (ChangeType != EPCGChangeType::None)
 			{
diff --git a/Engine/Plugins/PCG/Source/PCGEditor/Private/PCGEditorCommands.cpp b/Engine/Plugins/PCG/Source/PCGEditor/Private/PCGEditorCommands.cpp
--- a/Engine/Plugins/PCG/Source/PCGEditor/Private/PCGEditorCommands.cpp
+++ b/Engine/Plugins/PCG/Source/PCGEditor/Private/PCGEditorCommands.cpp
@@ -48,6 +48,8 @@ void FPCGEditorCommands::RegisterCommands()
 	UI_COMMAND(SelectNamedRerouteDeclaration, "Select Named Reroute Declaration", "Selects the associated Named Reroute Declaration matching this Named Reroute Usage", EUserInterfaceActionType::Button, FInputChord());
 	UI_COMMAND(JumpToSource, "Jump to Source", "Jumps to the associated source file.", EUserInterfaceActionType::Button, FInputChord());
 	UI_COMMAND(AutoFocusViewport, "Auto Focus", "When inspection updates the viewport will re-focus to the data.", EUserInterfaceActionType::ToggleButton, FInputChord());
+	UI_COMMAND(ToggleViewportEditing, "Toggle Viewport Editing", "Toggle viewport manual editing for selected node", EUserInterfaceActionType::ToggleButton, FInputChord(EKeys::V));
+	UI_COMMAND(MarkForViewportEditing, "Mark for Viewport Editing", "Mark selected node for persistent viewport editing", EUserInterfaceActionType::ToggleButton, FInputChord());
 
 	//~ PLATFORM SPECIFIC ~//
 #if PLATFORM_MAC
diff --git a/Engine/Plugins/PCG/Source/PCGEditor/Private/PCGEditorDefaultMode.cpp b/Engine/Plugins/PCG/Source/PCGEditor/Private/PCGEditorDefaultMode.cpp
--- a/Engine/Plugins/PCG/Source/PCGEditor/Private/PCGEditorDefaultMode.cpp
+++ b/Engine/Plugins/PCG/Source/PCGEditor/Private/PCGEditorDefaultMode.cpp
@@ -77,6 +77,7 @@ namespace PCGEditorDefaultMode
 						->AddTab(PCGEditorTabs::ProfilingID, ETabState::OpenedTab)
 						->AddTab(PCGEditorTabs::FindID, ETabState::OpenedTab)
 						->AddTab(PCGEditorTabs::DeterminismID, ETabState::ClosedTab)
+						->AddTab(PCGEditorTabs::DataOverridesID, ETabState::ClosedTab)
 						->SetForegroundTab(FTabId(PCGEditorTabs::AttributesID[0]))
 					)
 				)
@@ -124,7 +125,8 @@ FPCGEditorDefaultMode::FPCGEditorDefaultMode(TSharedPtr<FPCGEditor> InPCGEditor)
 		{PCGEditorTabs::ViewportID[1], PCGEditorDefaultMode::FCreateTabFactoryDelegate::CreateSP(InPCGEditor.Get(), &FPCGEditor::CreateViewportTabFactory, 1, ViewportGroup)},
 		{PCGEditorTabs::ViewportID[2], PCGEditorDefaultMode::FCreateTabFactoryDelegate::CreateSP(InPCGEditor.Get(), &FPCGEditor::CreateViewportTabFactory, 2, ViewportGroup)},
 		{PCGEditorTabs::ViewportID[3], PCGEditorDefaultMode::FCreateTabFactoryDelegate::CreateSP(InPCGEditor.Get(), &FPCGEditor::CreateViewportTabFactory, 3, ViewportGroup)},
-		{PCGEditorTabs::EmbeddedSubgraphsID, PCGEditorDefaultMode::FCreateTabFactoryDelegate::CreateSP(InPCGEditor.Get(), &FPCGEditor::CreateEmbeddedSubgraphsTabFactory)}
+		{PCGEditorTabs::EmbeddedSubgraphsID, PCGEditorDefaultMode::FCreateTabFactoryDelegate::CreateSP(InPCGEditor.Get(), &FPCGEditor::CreateEmbeddedSubgraphsTabFactory)},
+		{PCGEditorTabs::DataOverridesID, PCGEditorDefaultMode::FCreateTabFactoryDelegate::CreateSP(InPCGEditor.Get(), &FPCGEditor::CreateDataOverridesTabFactory)}
 	};
 
 	for (TabInfo& Info : DefaultTabInfo)
diff --git a/Engine/Plugins/PCG/Source/PCGEditor/Private/PCGEditorGraph.cpp b/Engine/Plugins/PCG/Source/PCGEditor/Private/PCGEditorGraph.cpp
--- a/Engine/Plugins/PCG/Source/PCGEditor/Private/PCGEditorGraph.cpp
+++ b/Engine/Plugins/PCG/Source/PCGEditor/Private/PCGEditorGraph.cpp
@@ -217,6 +217,7 @@ void UPCGEditorGraph::UpdateVisualizations(IPCGGraphExecutionSource* PCGSourceBe
 			EPCGChangeType ChangeType = EPCGChangeType::None;
 			ChangeType |= PCGEditorNode->UpdateStructuralVisualization(PCGSourceBeingInspected, PCGStackBeingInspected);
 			ChangeType |= PCGEditorNode->UpdateGPUVisualization(PCGSourceBeingInspected, PCGStackBeingInspected);
+			ChangeType |= PCGEditorNode->UpdateOverrideVisualization(PCGSourceBeingInspected, PCGStackBeingInspected);
 
 			if (ChangeType != EPCGChangeType::None)
 			{
@@ -253,6 +254,11 @@ const UPCGEditorGraphNodeBase* UPCGEditorGraph::GetEditorNodeFromPCGNode(const U
 	return nullptr;
 }
 
+UPCGEditorGraphNodeBase* UPCGEditorGraph::GetEditorNodeFromPCGNode(const UPCGNode* InPCGNode)
+{
+	return const_cast<UPCGEditorGraphNodeBase*>(static_cast<const UPCGEditorGraph*>(this)->GetEditorNodeFromPCGNode(InPCGNode));
+}
+
 TSubclassOf<UPCGEditorGraphNodeBase> UPCGEditorGraph::GetGraphNodeClassFromPCGSettings(const UPCGSettings* Settings)
 {
 	if (Settings->IsA<UPCGNamedRerouteDeclarationSettings>())
diff --git a/Engine/Plugins/PCG/Source/PCGEditor/Private/PCGEditorGraph.h b/Engine/Plugins/PCG/Source/PCGEditor/Private/PCGEditorGraph.h
--- a/Engine/Plugins/PCG/Source/PCGEditor/Private/PCGEditorGraph.h
+++ b/Engine/Plugins/PCG/Source/PCGEditor/Private/PCGEditorGraph.h
@@ -60,6 +60,7 @@ public:
 
 	/** Returns the PCG editor graph node corresponding to the given PCG node. */
 	const UPCGEditorGraphNodeBase* GetEditorNodeFromPCGNode(const UPCGNode* InPCGNode) const;
+	UPCGEditorGraphNodeBase* GetEditorNodeFromPCGNode(const UPCGNode* InPCGNode);
 
 	/** Returns whether the edge is used by the downstream node's execution. Cached per graph change. */
 	bool IsEdgeUsedByExecution(const UPCGEdge* InEdge) const;
diff --git a/Engine/Plugins/PCG/Source/PCGEditor/Private/PCGEditorModule.cpp b/Engine/Plugins/PCG/Source/PCGEditor/Private/PCGEditorModule.cpp
--- a/Engine/Plugins/PCG/Source/PCGEditor/Private/PCGEditorModule.cpp
+++ b/Engine/Plugins/PCG/Source/PCGEditor/Private/PCGEditorModule.cpp
@@ -2,16 +2,18 @@
 
 #include "PCGEditorModule.h"
 
-#include "ComponentVisualizers.h"
 #include "PCGComponent.h"
 #include "PCGComponentVisualizer.h"
 #include "PCGEngineSettings.h"
+#include "PCGGraph.h"
 #include "PCGGraphExecutionStateInterface.h"
 #include "PCGModule.h"
+#include "PCGNode.h"
 #include "PCGParamData.h"
+#include "PCGSettings.h"
 #include "PCGSubsystem.h"
 #include "PCGVolumeFactory.h"
-#include "Data/PCGTexture2DArrayData.h"
+#include "PCGWorldActor.h"
 #include "Data/PCGCollisionShapeData.h"
 #include "Data/PCGCollisionWrapperData.h"
 #include "Data/PCGDynamicMeshData.h"
@@ -23,6 +25,7 @@
 #include "Data/PCGSpatialData.h"
 #include "Data/PCGSplineData.h"
 #include "Data/PCGStaticMeshResourceData.h"
+#include "Data/PCGTexture2DArrayData.h"
 #include "Data/PCGTextureData.h"
 #include "Data/PCGVirtualTextureData.h"
 #include "Data/PCGVolumeData.h"
@@ -38,13 +41,14 @@
 #include "PCGEditorSettings.h"
 #include "PCGEditorStyle.h"
 #include "PCGEditorUtils.h"
-#include "PCGWorldActor.h"
 #include "ChangeTracking/PCGActorTracker.h"
 #include "ChangeTracking/PCGComponentChangeHandler.h"
 #include "ChangeTracking/PCGLandscapeTracker.h"
 #include "ChangeTracking/PCGMiscTracker.h"
 #include "ChangeTracking/PCGRuntimeGenChangeHandler.h"
+#include "Data/DataView/PCGDataViewData.h"
 #include "DataVisualizations/PCGBaseTextureDataVisualization.h"
+#include "DataVisualizations/PCGDataViewVisualization.h"
 #include "DataVisualizations/PCGLandscapeDataVisualization.h"
 #include "DataVisualizations/PCGParamDataVisualization.h"
 #include "DataVisualizations/PCGPolygon2DDataVisualization.h"
@@ -53,6 +57,11 @@
 #include "DataVisualizations/PCGStaticMeshDataVisualization.h"
 #include "DataVisualizations/PCGTexture2DArrayDataVisualization.h"
 #include "DataVisualizations/PCGVolumeDataVisualization.h"
+#include "DeltaViewportExtensions/PCGPointDeletionViewportExtension.h"
+#include "DeltaViewportExtensions/PCGPointEditViewportExtension.h"
+#include "DeltaViewportExtensions/PCGPointInsertionViewportExtension.h"
+#include "DeltaVisualizations/PCGDeltaBaseVisualization.h"
+#include "DeltaVisualizations/PCGDeltaPointVisualization.h"
 #include "Details/EnumSelectorDetails.h"
 #include "Details/PCGAttributePropertySelectorDetails.h"
 #include "Details/PCGBlueprintSettingsDetails.h"
@@ -64,11 +73,18 @@
 #include "Details/PCGGraphDetails.h"
 #include "Details/PCGGraphInstanceDetails.h"
 #include "Details/PCGInstancedPropertyBagOverrideDetails.h"
-#include "Details/PCGVolumeDetails.h"
 #include "Details/PCGInteractiveToolSettingsDetails.h"
+#include "Details/PCGRaycastFilterRuleCollectionDetails.h"
+#include "Details/PCGVolumeDetails.h"
+#include "EditorMode/PCGEdMode.h"
+#include "EditorMode/Tools/PCGInteractiveToolSettings.h"
+#include "Graph/DataOverride/PCGDataOverridePoints.h"
 #include "Nodes/PCGEditorGraphNodeFactory.h"
+#include "PCGManualEditPanelManager.h"
 #include "Schema/PCGEditorGraphSchema.h"
+#include "Schema/PCGGraphParametersSchema.h"
 
+#include "ComponentVisualizers.h"
 #include "ContentBrowserMenuContexts.h"
 #include "ContentBrowserModule.h"
 #include "Editor.h"
@@ -81,22 +97,18 @@
 #include "LevelEditorMenuContext.h"
 #include "PropertyEditorModule.h"
 #include "ScopedTransaction.h"
+#include "Selection.h"
 #include "SEnumCombo.h"
 #include "SLevelViewport.h"
-#include "ToolMenus.h"
 #include "ToolMenuEntry.h"
+#include "ToolMenus.h"
 #include "UnrealEdGlobals.h"
-#include "Data/DataView/PCGDataViewData.h"
-#include "DataVisualizations/PCGDataViewVisualization.h"
-#include "Details/PCGRaycastFilterRuleCollectionDetails.h"
 #include "Editor/UnrealEdEngine.h"
 #include "Editor/SceneOutliner/Public/ISceneOutliner.h"
-#include "EditorMode/PCGEdMode.h"
-#include "EditorMode/Tools/PCGInteractiveToolSettings.h"
 #include "Framework/MultiBox/MultiBoxBuilder.h"
 #include "Framework/Notifications/NotificationManager.h"
 #include "HierarchyEditor/PropertyBagHierarchyViewModel.h"
-#include "Schema/PCGGraphParametersSchema.h"
+#include "Subsystems/AssetEditorSubsystem.h"
 #include "UObject/ObjectSaveContext.h"
 #include "Widgets/SPCGCodeEditorTextBox.h"
 
@@ -121,6 +133,8 @@ void FPCGEditorModule::StartupModule()
 	RegisterDetailsCustomizations();
 	RegisterSettings();
 	RegisterPCGDataVisualizations();
+	RegisterDeltaVisualizations();
+	RegisterDeltaViewportExtensions();
 	RegisterPinColorAndIcons();
 	RegisterTrackerFactories();
 
@@ -172,7 +186,11 @@ void FPCGEditorModule::ShutdownModule()
 
 	LoadDataAssetChangeTracker.Reset();
 
+	ManualEditPanelManager.Reset();
+
 	UnregisterTrackerFactories();
+	UnregisterDeltaViewportExtensions();
+	UnregisterDeltaVisualizations();
 	UnregisterPCGDataVisualizations();
 	UnregisterSettings();
 	UnregisterDetailsCustomizations();
@@ -309,6 +327,11 @@ void FPCGEditorModule::OnLevelEditorCreated(TSharedPtr<ILevelEditor> InLevelEdit
 		GEditor->OnSceneMaterialsModifiedEvent().AddRaw(this, &FPCGEditorModule::OnSceneMaterialsModified);
 
 		ExtendShowFlagsMenu();
+
+		if (!IsRunningCommandlet())
+		{
+			ManualEditPanelManager = MakeUnique<FPCGManualEditPanelManager>();
+		}
 	}
 }
 
@@ -361,6 +384,19 @@ void FPCGEditorModule::OnPreExit()
 	UnregisterPinColorAndIcons();
 }
 
+void FPCGEditorModule::UpdateManualEditPanelVisibility()
+{
+	if (ManualEditPanelManager)
+	{
+		ManualEditPanelManager->UpdateManualEditPanelVisibility();
+	}
+}
+
+TSharedPtr<SPCGManualEditPanel> FPCGEditorModule::GetManualEditPanel() const
+{
+	return ManualEditPanelManager ? ManualEditPanelManager->GetManualEditPanel() : nullptr;
+}
+
 void FPCGEditorModule::OnSceneMaterialsModified()
 {
 	if (!GetDefault<UPCGEditorProjectSettings>()->bAutoRefreshGPUStaticMeshSpawners)
@@ -1021,6 +1057,51 @@ void FPCGEditorModule::UnregisterPCGDataVisualizations()
 	FPCGModule::GetMutablePCGDataVisualizationRegistry().InternalRegistry.Empty();
 }
 
+void FPCGEditorModule::RegisterDeltaVisualizations()
+{
+	DeltaVisualizationRegistry.InternalRegistry.Add(FPCGDeltaBase::StaticStruct(), MakeUnique<const FPCGDeltaBaseVisualization>());
+	DeltaVisualizationRegistry.InternalRegistry.Add(FPCGPointEditDelta::StaticStruct(), MakeUnique<const FPCGPointEditDeltaVisualization>());
+	DeltaVisualizationRegistry.InternalRegistry.Add(FPCGPointDeletionDelta::StaticStruct(), MakeUnique<const FPCGPointDeletionDeltaVisualization>());
+	DeltaVisualizationRegistry.InternalRegistry.Add(FPCGPointInsertionDelta::StaticStruct(), MakeUnique<const FPCGPointInsertionDeltaVisualization>());
+}
+
+void FPCGEditorModule::UnregisterDeltaVisualizations()
+{
+	DeltaVisualizationRegistry.InternalRegistry.Empty();
+}
+
+const FPCGDeltaVisualizationRegistry& FPCGEditorModule::GetConstDeltaVisualizationRegistry()
+{
+	return FModuleManager::GetModuleChecked<FPCGEditorModule>("PCGEditor").DeltaVisualizationRegistry;
+}
+
+FPCGDeltaVisualizationRegistry& FPCGEditorModule::GetMutableDeltaVisualizationRegistry()
+{
+	return FModuleManager::GetModuleChecked<FPCGEditorModule>("PCGEditor").DeltaVisualizationRegistry;
+}
+
+void FPCGEditorModule::RegisterDeltaViewportExtensions()
+{
+	DeltaViewportExtensionRegistry.RegisterInternalExtension(FPCGPointEditDelta::StaticStruct(), MakeUnique<FPCGPointEditViewportExtension>());
+	DeltaViewportExtensionRegistry.RegisterInternalExtension(FPCGPointDeletionDelta::StaticStruct(), MakeUnique<FPCGPointDeletionViewportExtension>());
+	DeltaViewportExtensionRegistry.RegisterInternalExtension(FPCGPointInsertionDelta::StaticStruct(), MakeUnique<FPCGPointInsertionViewportExtension>());
+}
+
+void FPCGEditorModule::UnregisterDeltaViewportExtensions()
+{
+	DeltaViewportExtensionRegistry.UnregisterAllInternalExtensions();
+}
+
+const FPCGDeltaViewportExtensionRegistry& FPCGEditorModule::GetConstDeltaViewportExtensionRegistry()
+{
+	return FModuleManager::GetModuleChecked<FPCGEditorModule>("PCGEditor").DeltaViewportExtensionRegistry;
+}
+
+FPCGDeltaViewportExtensionRegistry& FPCGEditorModule::GetMutableDeltaViewportExtensionRegistry()
+{
+	return FModuleManager::GetModuleChecked<FPCGEditorModule>("PCGEditor").DeltaViewportExtensionRegistry;
+}
+
 void FPCGEditorModule::RegisterTrackerFactories()
 {
 	FPCGChangeTrackingRegistry& TrackingRegistry = FPCGModule::GetMutableChangeTrackingRegistry();
diff --git a/Engine/Plugins/PCG/Source/PCGEditor/Private/PCGEditorStyle.cpp b/Engine/Plugins/PCG/Source/PCGEditor/Private/PCGEditorStyle.cpp
--- a/Engine/Plugins/PCG/Source/PCGEditor/Private/PCGEditorStyle.cpp
+++ b/Engine/Plugins/PCG/Source/PCGEditor/Private/PCGEditorStyle.cpp
@@ -11,6 +11,7 @@
 #include "Styling/SlateStyleMacros.h"
 #include "Styling/SlateStyleRegistry.h"
 #include "Styling/SlateTypes.h"
+#include "Styling/StyleColors.h"
 
 void FPCGEditorStyle::Register()
 {
@@ -35,6 +36,9 @@ FPCGEditorStyle::FPCGEditorStyle() : FSlateStyleSet("PCGEditorStyle")
 
 	Set("PCG.NodeOverlay.Debug", new CORE_IMAGE_BRUSH_SVG("Starship/Blueprints/Breakpoint_Valid", Icon20x20, FSlateColor(FColor::Cyan)));
 	Set("PCG.NodeOverlay.Inspect", new CORE_IMAGE_BRUSH_SVG("Starship/Blueprints/Breakpoint_Valid", Icon20x20, FSlateColor(FColor::Orange)));
+	Set("PCG.NodeOverlay.ManualEditing", new CORE_IMAGE_BRUSH_SVG("Starship/Blueprints/Breakpoint_Valid", Icon20x20, FSlateColor(FColor::Green)));
+	Set("PCG.NodeOverlay.HasOverrides", new CORE_IMAGE_BRUSH_SVG("Starship/Blueprints/Breakpoint_Valid", Icon20x20, FSlateColor(FColor(160, 32, 240))));
+	Set("PCG.NodeOverlay.MarkedForEditing", new CORE_IMAGE_BRUSH_SVG("Starship/Blueprints/Breakpoint_Valid", Icon20x20, FSlateColor(FColor(100, 200, 255))));
 	Set("PCG.NodeOverlay.OnInactiveBranch", new CORE_IMAGE_BRUSH_SVG("Starship/Blueprints/CompileStatus_Fail_Badge", Icon20x20, FSlateColor(FColor::White)));
 	Set("PCG.NodeOverlay.GPUUpload", new CORE_IMAGE_BRUSH("Common/UpArrow", Icon16x16, FSlateColor(FColor::Yellow)));
 	Set("PCG.NodeOverlay.GPUReadback", new CORE_IMAGE_BRUSH("Common/DownArrow", Icon16x16, FSlateColor(FColor::Yellow)));
@@ -153,6 +157,17 @@ FPCGEditorStyle::FPCGEditorStyle() : FSlateStyleSet("PCGEditorStyle")
 	Set("PCG.Editor.AlternatePause", new IMAGE_BRUSH_SVG("Style/PCG_Editor_AlternatePause", Icon16x16));
 	Set("PCG.Editor.JumpTo", new IMAGE_BRUSH_SVG("Style/PCG_DebugTree_JumpTo", Icon16x16));
 
+	// Viewport overlay styles
+	{
+		FLinearColor ViewportOverlayBackgroundColor = FStyleColors::Foreground.GetSpecifiedColor();
+		ViewportOverlayBackgroundColor.A = 0.4f;
+		Set("PCG.OverlayBackgroundBrush", new FSlateRoundedBoxBrush(ViewportOverlayBackgroundColor, 8.0f, FStyleColors::Dropdown, 1.0f));
+
+		FLinearColor InnerPanelColor = FStyleColors::Panel.GetSpecifiedColor();
+		InnerPanelColor.A = 0.85f;
+		Set("PCG.OverlayInnerPanelBrush", new FSlateRoundedBoxBrush(InnerPanelColor, 6.0f));
+	}
+
 	// Text editor styles
 	{
 		Set("TextEditor.Border", new BOX_BRUSH("Images/TextEditorBorder", FMargin(4.0f / 16.0f), FLinearColor(0.02f, 0.02f, 0.02f, 1)));
diff --git a/Engine/Plugins/PCG/Source/PCGEditor/Private/PCGManualEditPanelManager.cpp b/Engine/Plugins/PCG/Source/PCGEditor/Private/PCGManualEditPanelManager.cpp
new file mode 100644
--- /dev/null
+++ b/Engine/Plugins/PCG/Source/PCGEditor/Private/PCGManualEditPanelManager.cpp
@@ -0,0 +1,315 @@
+// Copyright Epic Games, Inc. All Rights Reserved.
+
+#include "PCGManualEditPanelManager.h"
+
+#include "PCGComponent.h"
+#include "PCGComponentVisualizer.h"
+#include "PCGEditor.h"
+#include "PCGEditorGraph.h"
+#include "PCGEditorModule.h"
+#include "PCGGraph.h"
+#include "PCGGraphExecutionStateInterface.h"
+#include "PCGNode.h"
+#include "PCGSettings.h"
+#include "SPCGManualEditPanel.h"
+#include "Graph/DataOverride/PCGDataOverride.h"
+#include "Nodes/PCGEditorGraphNodeBase.h"
+
+#include "Editor.h"
+#include "LevelEditor.h"
+#include "Selection.h"
+#include "SLevelViewport.h"
+#include "UnrealEdGlobals.h"
+#include "Editor/UnrealEdEngine.h"
+#include "Modules/ModuleManager.h"
+
+namespace PCG::ManualEditPanel
+{
+	FPCGComponentVisualizer* FindPCGComponentVisualizer()
+	{
+		if (!GUnrealEd)
+		{
+			return nullptr;
+		}
+
+		TSharedPtr<FComponentVisualizer> Visualizer = GUnrealEd->FindComponentVisualizer(UPCGComponent::StaticClass()->GetFName());
+		return Visualizer.IsValid() ? static_cast<FPCGComponentVisualizer*>(Visualizer.Get()) : nullptr;
+	}
+}
+
+FPCGManualEditPanelManager::FPCGManualEditPanelManager()
+{
+	SelectionChangedDelegateHandle = USelection::SelectionChangedEvent.AddRaw(this, &FPCGManualEditPanelManager::OnEditorSelectionChanged);
+}
+
+FPCGManualEditPanelManager::~FPCGManualEditPanelManager()
+{
+	HideManualEditPanel();
+	USelection::SelectionChangedEvent.Remove(SelectionChangedDelegateHandle);
+}
+
+void FPCGManualEditPanelManager::OnEditorSelectionChanged(UObject* Object)
+{
+	if (!GEditor)
+	{
+		return;
+	}
+
+	USelection* Selection = GEditor->GetSelectedActors();
+	if (!Selection)
+	{
+		HideManualEditPanel();
+		return;
+	}
+
+	// Find a PCG component on the selected actor
+	UPCGComponent* FoundComponent = nullptr;
+	for (FSelectionIterator It(*Selection); It; ++It)
+	{
+		const AActor* Actor = Cast<AActor>(*It);
+		if (!Actor)
+		{
+			continue;
+		}
+
+		TArray<UPCGComponent*> PCGComponents;
+		Actor->GetComponents<UPCGComponent>(PCGComponents, /*bIncludeFromChildActors=*/true);
+
+		for (UPCGComponent* Component : PCGComponents)
+		{
+			if (Component && Component->GetGraph())
+			{
+				FoundComponent = Component;
+				break;
+			}
+		}
+
+		if (FoundComponent)
+		{
+			break;
+		}
+	}
+
+	if (!FoundComponent || !FoundComponent->GetGraph())
+	{
+		HideManualEditPanel();
+		return;
+	}
+
+	// Update tracked component if it changed
+	if (TrackedPCGComponent != FoundComponent)
+	{
+		HideManualEditPanel();
+		TrackedPCGComponent = FoundComponent;
+	}
+
+	UpdateManualEditPanelVisibility();
+}
+
+void FPCGManualEditPanelManager::UpdateManualEditPanelVisibility()
+{
+	UPCGComponent* Component = TrackedPCGComponent.Get();
+	if (!Component || !IsValid(Component))
+	{
+		HideManualEditPanel();
+		return;
+	}
+
+	UPCGGraph* Graph = Component->GetGraph();
+	if (!Graph)
+	{
+		HideManualEditPanel();
+		return;
+	}
+
+	// Check if any nodes are marked or actively editing
+	bool bPanelShouldBeVisible = false;
+	for (const UPCGNode* Node : Graph->GetNodes())
+	{
+		if (Node && Node->GetSettingsInterface())
+		{
+			if (Node->GetSettingsInterface()->bMarkedForManualEditing || Node->GetSettingsInterface()->bTemporaryManualEditingEnabled)
+			{
+				bPanelShouldBeVisible = true;
+				break;
+			}
+		}
+	}
+
+	if (!bPanelShouldBeVisible)
+	{
+		HideManualEditPanel();
+		return;
+	}
+
+	// Ensure the component has inspection data for the visualizer. If the graph was executed before any node was marked
+	// for viewport editing, the inspection data may not exist yet.
+	FPCGGraphExecutionInspection& Inspection = Component->GetExecutionState().GetInspection();
+	if (!Inspection.IsInspecting())
+	{
+		Inspection.EnableInspection();
+
+		IPCGGraphExecutionState::FGenerateParams GenerateParams;
+		GenerateParams.bEvenIfAlreadyGenerated = true;
+		Component->GetExecutionState().Generate(GenerateParams);
+	}
+
+	if (ManualEditPanel.IsValid())
+	{
+		ManualEditPanel->RefreshNodeList();
+	}
+	else
+	{
+		ShowManualEditPanel(Graph);
+	}
+}
+
+void FPCGManualEditPanelManager::ShowManualEditPanel(UPCGGraph* InGraph)
+{
+	if (ManualEditPanel.IsValid())
+	{
+		return;
+	}
+
+	ManualEditPanel = SNew(SPCGManualEditPanel).Graph(InGraph);
+	ManualEditOverlayWidget = ManualEditPanel->CreateOverlayWidget();
+
+	// Wire generic actions from the visualizer
+	TWeakPtr<SPCGManualEditPanel> WeakPanel = ManualEditPanel;
+
+	FPCGDeltaViewportCallbacks ViewportCallbacks;
+
+	ViewportCallbacks.RevertDelta = [](const FPCGDeltaKey& Key)
+	{
+		if (FPCGComponentVisualizer* Visualizer = PCG::ManualEditPanel::FindPCGComponentVisualizer())
+		{
+			Visualizer->RestoreDelta(Key);
+		}
+	};
+
+	ViewportCallbacks.RevertAllRestorableDeltas = []()
+	{
+		if (FPCGComponentVisualizer* Visualizer = PCG::ManualEditPanel::FindPCGComponentVisualizer())
+		{
+			Visualizer->RestoreAllDeltas();
+		}
+	};
+
+	ViewportCallbacks.ClearSelection = []()
+	{
+		if (FPCGComponentVisualizer* Visualizer = PCG::ManualEditPanel::FindPCGComponentVisualizer())
+		{
+			Visualizer->ClearActiveSelection();
+		}
+	};
+
+	ViewportCallbacks.SelectElement = [](const FPCGDeltaKey& DeltaKey, int32 SubIndex)
+	{
+		if (FPCGComponentVisualizer* Visualizer = PCG::ManualEditPanel::FindPCGComponentVisualizer())
+		{
+			Visualizer->SelectElement(DeltaKey, SubIndex);
+		}
+	};
+
+	ViewportCallbacks.RequestListRefresh = [WeakPanel]()
+	{
+		if (TSharedPtr<SPCGManualEditPanel> Panel = WeakPanel.Pin())
+		{
+			Panel->RefreshActiveExtensionLists();
+		}
+	};
+
+	ManualEditPanel->SetVisualizerActions(ViewportCallbacks);
+
+	// When a node is selected in the panel, 'inspect' it in the graph editor's Attributes tab
+	ManualEditPanel->OnNodeSelected.BindLambda([WeakGraph = TWeakObjectPtr<UPCGGraph>(InGraph)](const UPCGNode* SelectedNode)
+	{
+		const UPCGGraph* Graph = WeakGraph.Get();
+		if (!SelectedNode || !Graph || !Graph->GetEditorGraph())
+		{
+			return;
+		}
+
+		const TSharedPtr<FPCGEditor> Editor = Graph->GetEditorGraph()->GetEditor().Pin();
+		if (!Editor.IsValid())
+		{
+			return;
+		}
+
+		if (UPCGEditorGraphNodeBase* EditorNode = Graph->GetEditorGraph()->GetEditorNodeFromPCGNode(SelectedNode))
+		{
+			Editor->SetNodeInspected(EditorNode, /*bValue=*/true);
+		}
+	});
+
+	// Jump-to-node: focus the graph editor on the selected node
+	ManualEditPanel->OnJumpToNode.BindLambda([WeakGraph = TWeakObjectPtr<UPCGGraph>(InGraph)](const UPCGNode* Node)
+	{
+		const UPCGGraph* Graph = WeakGraph.Get();
+		if (!Node || !Graph || !Graph->GetEditorGraph())
+		{
+			return;
+		}
+
+		TSharedPtr<FPCGEditor> Editor = Graph->GetEditorGraph()->GetEditor().Pin();
+		if (Editor.IsValid())
+		{
+			Editor->JumpToNode(Node);
+		}
+	});
+
+	FLevelEditorModule& LevelEditorModule = FModuleManager::GetModuleChecked<FLevelEditorModule>("LevelEditor");
+	TSharedPtr<ILevelEditor> LevelEditor = LevelEditorModule.GetFirstLevelEditor();
+	if (LevelEditor.IsValid())
+	{
+		TSharedPtr<SLevelViewport> ActiveVP = LevelEditor->GetActiveViewportInterface();
+		if (ActiveVP.IsValid())
+		{
+			ActiveVP->AddOverlayWidget(ManualEditOverlayWidget.ToSharedRef());
+			OverlayViewport = ActiveVP;
+		}
+		else
+		{
+			UE_LOG(LogPCGEditor, Warning, TEXT("[ManualEdit] ShowManualEditPanel: No active viewport"));
+		}
+	}
+	else
+	{
+		UE_LOG(LogPCGEditor, Warning, TEXT("[ManualEdit] ShowManualEditPanel: No LevelEditor found"));
+	}
+}
+
+void FPCGManualEditPanelManager::HideManualEditPanel()
+{
+	if (ManualEditOverlayWidget.IsValid())
+	{
+		TSharedPtr<SLevelViewport> Viewport = OverlayViewport.Pin();
+		if (Viewport.IsValid())
+		{
+			Viewport->RemoveOverlayWidget(ManualEditOverlayWidget.ToSharedRef());
+		}
+
+		OverlayViewport.Reset();
+	}
+
+	const UPCGComponent* Component = TrackedPCGComponent.Get();
+	if (const UPCGGraph* PCGGraph = Component ? Component->GetGraph() : nullptr)
+	{
+		for (UPCGNode* Node : PCGGraph->GetNodes())
+		{
+			if (UPCGSettingsInterface* Settings = Node ? Node->GetSettingsInterface() : nullptr)
+			{
+				if (Settings->bTemporaryManualEditingEnabled)
+				{
+					Settings->bTemporaryManualEditingEnabled = false;
+					// Broadcast through the PCGNode to trigger a full node reconstruct. The editor graph node is not available here.
+					Node->OnNodeChangedDelegate.Broadcast(Node, EPCGChangeType::Cosmetic);
+				}
+			}
+		}
+	}
+
+	ManualEditPanel.Reset();
+	ManualEditOverlayWidget.Reset();
+	TrackedPCGComponent.Reset();
+}
diff --git a/Engine/Plugins/PCG/Source/PCGEditor/Private/PCGManualEditPanelManager.h b/Engine/Plugins/PCG/Source/PCGEditor/Private/PCGManualEditPanelManager.h
new file mode 100644
--- /dev/null
+++ b/Engine/Plugins/PCG/Source/PCGEditor/Private/PCGManualEditPanelManager.h
@@ -0,0 +1,45 @@
+// Copyright Epic Games, Inc. All Rights Reserved.
+
+#pragma once
+
+#include "Delegates/IDelegateInstance.h"
+#include "Templates/SharedPointer.h"
+#include "UObject/WeakObjectPtr.h"
+
+class SLevelViewport;
+class SPCGManualEditPanel;
+class UPCGComponent;
+class UPCGGraph;
+class SWidget;
+class UObject;
+
+/**
+ * Manages the lifecycle of the manual edit panel viewport overlay, including selection tracking, panel
+ * creation/destruction, delegates, etc.
+ */
+class FPCGManualEditPanelManager
+{
+public:
+	FPCGManualEditPanelManager();
+	~FPCGManualEditPanelManager();
+
+	FPCGManualEditPanelManager(const FPCGManualEditPanelManager&) = delete;
+	FPCGManualEditPanelManager& operator=(const FPCGManualEditPanelManager&) = delete;
+
+	/** Returns the manual edit panel if one is currently active. */
+	TSharedPtr<SPCGManualEditPanel> GetManualEditPanel() const { return ManualEditPanel; }
+
+	/** Checks the tracked component and shows/hides the panel as needed. */
+	void UpdateManualEditPanelVisibility();
+
+private:
+	void OnEditorSelectionChanged(UObject* Object);
+	void ShowManualEditPanel(UPCGGraph* InGraph);
+	void HideManualEditPanel();
+
+	TSharedPtr<SPCGManualEditPanel> ManualEditPanel;
+	TSharedPtr<SWidget> ManualEditOverlayWidget;
+	TWeakPtr<SLevelViewport> OverlayViewport;
+	FDelegateHandle SelectionChangedDelegateHandle;
+	TWeakObjectPtr<UPCGComponent> TrackedPCGComponent;
+};
diff --git a/Engine/Plugins/PCG/Source/PCGEditor/Private/SPCGManualEditPanel.cpp b/Engine/Plugins/PCG/Source/PCGEditor/Private/SPCGManualEditPanel.cpp
new file mode 100644
--- /dev/null
+++ b/Engine/Plugins/PCG/Source/PCGEditor/Private/SPCGManualEditPanel.cpp
@@ -0,0 +1,712 @@
+// Copyright Epic Games, Inc. All Rights Reserved.
+
+#include "SPCGManualEditPanel.h"
+
+#include "PCGCommon.h"
+#include "PCGComponent.h"
+#include "PCGGraph.h"
+#include "PCGNode.h"
+#include "PCGSettings.h"
+
+#include "PCGComponentVisualizer.h"
+#include "PCGEditorModule.h"
+#include "PCGEditorStyle.h"
+
+#include "Editor.h"
+#include "UnrealEdGlobals.h"
+#include "Editor/UnrealEdEngine.h"
+#include "Overlay/SDraggableBoxOverlay.h"
+#include "Styling/AppStyle.h"
+#include "Widgets/Images/SImage.h"
+#include "Widgets/Input/SButton.h"
+#include "Widgets/Input/SCheckBox.h"
+#include "Widgets/Input/SComboBox.h"
+#include "Widgets/Input/SSpinBox.h"
+#include "Widgets/Layout/SBackgroundBlur.h"
+#include "Widgets/Layout/SBorder.h"
+#include "Widgets/Layout/SBox.h"
+#include "Widgets/Layout/SScrollBox.h"
+#include "Widgets/Layout/SSeparator.h"
+#include "Widgets/Text/STextBlock.h"
+#include "Widgets/Views/SListView.h"
+
+#define LOCTEXT_NAMESPACE "SPCGManualEditPanel"
+
+void SPCGManualEditPanel::Construct(const FArguments& InArgs)
+{
+	Graph = InArgs._Graph;
+
+	// Initialize delta type options
+	DeltaTypeOptions = FPCGEditorModule::GetConstDeltaViewportExtensionRegistry().GetRegisteredDeltaTypes();
+
+	ChildSlot
+	[
+		SNew(SVerticalBox)
+		// Header with graph name
+		+ SVerticalBox::Slot()
+	   .AutoHeight()
+	   .Padding(4.0f, 2.0f, 4.0f, 4.0f)
+		[
+			SNew(STextBlock)
+		   .Text_Lambda([this]() -> FText
+			{
+				if (UPCGGraph* PCGGraph = Graph.Get())
+				{
+					return FText::Format(LOCTEXT("ManualEditHeaderFormat", "Viewport Editing ({0})"), FText::FromString(PCGGraph->GetName()));
+				}
+				else
+				{
+					return LOCTEXT("ManualEditHeaderNoGraph", "Viewport Editing");
+				}
+			})
+		   .TextStyle(FAppStyle::Get(), "NormalText.Important")
+		]
+		// Node list
+		+ SVerticalBox::Slot()
+	   .FillHeight(1.0f)
+	   .Padding(2.0f)
+		[
+			SAssignNew(NodeListView, SListView<TSharedPtr<FPCGManualEditNodeEntry>>)
+		   .ListItemsSource(&NodeEntries)
+		   .OnGenerateRow(this, &SPCGManualEditPanel::OnGenerateRow)
+		   .OnSelectionChanged(this, &SPCGManualEditPanel::OnListSelectionChanged)
+		   .SelectionMode(ESelectionMode::SingleToggle)
+		]
+		// Separator
+		+ SVerticalBox::Slot()
+	   .AutoHeight()
+	   .Padding(4.0f, 2.0f)
+		[
+			SNew(SSeparator)
+		   .Visibility_Lambda([this]()
+			{
+				return GetActiveNode() != nullptr ? EVisibility::Visible : EVisibility::Collapsed;
+			})
+		]
+		// Configuration section (shown when a node is selected)
+		+ SVerticalBox::Slot()
+	   .AutoHeight()
+	   .Padding(4.0f, 6.0f)
+		[
+			SAssignNew(ConfigurationSection, SVerticalBox)
+		   .Visibility(EVisibility::Collapsed)
+			// Delta Type row
+			+ SVerticalBox::Slot()
+		   .AutoHeight()
+		   .Padding(0.0f, 2.0f)
+			[
+				SNew(SHorizontalBox)
+				+ SHorizontalBox::Slot()
+			   .AutoWidth()
+			   .VAlign(VAlign_Center)
+			   .Padding(0.0f, 0.0f, 4.0f, 0.0f)
+				[
+					SNew(STextBlock)
+				   .Text(LOCTEXT("DeltaTypeLabel", "Delta Type:"))
+				   .Font(FCoreStyle::GetDefaultFontStyle("Regular", 8))
+				]
+				+ SHorizontalBox::Slot()
+			   .FillWidth(1.0f)
+			   .VAlign(VAlign_Center)
+				[
+					SNew(SComboButton)
+				   .OnGetMenuContent(this, &SPCGManualEditPanel::GenerateDeltaTypeMenu)
+				   .ButtonContent()
+					[
+						SNew(STextBlock)
+					   .Text(this, &SPCGManualEditPanel::GetCurrentDeltaTypeLabel)
+					   .Font(FCoreStyle::GetDefaultFontStyle("Regular", 8))
+					]
+				]
+			]
+			// PRS signature checkboxes
+			+ SVerticalBox::Slot()
+		   .AutoHeight()
+			[
+				SAssignNew(SignatureSlot, SVerticalBox)
+			   .Visibility(EVisibility::Collapsed)
+			]
+			// Dynamic extension widget slot
+			+ SVerticalBox::Slot()
+		   .AutoHeight()
+			[
+				SAssignNew(DeltaTypeWidgetSlot, SVerticalBox)
+			]
+		]
+	];
+
+	RefreshNodeList();
+}
+
+void SPCGManualEditPanel::RefreshNodeList()
+{
+	// Preserve current selection
+	const UPCGNode* PreviouslySelectedNode = GetActiveNode();
+
+	NodeEntries.Empty();
+
+	const UPCGGraph* PCGGraph = Graph.Get();
+	if (!PCGGraph)
+	{
+		if (NodeListView.IsValid())
+		{
+			NodeListView->RequestListRefresh();
+		}
+		return;
+	}
+
+	// @todo_pcg: GetNodes only returns top-level nodes.Investigate for subgraph support.
+	const TArray<UPCGNode*>& Nodes = PCGGraph->GetNodes();
+
+	const UPCGNode* CurrentActiveNode = GetActiveNode();
+
+	for (const UPCGNode* Node : Nodes)
+	{
+		if (!Node || !Node->GetSettingsInterface())
+		{
+			continue;
+		}
+
+		const UPCGSettingsInterface* Settings = Node->GetSettingsInterface();
+
+		const bool bIsMarked = Settings->bMarkedForManualEditing;
+		const bool bIsTemporarilyMarked = Settings->bTemporaryManualEditingEnabled;
+
+		// Only show nodes that are marked for editing or currently V-key activated
+		if (!bIsMarked && !bIsTemporarilyMarked)
+		{
+			continue;
+		}
+
+		TSharedPtr<FPCGManualEditNodeEntry> Entry = MakeShared<FPCGManualEditNodeEntry>();
+		Entry->Node = Node;
+		Entry->bIsMarked = bIsMarked;
+		Entry->bIsTemporary = bIsTemporarilyMarked && !bIsMarked;
+		Entry->bIsActive = (Node == CurrentActiveNode) || Settings->bTemporaryManualEditingEnabled;
+
+		Entry->DisplayName = Node->GetNodeTitle(EPCGNodeTitleType::ListView);
+
+		NodeEntries.Add(Entry);
+	}
+
+	if (NodeListView.IsValid())
+	{
+		NodeListView->RequestListRefresh();
+
+		// Prefer selecting the active node. Fall back to restoring previous selection otherwise.
+		TSharedPtr<FPCGManualEditNodeEntry> EntryToSelect;
+		for (const TSharedPtr<FPCGManualEditNodeEntry>& Entry : NodeEntries)
+		{
+			if (Entry->bIsActive)
+			{
+				EntryToSelect = Entry;
+				break; // Active node takes priority.
+			}
+
+			if (!EntryToSelect && PreviouslySelectedNode && Entry->Node == PreviouslySelectedNode)
+			{
+				EntryToSelect = Entry;
+				// Continue and keep scanning for an active entry.
+			}
+		}
+
+		if (EntryToSelect)
+		{
+			NodeListView->SetSelection(EntryToSelect, ESelectInfo::Direct);
+		}
+	}
+
+	UpdateConfigurationSection();
+}
+
+void SPCGManualEditPanel::OnListSelectionChanged(TSharedPtr<FPCGManualEditNodeEntry> SelectedEntry, const ESelectInfo::Type SelectInfo)
+{
+	if (SelectInfo == ESelectInfo::Direct)
+	{
+		return;
+	}
+
+	const UPCGGraph* PCGGraph = Graph.Get();
+	if (!PCGGraph)
+	{
+		return;
+	}
+
+	// SingleToggle fires with an invalid entry when the user clicks an already-selected node to deselect it.
+	const bool bIsDeselection = !SelectedEntry.IsValid() || !SelectedEntry->Node.IsValid();
+
+	// Clear any active visualizer selection before switching nodes, otherwise the proxy/gizmo from the old node
+	// persists and attaches to the newly selected node's data.
+	if (GUnrealEd)
+	{
+		if (TSharedPtr<FComponentVisualizer> Visualizer = GUnrealEd->FindComponentVisualizer(UPCGComponent::StaticClass()->GetFName()))
+		{
+			if (FPCGComponentVisualizer* PCGVisualizer = static_cast<FPCGComponentVisualizer*>(Visualizer.Get()))
+			{
+				PCGVisualizer->ClearActiveSelection();
+			}
+		}
+	}
+
+	// Activate the selected node
+	if (!bIsDeselection)
+	{
+		OnNodeSelected.ExecuteIfBound(SelectedEntry->Node.Get());
+	}
+
+	// Refresh entries to update active state
+	RefreshNodeList();
+
+	// Redraw viewports to trigger visualizer update
+	if (GEditor)
+	{
+		GEditor->RedrawAllViewports();
+	}
+}
+
+TSharedRef<ITableRow> SPCGManualEditPanel::OnGenerateRow(TSharedPtr<FPCGManualEditNodeEntry> Entry, const TSharedRef<STableViewBase>& OwnerTable)
+{
+	// Active node: green, Marked but not active: light blue, Unmarked: default
+	FLinearColor TextColor = FLinearColor::White;
+	if (Entry->bIsActive)
+	{
+		TextColor = FLinearColor::Green;
+	}
+	else if (Entry->bIsMarked)
+	{
+		TextColor = FLinearColor(0.4f, 0.7f, 1.0f); // Light blue to match badge
+	}
+
+	TSharedPtr<STableRow<TSharedPtr<FPCGManualEditNodeEntry>>> Row;
+
+	SAssignNew(Row, STableRow<TSharedPtr<FPCGManualEditNodeEntry>>, OwnerTable)
+	[
+		SNew(SHorizontalBox)
+		// Node name
+		+ SHorizontalBox::Slot()
+	   .FillWidth(1.0f)
+	   .VAlign(VAlign_Center)
+	   .Padding(4.0f, 2.0f)
+		[
+			SNew(STextBlock)
+		   .Text(Entry->DisplayName)
+		   .ColorAndOpacity(TextColor)
+		]
+		// Jump to node in graph editor
+		+ SHorizontalBox::Slot()
+	   .AutoWidth()
+	   .VAlign(VAlign_Center)
+	   .Padding(2.0f, 0.0f)
+		[
+			SNew(SButton)
+		   .ButtonStyle(FAppStyle::Get(), "SimpleButton")
+		   .ContentPadding(FMargin(1, 0))
+		   .ToolTipText(LOCTEXT("JumpToNodeTooltip", "Focus on this node in the graph editor"))
+		   .Visibility_Lambda([WeakRow = TWeakPtr<SWidget>(Row)]()
+			{
+				TSharedPtr<SWidget> PinnedRow = WeakRow.Pin();
+				return (PinnedRow.IsValid() && PinnedRow->IsHovered()) ? EVisibility::Visible : EVisibility::Hidden;
+			})
+		   .OnClicked_Lambda([this, WeakNode = Entry->Node]() -> FReply
+			{
+				if (const UPCGNode* Node = WeakNode.Get())
+				{
+					OnJumpToNode.ExecuteIfBound(Node);
+				}
+				return FReply::Handled();
+			})
+			[
+				SNew(SImage)
+			   .Image(FPCGEditorStyle::Get().GetBrush("PCG.Editor.JumpTo"))
+			   .ColorAndOpacity(FSlateColor::UseForeground())
+			]
+		]
+		// Active indicator
+		+ SHorizontalBox::Slot()
+	   .AutoWidth()
+	   .VAlign(VAlign_Center)
+	   .Padding(2.0f, 0.0f)
+		[
+			SNew(STextBlock)
+		   .Text(Entry->bIsActive ? LOCTEXT("ActiveIndicator", "LIVE") : FText::GetEmpty())
+		   .ColorAndOpacity(FLinearColor::Green)
+		   .Font(FCoreStyle::GetDefaultFontStyle("Bold", 8))
+		]
+	];
+
+	return Row.ToSharedRef();
+}
+
+TSharedRef<SWidget> SPCGManualEditPanel::CreateOverlayWidget()
+{
+	TSharedRef<SPCGManualEditPanel> Self = SharedThis(this);
+
+	SAssignNew(OverlayWidget, UE::ToolWidgets::SDraggableBoxOverlay)
+   .HAlign(HAlign_Right)
+   .VAlign(VAlign_Top)
+   .InitialAlignmentOffset(FVector2f(20.0f, 60.0f))
+   .IsDraggable(true)
+	[
+		SNew(SBackgroundBlur)
+	   .BlurStrength(8.0f)
+	   .CornerRadius(FVector4(8.0f, 8.0f, 8.0f, 8.0f))
+		[
+			SNew(SBorder)
+		   .BorderImage(FPCGEditorStyle::Get().GetBrush("PCG.OverlayBackgroundBrush"))
+		   .Padding(4.0f)
+			[
+				SNew(SBorder)
+			   .BorderImage(FPCGEditorStyle::Get().GetBrush("PCG.OverlayInnerPanelBrush"))
+			   .Padding(6.0f)
+				[
+					SNew(SBox)
+				   .MinDesiredWidth(250.0f)
+				   .MinDesiredHeight(120.0f)
+				   .MaxDesiredWidth(350.0f)
+				   .MaxDesiredHeight(500.0f)
+					[
+						Self
+					]
+				]
+			]
+		]
+	];
+
+	return StaticCastSharedRef<SWidget>(OverlayWidget.ToSharedRef());
+}
+
+TSharedPtr<SWidget> SPCGManualEditPanel::GetOverlayWidget() const
+{
+	return OverlayWidget;
+}
+
+TSharedPtr<const FPCGManualEditNodeConfiguration> SPCGManualEditPanel::GetConfigurationForActiveNode() const
+{
+	const UPCGNode* ActiveNode = GetActiveNode();
+	if (!ActiveNode)
+	{
+		return nullptr;
+	}
+
+	const TSharedPtr<FPCGManualEditNodeConfiguration>* Found = NodeConfigurationMap.Find(TWeakObjectPtr<const UPCGNode>(ActiveNode));
+	return Found ? *Found : nullptr;
+}
+
+TSharedPtr<FPCGManualEditNodeConfiguration> SPCGManualEditPanel::GetOrCreateConfigurationForNode(const TWeakObjectPtr<const UPCGNode> Node)
+{
+	TSharedPtr<FPCGManualEditNodeConfiguration>& Configuration = NodeConfigurationMap.FindOrAdd(Node);
+	if (!Configuration)
+	{
+		Configuration = MakeShared<FPCGManualEditNodeConfiguration>();
+	}
+
+	if (!Configuration->DeltaType && DeltaTypeOptions.Num() > 0)
+	{
+		Configuration->DeltaType = DeltaTypeOptions[0];
+	}
+
+	return Configuration;
+}
+
+const UPCGNode* SPCGManualEditPanel::GetActiveNode() const
+{
+	if (!NodeListView.IsValid())
+	{
+		return nullptr;
+	}
+
+	TArray<TSharedPtr<FPCGManualEditNodeEntry>> SelectedItems = NodeListView->GetSelectedItems();
+	if (SelectedItems.Num() > 0 && SelectedItems[0].IsValid() && SelectedItems[0]->Node.IsValid())
+	{
+		return SelectedItems[0]->Node.Get();
+	}
+
+	return nullptr;
+}
+
+void SPCGManualEditPanel::SetDeltaContext(UPCGComponent* InComponent, const FPCGSourceDataStorageKey& InStorageKey)
+{
+	ActivePCGComponent = InComponent;
+	ActiveStorageKey = InStorageKey;
+
+	if (ActiveExtension)
+	{
+		FPCGDeltaViewportContext Context = BuildExtensionContext();
+		ActiveExtension->RefreshLists(Context);
+	}
+}
+
+void SPCGManualEditPanel::SetSelectionState(
+	const bool bHasSelection,
+	const FPCGDeltaKey& SelectedKey,
+	const FTransform& InSelectionTransform,
+	const int32 InSelectedElementIndex,
+	const int32 InOriginalElementIndex)
+{
+	bSelectionActive = bHasSelection;
+	SelectionDeltaKey = SelectedKey;
+	SelectionTransform = InSelectionTransform;
+	SelectedElementIndex = InSelectedElementIndex;
+	OriginalElementIndex = InOriginalElementIndex;
+
+	if (ActiveExtension)
+	{
+		ActiveExtension->UpdateContext(BuildExtensionContext());
+	}
+}
+
+void SPCGManualEditPanel::RefreshActiveExtensionLists()
+{
+	if (ActiveExtension)
+	{
+		ActiveExtension->RefreshLists(BuildExtensionContext());
+	}
+}
+
+void SPCGManualEditPanel::SetVisualizerActions(const FPCGDeltaViewportCallbacks& InActions)
+{
+	VisualizerCallbacks = InActions;
+}
+
+void SPCGManualEditPanel::UpdateConfigurationSection()
+{
+	if (!ConfigurationSection.IsValid())
+	{
+		return;
+	}
+
+	if (const UPCGNode* ActiveNode = GetActiveNode())
+	{
+		// Ensure configuration exists for this node
+		GetOrCreateConfigurationForNode(ActiveNode);
+
+		ConfigurationSection->SetVisibility(EVisibility::SelfHitTestInvisible);
+
+		// Look up the extension for the current delta type
+		TSharedPtr<const FPCGManualEditNodeConfiguration> Configuration = GetConfigurationForActiveNode();
+		IPCGDeltaViewportExtension* Extension = nullptr;
+		if (Configuration && Configuration->DeltaType)
+		{
+			Extension = FPCGEditorModule::GetMutableDeltaViewportExtensionRegistry().GetExtension(Configuration->DeltaType);
+		}
+
+		// Swap extension widget if the type changed
+		if (Extension != ActiveExtension)
+		{
+			ActiveExtension = Extension;
+
+			// Signature checkboxes + tolerance are always visible when a node is selected.
+			if (SignatureSlot.IsValid())
+			{
+				SignatureSlot->ClearChildren();
+				SignatureSlot->SetVisibility(EVisibility::SelfHitTestInvisible);
+				SignatureSlot->AddSlot()
+				   .AutoHeight()
+				   .Padding(0.0f, 2.0f, 0.0f, 0.0f)
+					[
+						SNew(SHorizontalBox)
+						+ SHorizontalBox::Slot()
+					   .AutoWidth()
+					   .VAlign(VAlign_Center)
+						[
+							SNew(SBox)
+						   .MinDesiredWidth(60.0f)
+							[
+								SNew(STextBlock)
+							   .Text(LOCTEXT("SignatureLabel", "Signature:"))
+							   .Font(FCoreStyle::GetDefaultFontStyle("Regular", 8))
+							]
+						]
+						+ SHorizontalBox::Slot()
+					   .AutoWidth()
+					   .VAlign(VAlign_Center)
+					   .Padding(0.0f, 0.0f, 4.0f, 0.0f)
+						[
+							SNew(SCheckBox)
+						   .IsChecked_Lambda([this]() -> ECheckBoxState
+							{
+								TSharedPtr<const FPCGManualEditNodeConfiguration> NodeConfiguration = GetConfigurationForActiveNode();
+								return (NodeConfiguration && NodeConfiguration->bSignaturePosition) ? ECheckBoxState::Checked : ECheckBoxState::Unchecked;
+							})
+						   .OnCheckStateChanged_Lambda([this](const ECheckBoxState NewState)
+							{
+								if (const UPCGNode* Node = GetActiveNode())
+								{
+									GetOrCreateConfigurationForNode(Node)->bSignaturePosition = (NewState == ECheckBoxState::Checked);
+								}
+							})
+							[
+								SNew(STextBlock)
+							   .Text(LOCTEXT("SignaturePositionCheckbox", "Pos"))
+							   .Font(FCoreStyle::GetDefaultFontStyle("Regular", 8))
+							]
+						]
+						+ SHorizontalBox::Slot()
+					   .AutoWidth()
+					   .VAlign(VAlign_Center)
+					   .Padding(0.0f, 0.0f, 4.0f, 0.0f)
+						[
+							SNew(SCheckBox)
+						   .IsChecked_Lambda([this]() -> ECheckBoxState
+							{
+								TSharedPtr<const FPCGManualEditNodeConfiguration> NodeConfiguration = GetConfigurationForActiveNode();
+								return (NodeConfiguration && NodeConfiguration->bSignatureRotation) ? ECheckBoxState::Checked : ECheckBoxState::Unchecked;
+							})
+						   .OnCheckStateChanged_Lambda([this](const ECheckBoxState NewState)
+							{
+								if (const UPCGNode* Node = GetActiveNode())
+								{
+									GetOrCreateConfigurationForNode(Node)->bSignatureRotation = (NewState == ECheckBoxState::Checked);
+								}
+							})
+							[
+								SNew(STextBlock)
+							   .Text(LOCTEXT("SignatureRotationCheckbox", "Rot"))
+							   .Font(FCoreStyle::GetDefaultFontStyle("Regular", 8))
+							]
+						]
+						+ SHorizontalBox::Slot()
+					   .AutoWidth()
+					   .VAlign(VAlign_Center)
+					   .Padding(0.0f, 0.0f, 8.0f, 0.0f)
+						[
+							SNew(SCheckBox)
+						   .IsChecked_Lambda([this]() -> ECheckBoxState
+							{
+								TSharedPtr<const FPCGManualEditNodeConfiguration> NodeConfiguration = GetConfigurationForActiveNode();
+								return (NodeConfiguration && NodeConfiguration->bSignatureScale) ? ECheckBoxState::Checked : ECheckBoxState::Unchecked;
+							})
+						   .OnCheckStateChanged_Lambda([this](const ECheckBoxState NewState)
+							{
+								if (const UPCGNode* Node = GetActiveNode())
+								{
+									GetOrCreateConfigurationForNode(Node)->bSignatureScale = (NewState == ECheckBoxState::Checked);
+								}
+							})
+							[
+								SNew(STextBlock)
+							   .Text(LOCTEXT("SignatureScaleCheckbox", "Sca"))
+							   .Font(FCoreStyle::GetDefaultFontStyle("Regular", 8))
+							]
+						]
+						+ SHorizontalBox::Slot()
+					   .AutoWidth()
+					   .VAlign(VAlign_Center)
+						[
+							SNew(SBox)
+						   .MinDesiredWidth(50.0f)
+						   .MaxDesiredWidth(60.0f)
+							[
+								SNew(SSpinBox<double>)
+							   .MinValue(1.0)
+							   .MaxValue(1000.0)
+							   .Value_Lambda([this]() -> double
+								{
+									TSharedPtr<const FPCGManualEditNodeConfiguration> NodeConfiguration = GetConfigurationForActiveNode();
+									return NodeConfiguration ? NodeConfiguration->SignatureTolerance : 10.0;
+								})
+							   .OnValueChanged_Lambda([this](const double NewValue)
+								{
+									if (const UPCGNode* Node = GetActiveNode())
+									{
+										GetOrCreateConfigurationForNode(Node)->SignatureTolerance = NewValue;
+									}
+								})
+							   .Font(FCoreStyle::GetDefaultFontStyle("Regular", 7))
+							]
+						]
+					];
+			}
+
+			if (DeltaTypeWidgetSlot.IsValid())
+			{
+				DeltaTypeWidgetSlot->ClearChildren();
+
+				if (ActiveExtension)
+				{
+					DeltaTypeWidgetSlot->AddSlot()
+					   .AutoHeight()
+						[
+							ActiveExtension->CreateWidget(VisualizerCallbacks)
+						];
+				}
+			}
+		}
+
+		// Push current context to the extension
+		if (ActiveExtension)
+		{
+			FPCGDeltaViewportContext Context = BuildExtensionContext();
+			ActiveExtension->UpdateContext(Context);
+			ActiveExtension->RefreshLists(Context);
+		}
+	}
+	else
+	{
+		ConfigurationSection->SetVisibility(EVisibility::Collapsed);
+	}
+}
+
+void SPCGManualEditPanel::OnDeltaTypeChanged(const UScriptStruct* InType, ESelectInfo::Type SelectInfo)
+{
+	if (const UPCGNode* Node = GetActiveNode())
+	{
+		GetOrCreateConfigurationForNode(Node)->DeltaType = InType;
+		UpdateConfigurationSection();
+	}
+}
+
+TSharedRef<SWidget> SPCGManualEditPanel::GenerateDeltaTypeMenu()
+{
+	FMenuBuilder MenuBuilder(/*bInShouldCloseWindowAfterMenuSelection=*/true, nullptr);
+
+	for (const UScriptStruct* Type : DeltaTypeOptions)
+	{
+		FText Label;
+		if (const IPCGDeltaViewportExtension* Extension = FPCGEditorModule::GetConstDeltaViewportExtensionRegistry().GetExtension(Type))
+		{
+			Label = Extension->GetDisplayName();
+		}
+
+		MenuBuilder.AddMenuEntry(
+			Label,
+			FText::GetEmpty(),
+			FSlateIcon(),
+			FUIAction(FExecuteAction::CreateSP(this, &SPCGManualEditPanel::OnDeltaTypeChanged, Type, ESelectInfo::OnMouseClick))
+			);
+	}
+
+	return MenuBuilder.MakeWidget();
+}
+
+FText SPCGManualEditPanel::GetCurrentDeltaTypeLabel() const
+{
+	TSharedPtr<const FPCGManualEditNodeConfiguration> Configuration = GetConfigurationForActiveNode();
+
+	if (const UScriptStruct* DeltaType = Configuration ? Configuration->DeltaType : nullptr)
+	{
+		if (const IPCGDeltaViewportExtension* Extension = FPCGEditorModule::GetConstDeltaViewportExtensionRegistry().GetExtension(DeltaType))
+		{
+			return Extension->GetDisplayName();
+		}
+	}
+
+	return FText::GetEmpty();
+}
+
+FPCGDeltaViewportContext SPCGManualEditPanel::BuildExtensionContext()
+{
+	FPCGDeltaViewportContext Context;
+	Context.bSelectionActive = bSelectionActive;
+	Context.SelectionDeltaKey = SelectionDeltaKey;
+	Context.SelectionTransform = SelectionTransform;
+	Context.SelectedElementIndex = SelectedElementIndex;
+	Context.OriginalElementIndex = OriginalElementIndex;
+	Context.ActivePCGComponent = ActivePCGComponent;
+	Context.ActiveStorageKey = ActiveStorageKey;
+	Context.NodeConfiguration = GetConfigurationForActiveNode();
+
+	return Context;
+}
+
+#undef LOCTEXT_NAMESPACE
diff --git a/Engine/Plugins/PCG/Source/PCGEditor/Private/SPCGManualEditPanel.h b/Engine/Plugins/PCG/Source/PCGEditor/Private/SPCGManualEditPanel.h
new file mode 100644
--- /dev/null
+++ b/Engine/Plugins/PCG/Source/PCGEditor/Private/SPCGManualEditPanel.h
@@ -0,0 +1,162 @@
+// Copyright Epic Games, Inc. All Rights Reserved.
+
+#pragma once
+
+#include "DeltaViewportExtensions/PCGDeltaViewportExtension.h"
+#include "Graph/PCGSourceDataContainer.h"
+#include "Graph/DataOverride/PCGDataOverride.h"
+
+#include "Widgets/SCompoundWidget.h"
+#include "Widgets/Views/SListView.h"
+
+class UPCGComponent;
+class UPCGGraph;
+class UPCGNode;
+class UPCGSettingsInterface;
+
+namespace UE::ToolWidgets
+{
+	class SDraggableBoxOverlay;
+}
+
+DECLARE_DELEGATE_OneParam(FOnPCGManualEditPanelNodeSelected, const UPCGNode* /*SelectedNode*/);
+DECLARE_DELEGATE_OneParam(FOnPCGJumpToNode, const UPCGNode* /*Node*/);
+
+struct FPCGManualEditNodeEntry
+{
+	TWeakObjectPtr<const UPCGNode> Node;
+	FText DisplayName;
+	bool bIsMarked = false;
+	bool bIsTemporary = false;
+	bool bIsActive = false;
+};
+
+/** Per-node transient configuration for the manual edit panel. Controls delta creation behavior in the viewport. */
+struct FPCGManualEditNodeConfiguration
+{
+	const UScriptStruct* DeltaType = nullptr;
+
+	/** Signature controls. I.e. transform components that are used to match source points. */
+	bool bSignaturePosition = true;
+	bool bSignatureRotation = false;
+	bool bSignatureScale = false;
+	double SignatureTolerance = 10.0;
+};
+
+class SPCGManualEditPanel : public SCompoundWidget
+{
+public:
+	SLATE_BEGIN_ARGS(SPCGManualEditPanel) {}
+		SLATE_ARGUMENT(TWeakObjectPtr<UPCGGraph>, Graph)
+	SLATE_END_ARGS()
+
+	void Construct(const FArguments& InArgs);
+
+	/** Refresh the node list from the graph. */
+	void RefreshNodeList();
+
+	/** Get the overlay wrapper widget. */
+	TSharedPtr<SWidget> GetOverlayWidget() const;
+
+	/** Create the overlay wrapper around this panel for viewport integration. */
+	TSharedRef<SWidget> CreateOverlayWidget();
+
+	/** Returns the configuration for the currently selected node, or nullptr if none. */
+	TSharedPtr<const FPCGManualEditNodeConfiguration> GetConfigurationForActiveNode() const;
+
+	/** Sets the active delta context for querying the collection. */
+	void SetDeltaContext(UPCGComponent* InComponent, const FPCGSourceDataStorageKey& InStorageKey);
+
+	/** Update the panel and visualizer's current selection. */
+	void SetSelectionState(
+		bool bHasSelection,
+		const FPCGDeltaKey& SelectedKey,
+		const FTransform& InSelectionTransform = FTransform::Identity,
+		int32 InSelectedElementIndex = INDEX_NONE,
+		int32 InOriginalElementIndex = INDEX_NONE);
+
+	/** Refresh the active extension's lists from the current delta context. */
+	void RefreshActiveExtensionLists();
+
+	/** Callbacks provided by the visualizer for extensions to invoke. Must be called before extensions are activated. */
+	void SetVisualizerActions(const FPCGDeltaViewportCallbacks& InActions);
+
+	/** Delegate fired when a node is selected in the panel. */
+	FOnPCGManualEditPanelNodeSelected OnNodeSelected;
+
+	/** Jump to node in the graph editor. */
+	FOnPCGJumpToNode OnJumpToNode;
+
+	/** Returns the currently selected node or nullptr if none. */
+	const UPCGNode* GetActiveNode() const;
+
+private:
+	TSharedRef<ITableRow> OnGenerateRow(TSharedPtr<FPCGManualEditNodeEntry> Entry, const TSharedRef<STableViewBase>& OwnerTable);
+	void OnListSelectionChanged(TSharedPtr<FPCGManualEditNodeEntry> SelectedEntry, ESelectInfo::Type SelectInfo);
+	void OnDeltaTypeChanged(const UScriptStruct* InType, ESelectInfo::Type SelectInfo);
+	TSharedRef<SWidget> GenerateDeltaTypeMenu();
+	FText GetCurrentDeltaTypeLabel() const;
+
+	/** Get or create configuration for a node. */
+	TSharedPtr<FPCGManualEditNodeConfiguration> GetOrCreateConfigurationForNode(TWeakObjectPtr<const UPCGNode> Node);
+
+	/** Rebuilds the configuration section visibility and content. */
+	void UpdateConfigurationSection();
+
+	/** Builds an FPCGDeltaViewportContext from current panel state. */
+	FPCGDeltaViewportContext BuildExtensionContext();
+
+	/** Graph owning the nodes shown in the panel. Used to discover overridable nodes and observe structural changes. */
+	TWeakObjectPtr<UPCGGraph> Graph;
+
+	/** Cached list of nodes eligible for manual editing, rebuilt by RefreshNodeList. */
+	TArray<TSharedPtr<FPCGManualEditNodeEntry>> NodeEntries;
+
+	/** Per-node transient configuration map. */
+	TMap<TWeakObjectPtr<const UPCGNode>, TSharedPtr<FPCGManualEditNodeConfiguration>> NodeConfigurationMap;
+
+	/** Available delta types for the combo box. */
+	TArray<const UScriptStruct*> DeltaTypeOptions;
+
+	/** Node list widget for selecting the active editing node. */
+	TSharedPtr<SListView<TSharedPtr<FPCGManualEditNodeEntry>>> NodeListView;
+
+	/** Draggable overlay for viewport handles. */
+	TSharedPtr<UE::ToolWidgets::SDraggableBoxOverlay> OverlayWidget;
+
+	/** Collapsible section containing node specific configuration widgets. */
+	TSharedPtr<SVerticalBox> ConfigurationSection;
+
+	/** Container for the signature checkbox row. */
+	TSharedPtr<SVerticalBox> SignatureSlot;
+
+	/** Container for the active delta type extension's widget. Swapped when the delta type combo changes. */
+	TSharedPtr<SVerticalBox> DeltaTypeWidgetSlot;
+
+	/** The currently active viewport extension (owned by the registry). */
+	IPCGDeltaViewportExtension* ActiveExtension = nullptr;
+
+	/** Callbacks provided by the visualizer for extensions to invoke. */
+	FPCGDeltaViewportCallbacks VisualizerCallbacks;
+
+	/** The PCG component currently being edited. */
+	TWeakObjectPtr<UPCGComponent> ActivePCGComponent;
+
+	/** Storage key identifying the delta collection on the active component. */
+	FPCGSourceDataStorageKey ActiveStorageKey;
+
+	/** True when the visualizer has an active element selection. */
+	bool bSelectionActive = false;
+
+	/** Index of the selected element within an insertion delta, or INDEX_NONE. */
+	int32 SelectedElementIndex = INDEX_NONE;
+
+	/** Index of the selected point in the original source data, or INDEX_NONE. */
+	int32 OriginalElementIndex = INDEX_NONE;
+
+	/** The delta key of the currently selected element. */
+	FPCGDeltaKey SelectionDeltaKey;
+
+	/** World-space transform of the currently selected element. */
+	FTransform SelectionTransform = FTransform::Identity;
+};
diff --git a/Engine/Plugins/PCG/Source/PCGEditor/Private/Widgets/SPCGEditorGraphDataOverrides.cpp b/Engine/Plugins/PCG/Source/PCGEditor/Private/Widgets/SPCGEditorGraphDataOverrides.cpp
new file mode 100644
--- /dev/null
+++ b/Engine/Plugins/PCG/Source/PCGEditor/Private/Widgets/SPCGEditorGraphDataOverrides.cpp
@@ -0,0 +1,475 @@
+// Copyright Epic Games, Inc. All Rights Reserved.
+
+#include "SPCGEditorGraphDataOverrides.h"
+
+#include "PCGEditor.h"
+#include "PCGEditorModule.h"
+#include "PCGElement.h"
+#include "PCGGraph.h"
+#include "PCGGraphExecutionInspection.h"
+#include "PCGNode.h"
+#include "DeltaVisualizations/PCGDeltaVisualizationRegistry.h"
+#include "Graph/PCGSourceDataContainer.h"
+#include "Graph/DataOverride/PCGDataOverrideHelpers.h"
+#include "Nodes/PCGEditorGraphNodeBase.h"
+
+#include "Editor.h"
+#include "Widgets/SNullWidget.h"
+#include "Widgets/Input/SButton.h"
+#include "Widgets/Text/STextBlock.h"
+#include "Widgets/Views/SListView.h"
+
+#define LOCTEXT_NAMESPACE "PCGDataOverridesListView"
+
+namespace PCGDataOverride
+{
+	namespace Helpers
+	{
+		/** Returns the pins that carry override data for the given node. */
+		static const TArray<TObjectPtr<UPCGPin>>& GetOverridePins(const UPCGNode* PCGNode)
+		{
+			const UPCGSettings* Settings = PCGNode->GetSettings();
+			const FPCGElementPtr Element = Settings ? Settings->GetElement() : nullptr;
+			const bool bUsesInputPins = Element && Element->GetDataOverridePhase() == EPCGDataOverridePhase::PrepareData;
+			return bUsesInputPins ? PCGNode->GetInputPins() : PCGNode->GetOutputPins();
+		}
+	}
+
+	namespace Constants
+	{
+		static constexpr float CellPadding = 8.f;
+	}
+} // namespace PCGDataOverride
+
+void SPCGDataOverrideRow::Construct(const FArguments& InArgs, const TSharedRef<STableViewBase>& InOwnerTableView)
+{
+	Item = InArgs._Item;
+	check(Item.IsValid());
+
+	// @todo_pcg: Check Item->OrphanStatus == Orphaned and apply a tinted row style if so.
+	const FSuperRowType::FArguments Arguments = FTableRowArgs().Style(FAppStyle::Get(), "DataTableEditor.CellListViewRow");
+	SMultiColumnTableRow::Construct(Arguments, InOwnerTableView);
+}
+
+TSharedRef<SWidget> SPCGDataOverrideRow::GenerateWidgetForColumn(const FName& InColumnID)
+{
+	if (!Item.IsValid() || !Item->Visualizer)
+	{
+		return SNullWidget::NullWidget;
+	}
+
+	// Find and check for custom widget first.
+	const FConstStructView DeltaView(Item->DeltaStruct.GetScriptStruct(), Item->DeltaStruct.GetMemory());
+	const TSharedPtr<SWidget> CustomWidget = Item->Visualizer->CreateCellWidget(InColumnID, Item->DeltaKey, DeltaView);
+	if (CustomWidget.IsValid())
+	{
+		return CustomWidget.ToSharedRef();
+	}
+
+	// Fall back to text.
+	const FText CellText = Item->Visualizer->GetCellText(InColumnID, Item->DeltaKey, DeltaView);
+
+	return SNew(SBox)
+		.Padding(FMargin(PCGDataOverride::Constants::CellPadding, 0))
+		.VAlign(VAlign_Center)
+		[
+			SNew(STextBlock)
+			.Text(CellText)
+		];
+}
+
+void SPCGEditorGraphDataOverridesView::Construct(const FArguments& InArgs, TSharedPtr<FPCGEditor> InPCGEditor)
+{
+	check(InPCGEditor.IsValid());
+
+	PCGEditorPtr = InPCGEditor;
+
+	InPCGEditor->GetInspectionDataManager().OnInspectedStackChangedDelegate.AddSP(this, &SPCGEditorGraphDataOverridesView::OnInspectedStackChanged);
+
+	// Start with an empty header row; columns are populated during RefreshDataOverrides.
+	HeaderRow = SNew(SHeaderRow);
+
+	SAssignNew(ListView, SListView<FPCGDeltaOverrideItemPtr>)
+		.ListItemsSource(&ListViewItems)
+		.OnGenerateRow(this, &SPCGEditorGraphDataOverridesView::OnGenerateRow)
+		.HeaderRow(HeaderRow)
+		.SelectionMode(ESelectionMode::Single);
+
+	ChildSlot
+	[
+		SNew(SVerticalBox)
+		+ SVerticalBox::Slot()
+		.AutoHeight()
+		[
+			SNew(SHorizontalBox)
+			+ SHorizontalBox::Slot()
+			.AutoWidth()
+			[
+				SNew(SButton)
+				.ButtonStyle(FAppStyle::Get(), "SimpleButton")
+				.ToolTipText(LOCTEXT("ClearNodeOverridesTooltip", "Clear all overrides for this node"))
+				.ContentPadding(FMargin(2, 2))
+				.OnClicked(this, &SPCGEditorGraphDataOverridesView::OnClearNodeOverrides)
+				[
+					SNew(SImage)
+					.Image(FAppStyle::Get().GetBrush("Icons.Delete"))
+					.ColorAndOpacity(FSlateColor::UseForeground())
+				]
+			]
+		]
+		+ SVerticalBox::Slot()
+		.FillHeight(1)
+		[
+			ListView.ToSharedRef()
+		]
+	];
+}
+
+SPCGEditorGraphDataOverridesView::~SPCGEditorGraphDataOverridesView()
+{
+	if (const TSharedPtr<FPCGEditor> PCGEditor = PCGEditorPtr.Pin())
+	{
+		PCGEditor->GetInspectionDataManager().OnInspectedStackChangedDelegate.RemoveAll(this);
+	}
+}
+
+void SPCGEditorGraphDataOverridesView::Tick(const FGeometry& AllottedGeometry, const double InCurrentTime, const float InDeltaTime)
+{
+	SCompoundWidget::Tick(AllottedGeometry, InCurrentTime, InDeltaTime);
+
+	if (bNeedsRefresh)
+	{
+		bNeedsRefresh = false;
+		RefreshDataOverrides();
+	}
+}
+
+void SPCGEditorGraphDataOverridesView::SetNodeBeingInspected(UPCGEditorGraphNodeBase* InNode)
+{
+	if (PCGEditorGraphNode.Get() == InNode)
+	{
+		return;
+	}
+
+	PCGEditorGraphNode = InNode;
+	RequestRefresh();
+}
+
+UPCGEditorGraphNodeBase* SPCGEditorGraphDataOverridesView::GetNodeBeingInspected() const
+{
+	return PCGEditorGraphNode.Get();
+}
+
+void SPCGEditorGraphDataOverridesView::ClearItems()
+{
+	ListViewItems.Empty();
+	RequestRefresh();
+}
+
+void SPCGEditorGraphDataOverridesView::RequestRefresh()
+{
+	bNeedsRefresh = true;
+}
+
+TSharedRef<ITableRow> SPCGEditorGraphDataOverridesView::OnGenerateRow(FPCGDeltaOverrideItemPtr InItem, const TSharedRef<STableViewBase>& InOwnerTable)
+{
+	return SNew(SPCGDataOverrideRow, InOwnerTable).Item(InItem);
+}
+
+void SPCGEditorGraphDataOverridesView::OnInspectedStackChanged(const FPCGStack& InPCGStack)
+{
+	RequestRefresh();
+}
+
+void SPCGEditorGraphDataOverridesView::RebuildHeaderRow()
+{
+	if (!HeaderRow.IsValid())
+	{
+		return;
+	}
+
+	HeaderRow->ClearColumns();
+
+	// Add columns from the merged set.
+	for (const FPCGDeltaVisualizerColumnInfo& ColumnInfo : ActiveColumns)
+	{
+		SHeaderRow::FColumn::FArguments ColumnArgs;
+		ColumnArgs.ColumnId(ColumnInfo.Id);
+		ColumnArgs.DefaultLabel(ColumnInfo.Label);
+		ColumnArgs.DefaultTooltip(ColumnInfo.Tooltip);
+
+		if (ColumnInfo.Width > 0.f)
+		{
+			ColumnArgs.ManualWidth(ColumnInfo.Width);
+		}
+		else
+		{
+			ColumnArgs.FillWidth(1.0f);
+		}
+
+		switch (ColumnInfo.CellAlignment)
+		{
+			case EPCGTableVisualizerCellAlignment::Left:
+				ColumnArgs.HAlignCell(HAlign_Left);
+				break;
+			case EPCGTableVisualizerCellAlignment::Center:
+				ColumnArgs.HAlignCell(HAlign_Center);
+				ColumnArgs.HAlignHeader(HAlign_Center);
+				break;
+			case EPCGTableVisualizerCellAlignment::Right:
+				ColumnArgs.HAlignCell(HAlign_Right);
+				break;
+			default:
+				break;
+		}
+
+		HeaderRow->AddColumn(ColumnArgs);
+	}
+}
+
+void SPCGEditorGraphDataOverridesView::RefreshDataOverrides()
+{
+	ListViewItems.Empty();
+	ActiveColumns.Empty();
+
+	const TSharedPtr<FPCGEditor> PCGEditor = PCGEditorPtr.Pin();
+	if (!PCGEditor.IsValid() || !PCGEditorGraphNode.IsValid())
+	{
+		RebuildHeaderRow();
+
+		if (ListView.IsValid())
+		{
+			ListView->RequestListRefresh();
+		}
+
+		return;
+	}
+
+	const FPCGEditorInspectionDataManager& InspectionDataManager = PCGEditor->GetInspectionDataManager();
+
+	IPCGGraphExecutionSource* ExecutionSource = InspectionDataManager.GetPCGSourceBeingInspected();
+	const FPCGSourceDataContainer* DataContainer = ExecutionSource ? ExecutionSource->GetExecutionState().GetSourceDataContainer() : nullptr;
+
+	const UPCGNode* PCGNode = PCGEditorGraphNode->GetPCGNode();
+
+	if (!DataContainer || !PCGNode)
+	{
+		RebuildHeaderRow();
+
+		if (ListView.IsValid())
+		{
+			ListView->RequestListRefresh();
+		}
+
+		return;
+	}
+
+	// Get the execution stacks for this node. These are the actual runtime stacks captured during graph execution,
+	// which is the same stack source that the storage side uses. There exists a hash mismatch that occurs when
+	// using the inspection stack (GetStackBeingInspected), which may be reconstructed with different frames.
+	const FPCGGraphExecutionInspection& Inspection = ExecutionSource->GetExecutionState().GetInspection();
+	TMap<TObjectKey<const UPCGNode>, TSet<FPCGGraphExecutionInspection::FNodeExecutedNotificationData>> NodeStacks = Inspection.GetExecutedNodeStacks();
+
+	const TSet<FPCGGraphExecutionInspection::FNodeExecutedNotificationData>* StackSet = NodeStacks.Find(PCGNode);
+	if (!StackSet || StackSet->IsEmpty())
+	{
+		RebuildHeaderRow();
+
+		if (ListView.IsValid())
+		{
+			ListView->RequestListRefresh();
+		}
+
+		return;
+	}
+
+	const FPCGDeltaVisualizationRegistry& Registry = FPCGEditorModule::GetConstDeltaVisualizationRegistry();
+
+	// Track which column IDs have already been added to avoid duplicates from multiple visualizers.
+	TSet<FName> TrackedColumnIds;
+	TSet<const IPCGDeltaVisualization*> TrackedVisualizers;
+
+	using namespace PCG::DataOverride;
+
+	// Process a delta collection into list view items and merge columns.
+	auto ProcessDeltaCollection = [this, &Registry, &TrackedColumnIds, &TrackedVisualizers](const FPCGDeltaCollection* DeltaCollection, FName PinLabel)
+	{
+		DeltaCollection->ForEachDelta([this, &Registry, &TrackedColumnIds, &TrackedVisualizers, PinLabel](const FPCGDeltaKey& DeltaKey, const TInstancedStruct<FPCGDeltaBase>& DeltaStruct) -> bool
+		{
+			const UScriptStruct* ScriptStruct = DeltaStruct.GetScriptStruct();
+			const IPCGDeltaVisualization* Visualizer = Registry.GetDeltaVisualization(ScriptStruct);
+
+			// Merge visualizer columns into ActiveColumns (deduplicate by Id). Skip if this visualizer is already tracked.
+			bool bVisualizerIsAlreadyTracked = false;
+			if (Visualizer)
+			{
+				TrackedVisualizers.Add(Visualizer, &bVisualizerIsAlreadyTracked);
+			}
+
+			if (Visualizer && !bVisualizerIsAlreadyTracked)
+			{
+				for (const FPCGDeltaVisualizerColumnInfo& ColumnInfo : Visualizer->GetColumnInfos())
+				{
+					bool bColumnAlreadyTracked = false;
+					TrackedColumnIds.Add(ColumnInfo.Id, &bColumnAlreadyTracked);
+					if (!bColumnAlreadyTracked)
+					{
+						ActiveColumns.Add(ColumnInfo);
+					}
+				}
+			}
+
+			// Create the item.
+			TSharedPtr<FPCGDeltaOverrideItem> NewItem = MakeShared<FPCGDeltaOverrideItem>();
+			NewItem->DeltaKey = DeltaKey;
+			NewItem->DeltaStruct = DeltaStruct; // Copy
+			NewItem->Visualizer = Visualizer;
+			NewItem->PinLabel = PinLabel;
+			NewItem->OrphanStatus = EPCGDeltaOrphanStatus::Unknown;
+
+			ListViewItems.Emplace(MoveTemp(NewItem));
+			return true;
+		});
+	};
+
+	const TArray<TObjectPtr<UPCGPin>>& OverridePins = PCGDataOverride::Helpers::GetOverridePins(PCGNode);
+
+	// Iterate each execution stack for this node (there may be multiple in loop/subgraph scenarios).
+	for (const FPCGGraphExecutionInspection::FNodeExecutedNotificationData& StackData : *StackSet)
+	{
+		const FPCGStack& ExecutionStack = StackData.Stack;
+		const FXxHash64 NodeKey = Helpers::ComputeNodeOverrideKey(ExecutionStack, PCGNode);
+
+		// Try pin-level keys for each override pin.
+		bool bFoundAnyPinData = false;
+		for (const UPCGPin* Pin : OverridePins)
+		{
+			if (!Pin)
+			{
+				continue;
+			}
+
+			const FName PinLabel = Pin->Properties.Label;
+			const FPCGSourceDataStorageKey PinDataKey(Constants::DefaultOverrideLabel, Helpers::ComputePinOverrideKey(NodeKey, PinLabel).Hash);
+			FConstSharedStruct SharedStruct = DataContainer->Get<FPCGDeltaCollection>(PinDataKey);
+			const FPCGDeltaCollection* DeltaCollection = SharedStruct.IsValid() ? SharedStruct.GetPtr<const FPCGDeltaCollection>() : nullptr;
+
+			if (!DeltaCollection || DeltaCollection->IsEmpty())
+			{
+				continue;
+			}
+
+			bFoundAnyPinData = true;
+			ProcessDeltaCollection(DeltaCollection, PinLabel);
+		}
+
+		// Fall back to node-level key if no pin-level data was found.
+		if (!bFoundAnyPinData)
+		{
+			const FPCGSourceDataStorageKey NodeDataKey(Constants::DefaultOverrideLabel, NodeKey.Hash);
+			FConstSharedStruct SharedStruct = DataContainer->Get<FPCGDeltaCollection>(NodeDataKey);
+			const FPCGDeltaCollection* DeltaCollection = SharedStruct.IsValid() ? SharedStruct.GetPtr<const FPCGDeltaCollection>() : nullptr;
+
+			if (DeltaCollection && !DeltaCollection->IsEmpty())
+			{
+				ProcessDeltaCollection(DeltaCollection, NAME_None);
+			}
+		}
+	}
+
+	// @todo_pcg: Populate orphan status by running FilterCandidates + Resolve against current output data.
+
+	RebuildHeaderRow();
+
+	if (ListView.IsValid())
+	{
+		ListView->RequestListRefresh();
+	}
+}
+
+// Removes all delta collections (pin and node level) for the inspected node, then triggers re-execution.
+FReply SPCGEditorGraphDataOverridesView::OnClearNodeOverrides()
+{
+	const TSharedPtr<FPCGEditor> PCGEditor = PCGEditorPtr.Pin();
+	if (!PCGEditor.IsValid() || !PCGEditorGraphNode.IsValid())
+	{
+		return FReply::Handled();
+	}
+
+	const FPCGEditorInspectionDataManager& InspectionDataManager = PCGEditor->GetInspectionDataManager();
+	IPCGGraphExecutionSource* ExecutionSource = InspectionDataManager.GetPCGSourceBeingInspected();
+	if (!ExecutionSource)
+	{
+		return FReply::Handled();
+	}
+
+	FPCGSourceDataContainer* DataContainer = ExecutionSource->GetExecutionState().GetSourceDataContainer();
+	const UPCGNode* PCGNode = PCGEditorGraphNode->GetPCGNode();
+	if (!DataContainer || !PCGNode)
+	{
+		return FReply::Handled();
+	}
+
+	const FPCGGraphExecutionInspection& Inspection = ExecutionSource->GetExecutionState().GetInspection();
+	TMap<TObjectKey<const UPCGNode>, TSet<FPCGGraphExecutionInspection::FNodeExecutedNotificationData>> NodeStacks = Inspection.GetExecutedNodeStacks();
+
+	const TSet<FPCGGraphExecutionInspection::FNodeExecutedNotificationData>* StackSet = NodeStacks.Find(PCGNode);
+	if (!StackSet || StackSet->IsEmpty())
+	{
+		return FReply::Handled();
+	}
+
+	// Execution sources aren't always UObjects, so wrap in a manual transaction only when possible.
+	UObject* SourceObject = Cast<UObject>(ExecutionSource);
+	if (SourceObject)
+	{
+		GEditor->BeginTransaction(LOCTEXT("ClearNodeOverrides", "PCG: Clear All Node Overrides"));
+		SourceObject->Modify();
+	}
+
+	using namespace PCG::DataOverride;
+
+	const TArray<TObjectPtr<UPCGPin>>& OverridePins = PCGDataOverride::Helpers::GetOverridePins(PCGNode);
+
+	for (const FPCGGraphExecutionInspection::FNodeExecutedNotificationData& StackData : *StackSet)
+	{
+		const FPCGStack& ExecutionStack = StackData.Stack;
+		const FXxHash64 NodeKey = Helpers::ComputeNodeOverrideKey(ExecutionStack, PCGNode);
+
+		// Remove pin-level keys for each override pin.
+		for (const UPCGPin* Pin : OverridePins)
+		{
+			if (!Pin)
+			{
+				continue;
+			}
+
+			const FName PinLabel = Pin->Properties.Label;
+			const FPCGSourceDataStorageKey PinDataKey(Constants::DefaultOverrideLabel, Helpers::ComputePinOverrideKey(NodeKey, PinLabel).Hash);
+			DataContainer->Remove(PinDataKey);
+		}
+
+		// Remove the node-level key as well.
+		const FPCGSourceDataStorageKey NodeDataKey(Constants::DefaultOverrideLabel, NodeKey.Hash);
+		DataContainer->Remove(NodeDataKey);
+	}
+
+	DataContainer->MarkDirty();
+
+	if (SourceObject)
+	{
+		GEditor->EndTransaction();
+	}
+
+	// Notify the graph so it re-executes with the overrides removed.
+	if (UPCGGraph* Graph = PCGEditor->GetMainGraph())
+	{
+		Graph->ForceNotificationForEditor(EPCGChangeType::ExternalModification);
+	}
+
+	RequestRefresh();
+
+	return FReply::Handled();
+}
+
+#undef LOCTEXT_NAMESPACE
diff --git a/Engine/Plugins/PCG/Source/PCGEditor/Private/Widgets/SPCGEditorGraphDataOverrides.h b/Engine/Plugins/PCG/Source/PCGEditor/Private/Widgets/SPCGEditorGraphDataOverrides.h
new file mode 100644
--- /dev/null
+++ b/Engine/Plugins/PCG/Source/PCGEditor/Private/Widgets/SPCGEditorGraphDataOverrides.h
@@ -0,0 +1,92 @@
+// Copyright Epic Games, Inc. All Rights Reserved.
+
+#pragma once
+
+#include "DeltaVisualizations/PCGDeltaVisualization.h"
+#include "Graph/DataOverride/PCGDataOverride.h"
+
+#include "StructUtils/InstancedStruct.h"
+#include "Widgets/Views/STableRow.h"
+
+class FPCGEditor;
+class IPCGDeltaVisualization;
+class UPCGEditorGraphNodeBase;
+struct FPCGStack;
+
+/** Whether the delta was consumed. Used for override detection. */
+enum class EPCGDeltaOrphanStatus : uint8
+{
+	Unknown,
+	Resolved,
+	Orphaned
+};
+
+/** View model for a single row in the Data Overrides table. Pairs a delta key/struct with its resolved visualizer. */
+struct FPCGDeltaOverrideItem
+{
+	FPCGDeltaKey DeltaKey;
+	TInstancedStruct<FPCGDeltaBase> DeltaStruct;
+	const IPCGDeltaVisualization* Visualizer = nullptr;
+	FName PinLabel;
+	EPCGDeltaOrphanStatus OrphanStatus = EPCGDeltaOrphanStatus::Unknown;
+};
+
+typedef TSharedPtr<FPCGDeltaOverrideItem> FPCGDeltaOverrideItemPtr;
+
+/** Table row widget. Populated by a registered IPCGDeltaVisualization. */
+class SPCGDataOverrideRow : public SMultiColumnTableRow<FPCGDeltaOverrideItemPtr>
+{
+public:
+	SLATE_BEGIN_ARGS(SPCGDataOverrideRow) {}
+		SLATE_ARGUMENT(FPCGDeltaOverrideItemPtr, Item)
+	SLATE_END_ARGS()
+
+	void Construct(const FArguments& InArgs, const TSharedRef<STableViewBase>& InOwnerTableView);
+	virtual TSharedRef<SWidget> GenerateWidgetForColumn(const FName& InColumnID) override;
+
+private:
+	FPCGDeltaOverrideItemPtr Item;
+};
+
+/**
+ * Data Overrides tab widget. Displays a table of delta overrides (manual edits) for the currently inspected node. Reads
+ * delta collections from the execution source's data container using an identical key signature.
+ */
+class SPCGEditorGraphDataOverridesView : public SCompoundWidget
+{
+public:
+	SLATE_BEGIN_ARGS(SPCGEditorGraphDataOverridesView) {}
+	SLATE_END_ARGS()
+
+	void Construct(const FArguments& InArgs, TSharedPtr<FPCGEditor> InPCGEditor);
+
+	virtual ~SPCGEditorGraphDataOverridesView() override;
+	virtual void Tick(const FGeometry& AllottedGeometry, const double InCurrentTime, const float InDeltaTime) override;
+
+	// To mirror the ALV
+	void SetNodeBeingInspected(UPCGEditorGraphNodeBase* InNode);
+	UPCGEditorGraphNodeBase* GetNodeBeingInspected() const;
+
+	void ClearItems();
+	void RequestRefresh();
+
+private:
+	FReply OnClearNodeOverrides();
+
+	TSharedRef<ITableRow> OnGenerateRow(FPCGDeltaOverrideItemPtr InItem, const TSharedRef<STableViewBase>& InOwnerTable);
+
+	void OnInspectedStackChanged(const FPCGStack& InPCGStack);
+	void RefreshDataOverrides();
+	void RebuildHeaderRow();
+
+	TWeakObjectPtr<UPCGEditorGraphNodeBase> PCGEditorGraphNode;
+	bool bNeedsRefresh = false;
+
+	TWeakPtr<FPCGEditor> PCGEditorPtr;
+	TSharedPtr<SHeaderRow> HeaderRow;
+	TSharedPtr<SListView<FPCGDeltaOverrideItemPtr>> ListView;
+	TArray<FPCGDeltaOverrideItemPtr> ListViewItems;
+
+	/** Merged column set from all visualizers active in the current view. */
+	TArray<FPCGDeltaVisualizerColumnInfo> ActiveColumns;
+};
diff --git a/Engine/Plugins/PCG/Source/PCGEditor/Public/DeltaViewportExtensions/PCGDeltaViewportExtension.h b/Engine/Plugins/PCG/Source/PCGEditor/Public/DeltaViewportExtensions/PCGDeltaViewportExtension.h
new file mode 100644
--- /dev/null
+++ b/Engine/Plugins/PCG/Source/PCGEditor/Public/DeltaViewportExtensions/PCGDeltaViewportExtension.h
@@ -0,0 +1,184 @@
+// Copyright Epic Games, Inc. All Rights Reserved.
+
+#pragma once
+
+#include "Graph/DataOverride/PCGDataOverride.h"
+#include "Graph/PCGSourceDataContainer.h"
+
+#include "InputCoreTypes.h"
+#include "StructUtils/StructView.h"
+#include "Widgets/SCompoundWidget.h"
+
+class UPCGComponent;
+struct FPCGManualEditNodeConfiguration;
+
+/** Context passed to extensions when creating a new delta on first gizmo release. */
+struct FPCGDeltaCreateContext
+{
+	FBox ComponentBounds = FBox(ForceInit);
+	int32 OriginalElementIndex = INDEX_NONE;
+};
+
+/** Snapshot of panel state passed to extensions when the selection or configuration changes. */
+struct FPCGDeltaViewportContext
+{
+	bool bSelectionActive = false;
+	FPCGDeltaKey SelectionDeltaKey;
+	FTransform SelectionTransform = FTransform::Identity;
+	TWeakObjectPtr<UPCGComponent> ActivePCGComponent;
+	FPCGSourceDataStorageKey ActiveStorageKey;
+	TSharedPtr<const FPCGManualEditNodeConfiguration> NodeConfiguration;
+
+	/** Index of the selected element within an insertion delta, or INDEX_NONE. */
+	int32 SelectedElementIndex = INDEX_NONE;
+
+	/** Index of the selected point in the original source data, or INDEX_NONE. */
+	int32 OriginalElementIndex = INDEX_NONE;
+};
+
+/** Generic actions provided by the framework to extensions to subscribe to on creation. */
+struct FPCGDeltaViewportCallbacks
+{
+	TFunction<void(const FPCGDeltaKey&)> RevertDelta;
+	TFunction<void()> RevertAllRestorableDeltas;
+	TFunction<void()> RequestListRefresh;
+	TFunction<void()> ClearSelection;
+	TFunction<void(const FPCGDeltaKey&, int32)> SelectElement;
+};
+
+/** A key binding requested by an extension for an input action. */
+struct FPCGDeltaKeyBinding
+{
+	FKey Key;
+	TFunction<void()> Action;
+};
+
+/** An element inserted by an extension that has no corresponding source point. Drawn and hit-tested by the visualizer. */
+struct FInsertedViewportElement
+{
+	FPCGDeltaKey DeltaKey;
+	int32 ElementIndex = INDEX_NONE;
+	FTransform Transform = FTransform::Identity;
+};
+
+/**
+ * Interface for manual editing viewport behavior and panel UI. Implement this to control:
+ * - What UI appears in the manual edit panel for a given delta type.
+ * - How elements with this delta type are drawn, selected, and manipulated in the viewport.
+ */
+class IPCGDeltaViewportExtension
+{
+public:
+	virtual ~IPCGDeltaViewportExtension() = default;
+
+	/** Returns the display name for this delta type. */
+	virtual FText GetDisplayName() const = 0;
+
+	/** The delta name used for key computation. */
+	virtual FName GetDeltaName() const { return NAME_None; }
+
+	/** Sort priority for the delta type selector. Lower values appear first. */
+	virtual int32 GetSortPriority() const { return std::numeric_limits<int32>::max(); }
+
+	/** Build the Slate widget section for this delta type. Called when the extension is swapped in. */
+	virtual TSharedRef<SWidget> CreateWidget(FPCGDeltaViewportCallbacks Callbacks) = 0;
+
+	/** Called when the panel context changes. Update buttons, labels, etc. */
+	virtual void UpdateContext(const FPCGDeltaViewportContext& Context) = 0;
+
+	/** Called when the delta context changes and lists need refreshing. */
+	virtual void RefreshLists(const FPCGDeltaViewportContext& Context) = 0;
+
+	/** Returns key bindings this extension wants the visualizer to route. Called when the extension becomes active. */
+	virtual TArray<FPCGDeltaKeyBinding> GetKeyBindings() { return {}; }
+
+	/** Color for drawing elements with this delta type in the viewport. */
+	virtual FLinearColor GetDisplayColor(bool bIsSelected) const = 0;
+
+	/** Given a source element transform and the delta struct, return the transform to display. */
+	virtual FTransform GetDisplayTransform(const FTransform& SourceElementTransform, const FConstStructView DeltaStruct) const
+	{
+		return SourceElementTransform;
+	}
+
+	/** Given a source element's transform and a delta, return true if this delta corresponds to that source element. */
+	virtual bool MatchesSourceElement(const FTransform& SourceElementTransform, const FConstStructView DeltaStruct, const double SpatialTolerance) const
+	{
+		return false;
+	}
+
+	/** Collect elements inserted by this extension that have no corresponding source point. The visualizer draws and hit-tests these. */
+	virtual void GetInsertedElements(const FPCGDeltaCollection& Collection, TArray<FInsertedViewportElement>& OutElements) const {}
+
+	/** Returns true if this delta type uses a TRS gizmo when selected. */
+	virtual bool UsesTRSGizmo() const { return true; }
+
+	/** Alt+drag handler. Usually used to create/insert new elements. Returns the new element index, or INDEX_NONE. */
+	virtual int32 OnAltDrag(FPCGDeltaCollection& Collection, const FTransform& SourceTransform, FPCGDeltaKey& OutDeltaKey) const
+	{
+		return INDEX_NONE;
+	}
+
+	/** Apply the final gizmo transform to the delta struct. Returns true on success. */
+	virtual bool ApplyGizmoTransform(const FTransform& NewTransform, TInstancedStruct<FPCGDeltaBase>& DeltaStruct, int32 ElementIndex) const
+	{
+		return false;
+	}
+
+	/** Create a new delta struct in the collection on first gizmo release (deferred creation). */
+	virtual void CreateNewDelta(
+		const FPCGDeltaKey& DeltaKey,
+		FPCGDeltaCollection& Collection,
+		const FTransform& OriginalTransform,
+		const FTransform& NewTransform,
+		const FPCGDeltaCreateContext& Context) const
+	{}
+
+	/** Returns the original transform stored in the delta struct. */
+	virtual FTransform GetOriginalTransform(const FConstStructView DeltaStruct) const
+	{
+		return FTransform::Identity;
+	}
+
+	/** Collect all delta keys in the collection that this extension considers restorable. */
+	virtual void CollectRestorableKeys(const FPCGDeltaCollection& Collection, TArray<FPCGDeltaKey>& OutKeys) const {}
+};
+
+/** Register a delta struct type with an IPCGDeltaViewportExtension to control viewport behavior and panel UI. */
+struct FPCGDeltaViewportExtensionRegistry
+{
+	FPCGDeltaViewportExtensionRegistry() = default;
+	FPCGDeltaViewportExtensionRegistry(const FPCGDeltaViewportExtensionRegistry&) = delete;
+	FPCGDeltaViewportExtensionRegistry(FPCGDeltaViewportExtensionRegistry&&) = default;
+	FPCGDeltaViewportExtensionRegistry& operator=(const FPCGDeltaViewportExtensionRegistry&) = delete;
+	FPCGDeltaViewportExtensionRegistry& operator=(FPCGDeltaViewportExtensionRegistry&&) = default;
+	~FPCGDeltaViewportExtensionRegistry() = default;
+
+	/** Register an internal viewport extension. Internal extensions are provided by the PCG plugin itself. */
+	PCGEDITOR_API void RegisterInternalExtension(const UScriptStruct* DeltaStruct, TUniquePtr<IPCGDeltaViewportExtension> Extension);
+
+	/** Unregister all internal viewport extensions. */
+	PCGEDITOR_API void UnregisterAllInternalExtensions();
+
+	/** Register an external viewport extension. External registrations take priority over internal ones. */
+	PCGEDITOR_API void RegisterExtension(const UScriptStruct* DeltaStruct, TUniquePtr<IPCGDeltaViewportExtension> Extension);
+
+	/** Unregister an external viewport extension. */
+	PCGEDITOR_API void UnregisterExtension(const UScriptStruct* DeltaStruct);
+
+	/** Look up an extension for the given delta struct type, walking the struct hierarchy toward FPCGDeltaBase. */
+	template <typename DeltaType>
+	IPCGDeltaViewportExtension* GetExtension() const { return GetExtension(DeltaType::StaticStruct()); }
+
+	PCGEDITOR_API IPCGDeltaViewportExtension* GetExtension(const UScriptStruct* DeltaStruct) const;
+
+	/** Returns all registered delta struct types (from both internal and external registries, deduplicated). */
+	PCGEDITOR_API TArray<const UScriptStruct*> GetRegisteredDeltaTypes() const;
+
+private:
+	/** Registry for extensions defined inside the PCG Plugin. */
+	TMap<const UScriptStruct*, TUniquePtr<IPCGDeltaViewportExtension>> InternalRegistry;
+
+	/** Registry for extensions defined outside the PCG Plugin. These take priority, allowing users to override default behavior. */
+	TMap<const UScriptStruct*, TUniquePtr<IPCGDeltaViewportExtension>> ExternalRegistry;
+};
diff --git a/Engine/Plugins/PCG/Source/PCGEditor/Public/DeltaVisualizations/PCGDeltaVisualization.h b/Engine/Plugins/PCG/Source/PCGEditor/Public/DeltaVisualizations/PCGDeltaVisualization.h
new file mode 100644
--- /dev/null
+++ b/Engine/Plugins/PCG/Source/PCGEditor/Public/DeltaVisualizations/PCGDeltaVisualization.h
@@ -0,0 +1,48 @@
+// Copyright Epic Games, Inc. All Rights Reserved.
+
+#pragma once
+
+#include "PCGDataVisualization.h"
+
+#include "StructUtils/StructView.h"
+
+struct FPCGDeltaKey;
+class SWidget;
+
+/** Lightweight column descriptor for the delta overrides table. */
+struct FPCGDeltaVisualizerColumnInfo
+{
+	FName Id;
+	FText Label;
+	FText Tooltip;
+	float Width = -1.0f; // Calculated automatically if < 0
+	EPCGTableVisualizerCellAlignment CellAlignment = EPCGTableVisualizerCellAlignment::Left;
+};
+
+/** Common column identifiers shared across delta visualizers. */
+namespace PCGDeltaVisualization::Constants
+{
+	static const FName DeltaTypeId = TEXT("DeltaType");
+} // namespace PCGDeltaVisualization::Constants
+
+/**
+ * Interface for delta type visualizers. If implemented, this controls how a specific delta struct type is displayed in
+ * the Data Overrides panel.
+ */
+class IPCGDeltaVisualization
+{
+public:
+	virtual ~IPCGDeltaVisualization() = default;
+
+	/** Returns the column descriptors this visualizer provides. */
+	virtual TArray<FPCGDeltaVisualizerColumnInfo> GetColumnInfos() const = 0;
+
+	/** Returns the text content for a specific cell. */
+	virtual FText GetCellText(FName ColumnId, const FPCGDeltaKey& DeltaKey, FConstStructView Delta) const = 0;
+
+	/**
+	 * Returns a custom widget for a cell. Returns nullptr by default to fall back to GetCellText wrapped in STextBlock.
+	 * Note: Optional.
+	 */
+	virtual TSharedPtr<SWidget> CreateCellWidget(FName ColumnId, const FPCGDeltaKey& DeltaKey, FConstStructView Delta) const { return nullptr; }
+};
diff --git a/Engine/Plugins/PCG/Source/PCGEditor/Public/DeltaVisualizations/PCGDeltaVisualizationRegistry.h b/Engine/Plugins/PCG/Source/PCGEditor/Public/DeltaVisualizations/PCGDeltaVisualizationRegistry.h
new file mode 100644
--- /dev/null
+++ b/Engine/Plugins/PCG/Source/PCGEditor/Public/DeltaVisualizations/PCGDeltaVisualizationRegistry.h
@@ -0,0 +1,38 @@
+// Copyright Epic Games, Inc. All Rights Reserved.
+
+#pragma once
+
+#include "DeltaVisualizations/PCGDeltaVisualization.h"
+
+class UScriptStruct;
+
+/** Register a delta struct type with an IPCGDeltaVisualization to implement custom visualization behavior. Mirrors FPCGDataVisualizationRegistry. */
+struct FPCGDeltaVisualizationRegistry
+{
+	FPCGDeltaVisualizationRegistry() = default;
+	FPCGDeltaVisualizationRegistry(const FPCGDeltaVisualizationRegistry&) = delete;
+	FPCGDeltaVisualizationRegistry(FPCGDeltaVisualizationRegistry&&) = default;
+	FPCGDeltaVisualizationRegistry& operator=(const FPCGDeltaVisualizationRegistry&) = delete;
+	FPCGDeltaVisualizationRegistry& operator=(FPCGDeltaVisualizationRegistry&&) = default;
+	~FPCGDeltaVisualizationRegistry() = default;
+
+	/** Register an external delta visualization. */
+	PCGEDITOR_API void RegisterDeltaVisualization(const UScriptStruct* DeltaStruct, TUniquePtr<const IPCGDeltaVisualization> Visualization);
+
+	/** Unregister an external delta visualization. */
+	PCGEDITOR_API void UnregisterDeltaVisualization(const UScriptStruct* DeltaStruct);
+
+	template <typename DeltaType>
+	const IPCGDeltaVisualization* GetDeltaVisualization() const { return GetDeltaVisualization(DeltaType::StaticStruct()); }
+
+	PCGEDITOR_API const IPCGDeltaVisualization* GetDeltaVisualization(const UScriptStruct* DeltaStruct) const;
+
+private:
+	/** Registry for delta visualizations defined inside the PCG Plugin. */
+	TMap<const UScriptStruct*, TUniquePtr<const IPCGDeltaVisualization>> InternalRegistry;
+
+	/** Registry for delta visualizations defined outside the PCG Plugin. These take priority, allowing the user to override default visualization behavior for internal types. */
+	TMap<const UScriptStruct*, TUniquePtr<const IPCGDeltaVisualization>> ExternalRegistry;
+
+	friend class FPCGEditorModule;
+};
diff --git a/Engine/Plugins/PCG/Source/PCGEditor/Public/Nodes/PCGEditorGraphNodeBase.h b/Engine/Plugins/PCG/Source/PCGEditor/Public/Nodes/PCGEditorGraphNodeBase.h
--- a/Engine/Plugins/PCG/Source/PCGEditor/Public/Nodes/PCGEditorGraphNodeBase.h
+++ b/Engine/Plugins/PCG/Source/PCGEditor/Public/Nodes/PCGEditorGraphNodeBase.h
@@ -84,6 +84,9 @@ public:
 	void SetTriggeredGPUReadback(bool bInTriggeredGPUReadback) { bTriggeredGPUReadback = bInTriggeredGPUReadback; }
 	bool GetTriggeredGPUReadback() const { return bTriggeredGPUReadback; }
 
+	void SetAppliedDataOverrides(bool bInAppliedDataOverrides) { bAppliedDataOverrides = bInAppliedDataOverrides; }
+	bool GetAppliedDataOverrides() const { return bAppliedDataOverrides; }
+
 	/** Increase deferred reconstruct counter, calls to ReconstructNode will flag reconstruct to happen when count hits zero */
 	void EnableDeferredReconstruct();
 	/** Decrease deferred reconstruct counter, ReconstructNode will be called if counter hits zero and the node is flagged for reconstruction  */
@@ -97,6 +100,8 @@ public:
 
 	EPCGChangeType UpdateGPUVisualization(const IPCGGraphExecutionSource* InSourceBeingDebugged, const FPCGStack* InStackBeingInspected);
 
+	EPCGChangeType UpdateOverrideVisualization(const IPCGGraphExecutionSource* InSourceBeingDebugged, const FPCGStack* InStackBeingInspected);
+
 	bool CanUserAddRemoveDynamicInputPins() const;
 
 	/** The settings support default value inline constants. */
@@ -202,6 +207,7 @@ protected:
 
 	bool bTriggeredGPUUpload = false;
 	bool bTriggeredGPUReadback = false;
+	bool bAppliedDataOverrides = false;
 
 	/** Bitmask of inactive output pins. Bit N will be set if output pin index N is inactive. */
 	uint64 InactiveOutputPinMask = 0;
diff --git a/Engine/Plugins/PCG/Source/PCGEditor/Public/PCGComponentVisualizer.h b/Engine/Plugins/PCG/Source/PCGEditor/Public/PCGComponentVisualizer.h
--- a/Engine/Plugins/PCG/Source/PCGEditor/Public/PCGComponentVisualizer.h
+++ b/Engine/Plugins/PCG/Source/PCGEditor/Public/PCGComponentVisualizer.h
@@ -2,6 +2,7 @@
 
 #pragma once
 
+#include "DeltaViewportExtensions/PCGDeltaViewportExtension.h"
 #include "Graph/DataOverride/PCGDataOverride.h"
 #include "Graph/PCGSourceDataContainer.h"
 #include "Graph/PCGStackContext.h"
@@ -15,25 +16,34 @@
 #include "Templates/SharedPointer.h"
 
 class FUICommandList;
+class UStaticMesh;
+
+namespace PCG::ComponentVisualizer
+{
+	struct FCachedMeshWireframe;
+}
 
 struct HPCGComponentVisualizerHitProxy : public HComponentVisProxy
 {
 	DECLARE_HIT_PROXY(PCGEDITOR_API);
 
-	HPCGComponentVisualizerHitProxy(const UActorComponent* InComponent, const FPCGStack& InPinStack, int InDataIndex, int InPointIndex);
+	/** Regular point constructor. */
+	HPCGComponentVisualizerHitProxy(const UActorComponent* InComponent, const FPCGStack& InPinStack, int32 InDataIndex, int32 InPointIndex);
 
-	FPCGStack PinStack;
-	int DataIndex;
-	int PointIndex;
-};
+	/** Inserted element constructor. */
+	HPCGComponentVisualizerHitProxy(const UActorComponent* InComponent, const FPCGStack& InPinStack, const FPCGDeltaKey& InInsertedDeltaKey, int32 InInsertedElementIndex);
 
-struct HPCGManualEditHitProxy : public HComponentVisProxy
-{
-	DECLARE_HIT_PROXY(PCGEDITOR_API);
+	FPCGStack PinStack;
+	int32 DataIndex = INDEX_NONE;
+	int32 PointIndex = INDEX_NONE;
 
-	HPCGManualEditHitProxy(const UActorComponent* InComponent);
+	// Inserted element fields
+	bool bIsInsertedElement = false;
+	int32 InsertedElementIndex = INDEX_NONE;
+	FPCGDeltaKey InsertedDeltaKey;
 };
 
+class UPCGNode;
 
 class FPCGComponentVisualizer : public FComponentVisualizer
 {
@@ -54,36 +64,106 @@ public:
 	PCGEDITOR_API virtual TSharedPtr<SWidget> GenerateContextMenu() const override;
 	//~ End FComponentVisualizer Interface
 
+	/** Removes a specific delta by key. */
+	void RestoreDelta(const FPCGDeltaKey& Key);
+
+	/** Removes all restorable deltas from the active collection (delegates to each extension's CollectRestorableKeys). */
+	void RestoreAllDeltas();
+
+	/** Selects a delta element by key and optional sub-index. Looks up the extension from the registry. */
+	void SelectElement(const FPCGDeltaKey& DeltaKey, int32 ElementIndex = INDEX_NONE);
+
+	/** Clears the current selection (EditTarget + gizmo). Call before switching active editing nodes. */
+	PCGEDITOR_API void ClearActiveSelection();
 
 protected:
-	void SetSelectedTransform(const FTransform& Transform);
+	/** Called each frame while the TRS gizmo is being dragged. Updates the cached selection transform. */
+	void OnGizmoDragged(const FTransform& NewTransform);
+
+	/** Called when the TRS gizmo drag ends. Materializes deferred deltas and applies the final transform. */
+	void OnGizmoReleased(const FTransform& FinalTransform);
+
+	/** Removes all deltas from the active editing node's collection and clears the selection. */
 	void RemoveAllDeltas();
+
+	/** Removes the currently selected delta and clears the selection. */
 	void RemoveSelectedDeltas();
+
+	/** Returns true if there is a valid PCG component and an active edit target. */
 	bool HasSelection() const;
+
+	/** Hides and destroys the active TRS gizmo. */
 	void ResetGizmo();
 
-	PCGEDITOR_API void EndCurrentElementManualEditing();
+	/** Creates a TRS gizmo at the given transform with drag/release callbacks bound to this visualizer. */
+	void SetupGizmo(UInteractiveGizmoManager* GizmoManager, const FTransform& InitialTransform);
+
+	/** Rebuilds the active key bindings from the extension's registered input actions. */
+	void RefreshKeyBindings(IPCGDeltaViewportExtension* Extension);
 
-	/** Registered Menu commands*/
+	/** Registered Menu commands */
 	TSharedPtr<FUICommandList> ComponentVisualizerActions;
 
-	bool bManualEditMode = false;
+	TWeakObjectPtr<UPCGNode> ActiveEditingNode;
 
 	TWeakObjectPtr<UPCGComponent> LastPCGComponent;
 
+	/** Cached result of resolving the mutable delta collection for the active editing node. Invalidated when the active node changes. */
+	struct FPCGCachedDeltaContext
+	{
+		FPCGSourceDataContainer* DataContainer = nullptr;
+		FPCGSourceDataStorageKey StorageKey;
+		FPCGStack PinStack;
+
+		bool IsValid() const { return DataContainer != nullptr; }
+		void Reset() { *this = {}; }
+	};
+
+	FPCGCachedDeltaContext CachedDeltaContext;
+
+	/** Tracks the state of the currently selected delta element in the viewport, including the owning component, delta key, gizmo transform, and per-extension metadata. */
 	struct FEditState
 	{
 		TWeakObjectPtr<UPCGComponent> PCGComponent;
 		FPCGStack Stack;
 		FPCGSourceDataStorageKey StorageKey;
 		FPCGDeltaKey DeltaKey;
-		bool bWasInDelta;
+		bool bWasInDelta = false;
 		FTransform Transform;
+		FTransform OriginalTransform;
 
 		bool bModified = false;
+
+		/** True while the gizmo is actively being dragged. Used to gate the drag wireframe. */
+		bool bGizmoDragging = false;
+
+		/** True when this is a new edit delta that hasn't been stored in the collection yet (deferred until first gizmo move). */
+		bool bDeltaDeferred = false;
+
+		/** Extension governing this selection's viewport behavior. Null when no delta is active yet (deferred). */
+		IPCGDeltaViewportExtension* ActiveExtension = nullptr;
+
+		/** UScriptStruct of the delta type, for registry re-lookups when ActiveExtension is null. */
+		const UScriptStruct* DeltaStructType = nullptr;
+
+		/** Index of the selected element within an insertion delta (e.g. which inserted point is selected). */
+		int32 SelectedElementIndex = INDEX_NONE;
+
+		/** Original index into the point data array at creation time (for ElementIndex debug field). */
+		int32 OriginalPointIndex = INDEX_NONE;
 	};
 
+	/** Active viewport selection. Set when a delta element is clicked. Reset on deselect or node change. */
 	TOptional<FEditState> EditTarget = {};
 
-	UCombinedTransformGizmo* Gizmo = nullptr;
+	TWeakObjectPtr<UCombinedTransformGizmo> Gizmo;
+
+	/** Key bindings from the active extension, cached when the extension changes. */
+	TArray<FPCGDeltaKeyBinding> ActiveKeyBindings;
+
+	/** Persists across frames to avoid per-frame index buffer copies. */
+	TUniquePtr<PCG::ComponentVisualizer::FCachedMeshWireframe> CachedWireframe;
+
+	/** Resolved default point mesh, cached to avoid per-frame LoadSynchronous. */
+	TWeakObjectPtr<UStaticMesh> CachedDefaultPointMesh;
 };
diff --git a/Engine/Plugins/PCG/Source/PCGEditor/Public/PCGEditor.h b/Engine/Plugins/PCG/Source/PCGEditor/Public/PCGEditor.h
--- a/Engine/Plugins/PCG/Source/PCGEditor/Public/PCGEditor.h
+++ b/Engine/Plugins/PCG/Source/PCGEditor/Public/PCGEditor.h
@@ -28,6 +28,7 @@ class SPCGEditorGraphDetailsView;
 class SPCGEditorGraphDeterminismListView;
 class SPCGEditorGraphFind;
 class SPCGEditorGraphEmbeddedSubgraphsView;
+class SPCGEditorGraphDataOverridesView;
 class SPCGEditorGraphLogView;
 class SPCGEditorGraphNodePalette;
 class SPCGEditorGraphProfilingView;
@@ -112,6 +113,7 @@ namespace PCGEditorTabs
 		TEXT("Viewport4") };
 	constexpr FLazyName EmbeddedSubgraphsID = TEXT("EmbeddedSubgraphs");
 	constexpr FLazyName CodeEditorID = TEXT("CodeEditor");
+	constexpr FLazyName DataOverridesID = TEXT("DataOverrides");
 
 	constexpr int32 NumPropertyDetailTabs = sizeof(PropertyDetailsID) / sizeof(FLazyName);
 	constexpr int32 NumAttributeTabs = sizeof(AttributesID) / sizeof(FLazyName);
@@ -123,6 +125,9 @@ namespace PCGEditorTabs
 
 class FPCGEditor : public FWorkflowCentricApplication, public FGCObject, public FSelfRegisteringEditorUndoClient
 {
+	friend class FPCGEditorModule;
+	friend class FPCGManualEditPanelManager;
+
 public:
 	/** Edits the specified PCGGraph */
 	PCGEDITOR_API void InitializePCGEditorGraph(UPCGGraph* InPCGGraph);
@@ -280,6 +285,7 @@ PRAGMA_ENABLE_DEPRECATION_WARNINGS
 	PCGEDITOR_API TSharedPtr<FWorkflowTabFactory> CreatePropertyDetailsTabFactory(int32 Index, TSharedRef<FWorkspaceItem> Group);
 	PCGEDITOR_API TSharedPtr<FWorkflowTabFactory> CreateViewportTabFactory(int32 Index, TSharedRef<FWorkspaceItem> Group);
 	PCGEDITOR_API TSharedPtr<FWorkflowTabFactory> CreateEmbeddedSubgraphsTabFactory();
+	PCGEDITOR_API TSharedPtr<FWorkflowTabFactory> CreateDataOverridesTabFactory();
 
 	UE_DEPRECATED(5.8, "Use CreateCodeEditorTabFactory() instead")
 	PCGEDITOR_API TSharedPtr<FWorkflowTabFactory> CreateNodeSourceTabFactory();
@@ -415,6 +421,21 @@ private:
 	/** Disable node debug state for all nodes */
 	void OnDisableDebugOnAllNodes();
 
+	/** Toggle viewport editing for the selected node (one-at-a-time). */
+	void OnToggleViewportEditing();
+	/** Returns true if the selected node supports viewport editing. */
+	bool CanToggleViewportEditing() const;
+	/** Returns the check box state of viewport editing for the selected node. */
+	ECheckBoxState GetViewportEditingCheckState() const;
+
+	/** Mark/unmark a node for persistent viewport editing. */
+	void OnMarkForViewportEditing();
+	/** Returns the check box state for whether the selected node is marked for viewport editing. */
+	ECheckBoxState GetMarkForViewportEditingCheckState() const;
+
+	/** Notify the module to re-evaluate manual edit panel visibility after node state changes */
+	void NotifyModuleManualEditStateChanged();
+
 	/** Cancels the current execution of the selected graph */
 	void OnCancelExecution_Clicked();
 
@@ -577,6 +598,9 @@ private:
 	/** Get or Create a new embedded subgraphs tab widget */
 	TSharedRef<SPCGEditorGraphEmbeddedSubgraphsView> GetOrCreateEmbeddedSubgraphsWidget();
 
+	/** Get or Create a new data overrides tab widget */
+	TSharedRef<SPCGEditorGraphDataOverridesView> GetOrCreateDataOverridesWidget();
+
 	/** Called when the component inspected is generated/cleaned */
 	void OnSourceGenerated(IPCGGraphExecutionSource* InSource);
 
@@ -635,6 +659,7 @@ private:
 	TSharedPtr<SPCGEditorGraphUserParametersView> UserParamsWidget;
 	TSharedPtr<SPCGEditorViewport> ViewportWidgets[PCGEditorTabs::NumViewportTabs];
 	TSharedPtr<SPCGEditorGraphEmbeddedSubgraphsView> EmbeddedSubgraphsWidget;
+	TSharedPtr<SPCGEditorGraphDataOverridesView> DataOverridesWidget;
 
 	TSharedPtr<FUICommandList> GraphEditorCommands;
 
diff --git a/Engine/Plugins/PCG/Source/PCGEditor/Public/PCGEditorCommands.h b/Engine/Plugins/PCG/Source/PCGEditor/Public/PCGEditorCommands.h
--- a/Engine/Plugins/PCG/Source/PCGEditor/Public/PCGEditorCommands.h
+++ b/Engine/Plugins/PCG/Source/PCGEditor/Public/PCGEditorCommands.h
@@ -46,6 +46,8 @@ public:
 	TSharedPtr<FUICommandInfo> SelectNamedRerouteDeclaration;
 	TSharedPtr<FUICommandInfo> JumpToSource;
 	TSharedPtr<FUICommandInfo> AutoFocusViewport;
+	TSharedPtr<FUICommandInfo> ToggleViewportEditing;
+	TSharedPtr<FUICommandInfo> MarkForViewportEditing;
 };
 
 struct FPCGSpawnNodeCommandInfo
diff --git a/Engine/Plugins/PCG/Source/PCGEditor/Public/PCGEditorModule.h b/Engine/Plugins/PCG/Source/PCGEditor/Public/PCGEditorModule.h
--- a/Engine/Plugins/PCG/Source/PCGEditor/Public/PCGEditorModule.h
+++ b/Engine/Plugins/PCG/Source/PCGEditor/Public/PCGEditorModule.h
@@ -5,6 +5,8 @@
 #include "Editor/IPCGEditorModule.h"
 
 #include "PCGGraphExecutionStateInterface.h"
+#include "DeltaViewportExtensions/PCGDeltaViewportExtension.h"
+#include "DeltaVisualizations/PCGDeltaVisualizationRegistry.h"
 
 #include "Utils/PCGNodeVisualLogs.h"
 
@@ -14,7 +16,11 @@
 #include "Toolkits/AssetEditorToolkit.h"
 
 class FPCGLoadDataAssetChangeTracker;
+class FPCGManualEditPanelManager;
 class ILevelEditor;
+class SPCGManualEditPanel;
+class UPCGGraph;
+
 // Logs
 DECLARE_LOG_CATEGORY_EXTERN(LogPCGEditor, Log, All);
 
@@ -42,10 +48,20 @@ public:
 	virtual EPCGEditorPermissionMode GetPermissionMode() const { return PermissionMode; }
 	virtual void TransferPropertyBagMetadataIntoHierarchy(UPCGGraph* Graph) const override;
 	virtual UObject* CreatePropertyBagHierarchyRootObject(UPCGGraph* Graph) const override;
-	
+
+	static const FPCGDeltaVisualizationRegistry& GetConstDeltaVisualizationRegistry();
+	static FPCGDeltaVisualizationRegistry& GetMutableDeltaVisualizationRegistry();
+	static const FPCGDeltaViewportExtensionRegistry& GetConstDeltaViewportExtensionRegistry();
+	static FPCGDeltaViewportExtensionRegistry& GetMutableDeltaViewportExtensionRegistry();
+
 	DECLARE_DERIVED_EVENT(FPCGEditorModule, IPCGEditorModule::FPermissionModeChanged, FPermissionModeChanged);
 	virtual FPermissionModeChanged& OnPermissionModeChanged() override { return PermissionModeChangedEvent; }
 
+	void UpdateManualEditPanelVisibility();
+
+	/** Returns the manual edit panel if one is currently active, or nullptr. */
+	TSharedPtr<SPCGManualEditPanel> GetManualEditPanel() const;
+
 protected:
 	virtual void OnScheduleGraph(const FPCGStackContext& StackContext) override;
 	virtual void OnGraphPreSave(UPCGGraph* Graph, FObjectPreSaveContext ObjectSaveContext) override;
@@ -66,6 +82,10 @@ protected:
 	void UnregisterSettings();
 	void RegisterPCGDataVisualizations();
 	void UnregisterPCGDataVisualizations();
+	void RegisterDeltaVisualizations();
+	void UnregisterDeltaVisualizations();
+	void RegisterDeltaViewportExtensions();
+	void UnregisterDeltaViewportExtensions();
 	void RegisterPinColorAndIcons();
 	void UnregisterPinColorAndIcons();
 	void RegisterTrackerFactories();
@@ -115,4 +135,9 @@ protected:
 	TUniquePtr<FPCGLoadDataAssetChangeTracker> LoadDataAssetChangeTracker;
 
 	TArray<FName> VisualizersToUnregisterOnShutdown;
+
+	TUniquePtr<FPCGManualEditPanelManager> ManualEditPanelManager;
+
+	FPCGDeltaVisualizationRegistry DeltaVisualizationRegistry;
+	FPCGDeltaViewportExtensionRegistry DeltaViewportExtensionRegistry;
 };
diff --git a/Engine/Plugins/PCG/Source/PCGEditor/Public/Schema/PCGEditorGraphSchema.h b/Engine/Plugins/PCG/Source/PCGEditor/Public/Schema/PCGEditorGraphSchema.h
--- a/Engine/Plugins/PCG/Source/PCGEditor/Public/Schema/PCGEditorGraphSchema.h
+++ b/Engine/Plugins/PCG/Source/PCGEditor/Public/Schema/PCGEditorGraphSchema.h
@@ -87,6 +87,7 @@ public:
 	// To customize if we want to have special brushes for Inspect and Debug on the nodes.
 	virtual bool ShouldAddInspectBrushOnNode() const { return true; }
 	virtual bool ShouldAddDebugBrushOnNode() const { return true; }
+	virtual bool ShouldAddViewportEditBrushOnNode() const { return true; }
 
 protected:
 	PCGEDITOR_API virtual void GetGraphActions(FGraphActionMenuBuilder& ActionMenuBuilder, const FPCGActionsFilter& ActionFilter, bool bIsContextual) const;
