[PCG] Fixed issues with managed resources when using standalone graphs

#Synced-CL 52517676

diff --git a/Engine/Plugins/PCG/Source/PCG/Private/PCGManagedResourceContainer.cpp b/Engine/Plugins/PCG/Source/PCG/Private/PCGManagedResourceContainer.cpp
--- a/Engine/Plugins/PCG/Source/PCG/Private/PCGManagedResourceContainer.cpp
+++ b/Engine/Plugins/PCG/Source/PCG/Private/PCGManagedResourceContainer.cpp
@@ -15,6 +15,22 @@
 #include "EditorActorFolders.h"
 #endif
 
+void FPCGManagedResourceContainerHelper::ForgetAll()
+{
+	check(IsValid());
+	PCG::TScopeLock ResourcesLock(Owner->GetExecutionState().GetManagedResourceContainerLock());
+
+	for (TObjectPtr<UPCGManagedResource>& ManagedResource : Container->GeneratedResources)
+	{
+		if (UPCGManagedResource* ValidManagedResource = GetValid(ManagedResource))
+		{
+			ValidManagedResource->Rename(nullptr, GetTransientPackage(), REN_DoNotDirty | REN_DontCreateRedirectors | REN_NonTransactional);
+		}
+	}
+	
+	Container->GeneratedResources.Empty();
+}
+
 void FPCGManagedResourceContainerHelper::AddManagedResource(UPCGManagedResource* InManagedResource)
 {
 	check(IsValid());
diff --git a/Engine/Plugins/PCG/Source/PCG/Private/Subsystems/IPCGBaseSubsystem.cpp b/Engine/Plugins/PCG/Source/PCG/Private/Subsystems/IPCGBaseSubsystem.cpp
--- a/Engine/Plugins/PCG/Source/PCG/Private/Subsystems/IPCGBaseSubsystem.cpp
+++ b/Engine/Plugins/PCG/Source/PCG/Private/Subsystems/IPCGBaseSubsystem.cpp
@@ -306,9 +306,12 @@ void IPCGBaseSubsystem::OnPCGSourceGenerationDone(IPCGGraphExecutionSource* InEx
 			PCGEditorModule->ClearExecutionMetadata(ExecutionSource);
 		}
 
-		// Mark source as invalid, request GC (non-blocking)
-		ExecutionSource->MarkAsGarbage();
-		GEngine->ForceGarbageCollection();
+		// Since the execution source is setup to be deleted after this execution, we need to forget all the managed resources of this source.
+		FPCGManagedResourceContainerHelper ManagedResourceContainerHelper(ExecutionSource);
+		if (ManagedResourceContainerHelper.IsValid())
+		{
+			ManagedResourceContainerHelper.ForgetAll();
+		}
 	}
 }
 
diff --git a/Engine/Plugins/PCG/Source/PCG/Public/PCGDefaultWorldObjectExecutionSource.h b/Engine/Plugins/PCG/Source/PCG/Public/PCGDefaultWorldObjectExecutionSource.h
--- a/Engine/Plugins/PCG/Source/PCG/Public/PCGDefaultWorldObjectExecutionSource.h
+++ b/Engine/Plugins/PCG/Source/PCG/Public/PCGDefaultWorldObjectExecutionSource.h
@@ -31,6 +31,11 @@ public:
 
 private:
 	TWeakObjectPtr<UObject> WorldObject = nullptr;
+
+	/** Managed resources need while the execution occurs. They will be forgotten when execution is completed*/
+	UPROPERTY()
+	FPCGManagedResourceContainer ManagedResourceContainer;
+	mutable FTransactionallySafeCriticalSection ManagedResourcesLock;
 };
 
 class FPCGDefaultWorldObjectExecutionState : public FPCGDefaultExecutionState
@@ -42,14 +47,10 @@ public:
 	PCG_API virtual FString GetDebugName() const override;
 	PCG_API virtual UWorld* GetWorld() const override;
 	PCG_API virtual UObject* GetTarget() const override;
-	virtual FPCGManagedResourceContainer* GetManagedResourceContainer() override { return &ManagedResourceContainer; }
-	virtual FTransactionallySafeCriticalSection* GetManagedResourceContainerLock() override { return &ManagedResourcesLock; }
+	virtual FPCGManagedResourceContainer* GetManagedResourceContainer() override { return GetTypedSource() ? &GetTypedSource()->ManagedResourceContainer : nullptr; }
+	virtual FTransactionallySafeCriticalSection* GetManagedResourceContainerLock() override { return GetTypedSource() ? &GetTypedSource()->ManagedResourcesLock : nullptr; }
 
 protected:
 	UPCGDefaultWorldObjectExecutionSource* GetTypedSource() const { return CastChecked<UPCGDefaultWorldObjectExecutionSource>(GetSource()); }
 	UObject* GetWorldObject() const { return GetTypedSource()->GetWorldObject(); }
-
-private:
-	FPCGManagedResourceContainer ManagedResourceContainer;
-	mutable FTransactionallySafeCriticalSection ManagedResourcesLock;
 };
\ No newline at end of file
diff --git a/Engine/Plugins/PCG/Source/PCG/Public/PCGManagedResourceContainer.h b/Engine/Plugins/PCG/Source/PCG/Public/PCGManagedResourceContainer.h
--- a/Engine/Plugins/PCG/Source/PCG/Public/PCGManagedResourceContainer.h
+++ b/Engine/Plugins/PCG/Source/PCG/Public/PCGManagedResourceContainer.h
@@ -75,6 +75,8 @@ struct FPCGManagedResourceContainerHelper final : public FPCGManagedResourceCons
 	{
 	}
 
+	PCG_API void ForgetAll();
+
 	PCG_API void AddManagedResource(UPCGManagedResource* InManagedResource);
 	PCG_API void CleanupUnusedManagedResources();
 	PCG_API void SetManagedResources(TArray<TObjectPtr<UPCGManagedResource>> InManagedResources);
