[PCG] Add configurable font size to HLSL Code Editor

#Synced-CL 52605476

diff --git a/Engine/Plugins/PCG/Source/PCGEditor/Private/PCGHLSLSyntaxHighlighter.cpp b/Engine/Plugins/PCG/Source/PCGEditor/Private/PCGHLSLSyntaxHighlighter.cpp
--- a/Engine/Plugins/PCG/Source/PCGEditor/Private/PCGHLSLSyntaxHighlighter.cpp
+++ b/Engine/Plugins/PCG/Source/PCGEditor/Private/PCGHLSLSyntaxHighlighter.cpp
@@ -4,12 +4,37 @@
 
 #include "Compute/PCGHLSLSyntaxTokenizer.h"
 
+#include "PCGEditorSettings.h"
 #include "PCGEditorStyle.h"
+#include "Styling/CoreStyle.h"
 #include "Styling/SlateTypes.h"
 #include "Framework/Text/IRun.h"
 #include "Framework/Text/TextLayout.h"
 #include "Framework/Text/SlateTextUnderlineLineHighlighter.h"
 
+namespace PCGHLSLSyntaxHighlighterHelpers
+{
+	FHLSLSyntaxHighlighterMarshaller::FSyntaxTextStyle BuildSyntaxTextStyle(int32 InFontSize)
+	{
+		auto GetStyle = [InFontSize](const TCHAR* StyleName)
+		{
+			FTextBlockStyle Style(FAppStyle::Get().GetWidgetStyle<FTextBlockStyle>(StyleName));
+			Style.Font.Size = InFontSize;
+			return Style;
+		};
+
+		return FHLSLSyntaxHighlighterMarshaller::FSyntaxTextStyle(
+			GetStyle(TEXT("SyntaxHighlight.SourceCode.Normal")),
+			GetStyle(TEXT("SyntaxHighlight.SourceCode.Operator")),
+			GetStyle(TEXT("SyntaxHighlight.SourceCode.Keyword")),
+			GetStyle(TEXT("SyntaxHighlight.SourceCode.String")),
+			GetStyle(TEXT("SyntaxHighlight.SourceCode.Number")),
+			GetStyle(TEXT("SyntaxHighlight.SourceCode.Comment")),
+			GetStyle(TEXT("SyntaxHighlight.SourceCode.PreProcessorKeyword")),
+			GetStyle(TEXT("SyntaxHighlight.SourceCode.Error")));
+	}
+}
+
 TSharedRef<FPCGHLSLSyntaxHighlighter> FPCGHLSLSyntaxHighlighter::Create(const FSyntaxTextStyle& InSyntaxTextStyle)
 {
 	return MakeShareable(new FPCGHLSLSyntaxHighlighter(MakeShareable(new FPCGHLSLSyntaxTokenizer()), InSyntaxTextStyle));
@@ -42,6 +67,22 @@ FPCGHLSLSyntaxHighlighter::FPCGHLSLSyntaxHighlighter(TSharedPtr<ISyntaxTokenizer
 {
 }
 
+void FPCGHLSLSyntaxHighlighter::Tick(float DeltaTime)
+{
+	const int32 FontSize = GetDefault<UPCGEditorSettings>()->CodeEditorFontSize;
+	if (FontSize != CachedFontSize)
+	{
+		CachedFontSize = FontSize;
+		SyntaxTextStyle = PCGHLSLSyntaxHighlighterHelpers::BuildSyntaxTextStyle(FontSize);
+		MakeDirty();
+	}
+}
+
+TStatId FPCGHLSLSyntaxHighlighter::GetStatId() const
+{
+	RETURN_QUICK_DECLARE_CYCLE_STAT(FPCGHLSLSyntaxHighlighter, STATGROUP_Tickables);
+}
+
 void FPCGHLSLSyntaxHighlighter::ParseTokens(const FString& SourceString, FTextLayout& TargetTextLayout,	TArray<ISyntaxTokenizer::FTokenizedLine> TokenizedLines)
 {
 	LineHighlightsToAdd.Empty();
@@ -87,13 +128,5 @@ FTextLayout::FNewLineData FPCGHLSLSyntaxHighlighter::ProcessTokenizedLine(const
 
 FPCGHLSLSyntaxHighlighter::FSyntaxTextStyle FPCGHLSLSyntaxHighlighter::GetSyntaxTextStyle()
 {
-	return FSyntaxTextStyle(
-		FAppStyle::Get().GetWidgetStyle<FTextBlockStyle>("SyntaxHighlight.SourceCode.Normal"),
-		FAppStyle::Get().GetWidgetStyle<FTextBlockStyle>("SyntaxHighlight.SourceCode.Operator"),
-		FAppStyle::Get().GetWidgetStyle<FTextBlockStyle>("SyntaxHighlight.SourceCode.Keyword"),
-		FAppStyle::Get().GetWidgetStyle<FTextBlockStyle>("SyntaxHighlight.SourceCode.String"),
-		FAppStyle::Get().GetWidgetStyle<FTextBlockStyle>("SyntaxHighlight.SourceCode.Number"),
-		FAppStyle::Get().GetWidgetStyle<FTextBlockStyle>("SyntaxHighlight.SourceCode.Comment"),
-		FAppStyle::Get().GetWidgetStyle<FTextBlockStyle>("SyntaxHighlight.SourceCode.PreProcessorKeyword"),
-		FAppStyle::Get().GetWidgetStyle<FTextBlockStyle>("SyntaxHighlight.SourceCode.Error"));
+	return PCGHLSLSyntaxHighlighterHelpers::BuildSyntaxTextStyle(GetDefault<UPCGEditorSettings>()->CodeEditorFontSize);
 }
diff --git a/Engine/Plugins/PCG/Source/PCGEditor/Private/PCGHLSLSyntaxHighlighter.h b/Engine/Plugins/PCG/Source/PCGEditor/Private/PCGHLSLSyntaxHighlighter.h
--- a/Engine/Plugins/PCG/Source/PCGEditor/Private/PCGHLSLSyntaxHighlighter.h
+++ b/Engine/Plugins/PCG/Source/PCGEditor/Private/PCGHLSLSyntaxHighlighter.h
@@ -6,6 +6,8 @@
 
 #include "Compute/PCGCompilerDiagnostic.h"
 
+#include "Tickable.h"
+#include "TickableEditorObject.h"
 #include "Framework/Text/SyntaxTokenizer.h"
 #include "Text/HLSLSyntaxHighlighterMarshaller.h"
 
@@ -15,7 +17,7 @@ struct FPCGCompilerDiagnostic;
 /**
  * Get/set the raw text to/from a text layout, and also inject syntax highlighting for our rich-text markup
  */
-class FPCGHLSLSyntaxHighlighter : public FHLSLSyntaxHighlighterMarshaller
+class FPCGHLSLSyntaxHighlighter : public FHLSLSyntaxHighlighterMarshaller, public FTickableEditorObject
 {
 public:
 	static TSharedRef<FPCGHLSLSyntaxHighlighter> Create(const FSyntaxTextStyle& InSyntaxTextStyle = GetSyntaxTextStyle());
@@ -24,6 +26,12 @@ public:
 	void SetCompilerMessages(const FPCGCompilerDiagnostics& InCompilerMessages);
 	void ClearCompilerMessages();
 
+	// @todo_pcg: Replace tick with an OnSettingChanged delegate on UPCGEditorSettings.
+	//~ Begin FTickableEditorObject interface
+	virtual void Tick(float DeltaTime) override;
+	virtual TStatId GetStatId() const override;
+	//~ End FTickableEditorObject interface
+
 protected:
 	FPCGHLSLSyntaxHighlighter(TSharedPtr<ISyntaxTokenizer> InTokenizer, const FSyntaxTextStyle& InSyntaxTextStyle);
 
@@ -33,5 +41,6 @@ protected:
 	TMultiMap<int32 /*Line*/, FPCGCompilerDiagnostic> CompilerMessages;
 
 private:
+	int32 CachedFontSize = 0;
 	TArray<FTextLineHighlight> LineHighlightsToAdd;
 };
diff --git a/Engine/Plugins/PCG/Source/PCGEditor/Public/PCGEditorSettings.h b/Engine/Plugins/PCG/Source/PCGEditor/Public/PCGEditorSettings.h
--- a/Engine/Plugins/PCG/Source/PCGEditor/Public/PCGEditorSettings.h
+++ b/Engine/Plugins/PCG/Source/PCGEditor/Public/PCGEditorSettings.h
@@ -345,6 +345,10 @@ public:
 	UPROPERTY(EditAnywhere, config, Category = "Viewport")
 	bool bAutoFocusViewport = true;
 
+	/** Font size used in the HLSL Code Editor. */
+	UPROPERTY(EditAnywhere, config, Category = "Code Editor", meta = (ClampMin = 6, ClampMax = 32))
+	int32 CodeEditorFontSize = 9;
+
 	// Begin deprecated properties
 	UE_DEPRECATED(5.6, "bEnableNavigateToNativeNodes has been replaced by NodeDoubleClickAction.")
 	UPROPERTY()
