[PCG] Fix test Plugins.PCG.Display.PCGComponentInNewPCGVolumes.

Don't know why it was calling that an asset deletion function to remove an actor.

#rnx
#Synced-CL 52611570

diff --git a/Engine/Plugins/PCG/Source/PCGEditor/Private/Tests/Display/PCGComponentInNewPCGVolumes.cpp b/Engine/Plugins/PCG/Source/PCGEditor/Private/Tests/Display/PCGComponentInNewPCGVolumes.cpp
--- a/Engine/Plugins/PCG/Source/PCGEditor/Private/Tests/Display/PCGComponentInNewPCGVolumes.cpp
+++ b/Engine/Plugins/PCG/Source/PCGEditor/Private/Tests/Display/PCGComponentInNewPCGVolumes.cpp
@@ -4,14 +4,10 @@
 #include "Engine/World.h"
 
 #include "Editor.h"
-#include "Editor/UnrealEd/Public/PackageTools.h"
-#include "ObjectTools.h"
-#include "AssetRegistry/AssetRegistryModule.h"
 
 #include "PCGVolume.h"
 #include "PCGComponent.h"
 #include "Tests/PCGTestsCommon.h"
-#include "Tests/AutomationEditorCommon.h"
 
 
 #if WITH_AUTOMATION_TESTS
@@ -58,20 +54,9 @@ bool FPCGComponentInNewPCGVolumes::RunTest(const FString& Parameters)
 
 	PCGComponent->ClearPCGLink(APCGVolume::StaticClass());
 
-	// Cleanup the VolumeActor
-	FAssetRegistryModule& AssetRegistryModule = FModuleManager::LoadModuleChecked<FAssetRegistryModule>("AssetRegistry");
-	IAssetRegistry& AssetRegistry = AssetRegistryModule.Get();
-	FAssetData PCGVolumeAssetData(VolumeActor);
-	AssetRegistry.AssetDeleted(PCGVolumeAssetData.GetAsset());
-
-	bool bSuccessful = ObjectTools::DeleteSingleObject(VolumeActor, true);
-	if (!bSuccessful)
-	{
-		//Clear references to the object so we can delete it
-		FAutomationEditorCommonUtils::NullReferencesToObject(VolumeActor);
-		bSuccessful = ObjectTools::DeleteSingleObject(VolumeActor, false);
-		TestTrue(TEXT("References to PCGVolume component could not be cleaned"), bSuccessful);
-	}
+	// Cleanup the VolumeActor.
+	const bool bSuccessful = World->DestroyActor(VolumeActor);
+	TestTrue(TEXT("References to PCGVolume component could not be cleaned"), bSuccessful);
 
 	return true;
 };
