[PCG] Fix parsing test failures on Linux

#rnx
#Synced-CL 52611570

diff --git a/Engine/Plugins/PCG/Source/PCG/Private/Tests/Elements/PCGParseStringTests.cpp b/Engine/Plugins/PCG/Source/PCG/Private/Tests/Elements/PCGParseStringTests.cpp
--- a/Engine/Plugins/PCG/Source/PCG/Private/Tests/Elements/PCGParseStringTests.cpp
+++ b/Engine/Plugins/PCG/Source/PCG/Private/Tests/Elements/PCGParseStringTests.cpp
@@ -84,7 +84,7 @@ bool FPCGParseStringTests_Bool::RunTest(const FString& Parameters)
 bool FPCGParseStringTests_Int32::RunTest(const FString& Parameters)
 {
 	return RunTestInternal<int32>({
-#if PLATFORM_MAC  // @todo_pcg: Mac Atoi handles this kind of overflow slightly differently, returning -1, which is a valid index. Disabling temporarily until further investigation with the core team.
+#if PLATFORM_MAC || PLATFORM_LINUX  // @todo_pcg: Mac/Linux Atoi handles this kind of overflow slightly differently, returning -1, which is a valid index. Disabling temporarily until further investigation with the core team.
 	"0", "1", "-1", "1254", "Bla", "0xff", "0547"},
 	{0, 1, -1, 1254, 0, 0xff, 547}
 #else
diff --git a/Engine/Plugins/PCG/Source/PCG/Private/Tests/Parsing/PCGParsingTest.cpp b/Engine/Plugins/PCG/Source/PCG/Private/Tests/Parsing/PCGParsingTest.cpp
--- a/Engine/Plugins/PCG/Source/PCG/Private/Tests/Parsing/PCGParsingTest.cpp
+++ b/Engine/Plugins/PCG/Source/PCG/Private/Tests/Parsing/PCGParsingTest.cpp
@@ -290,7 +290,7 @@ bool FPCGParsingTest_Robustness::RunTest(const FString& Parameters)
 	Results = PCGParsingTestPrivate::ParseAndRetrieveIndices(FString::FromInt(-1 * (ArraySize + 1)), ArraySize);
 	UTEST_EQUAL("Index outside the bounds of the array", Results.GetTotalIndexCount(), 0);
 
-#if !PLATFORM_MAC  // @todo_pcg: Mac Atoi handles this kind of overflow slightly differently, returning -1, which is a valid index. Disabling temporarily until further investigation with the core team.
+#if !PLATFORM_MAC && !PLATFORM_LINUX  // @todo_pcg: Mac/Linux Atoi handles this kind of overflow slightly differently, returning -1, which is a valid index. Disabling temporarily until further investigation with the core team.
 	Results = PCGParsingTestPrivate::ParseAndRetrieveIndices("65465465498475309485730495873409587654684986146162342348438", std::numeric_limits<int32>::max());
 	UTEST_EQUAL("Large value converts to int32 max which is the max size of the array which fails with no result", Results.GetTotalIndexCount(), 0);
 #endif
