[PCG] Add Keep Zero Density Points option to Get Texture Data

#Synced-CL 52605476

diff --git a/Engine/Plugins/PCG/Source/PCG/Private/Elements/PCGTextureSampler.cpp b/Engine/Plugins/PCG/Source/PCG/Private/Elements/PCGTextureSampler.cpp
--- a/Engine/Plugins/PCG/Source/PCG/Private/Elements/PCGTextureSampler.cpp
+++ b/Engine/Plugins/PCG/Source/PCG/Private/Elements/PCGTextureSampler.cpp
@@ -334,6 +334,7 @@ bool FPCGTextureSamplerElement::ExecuteInternal(FPCGContext* InContext) const
 	const FString DataLabel = bIsTextureArray ? FString::Format(TEXT("{0}_{1}"), { Settings->Texture.GetAssetName(), TextureArrayIndex }) : Settings->Texture.GetAssetName();
 	OutTaggedData.Tags.Add(PCGComputeHelpers::GetPrefixedDataLabel(DataLabel));
 
+	BaseTextureData->bKeepZeroDensityPoints = Settings->bKeepZeroDensityPoints;
 	BaseTextureData->bUseDensitySourceChannel = Settings->bUseDensitySourceChannel;
 	BaseTextureData->ColorChannel = Settings->ColorChannel;
 	BaseTextureData->Filter = Settings->Filter;
diff --git a/Engine/Plugins/PCG/Source/PCG/Public/Elements/PCGTextureSampler.h b/Engine/Plugins/PCG/Source/PCG/Public/Elements/PCGTextureSampler.h
--- a/Engine/Plugins/PCG/Source/PCG/Public/Elements/PCGTextureSampler.h
+++ b/Engine/Plugins/PCG/Source/PCG/Public/Elements/PCGTextureSampler.h
@@ -123,6 +123,10 @@ public:
 	bool bForceEditorOnlyCPUSampling = false;
 #endif
 
+	/** If enabled, output points with a density value of zero will NOT be automatically filtered out. */
+	UPROPERTY(BlueprintReadWrite, EditAnywhere, Category = Settings, meta = (PCG_Overridable))
+	bool bKeepZeroDensityPoints = false;
+
 	/** By default, texture loading is asynchronous, can force it synchronous if needed. */
 	UPROPERTY(BlueprintReadWrite, EditAnywhere, Category = "Settings|Debug")
 	bool bSynchronousLoad = false;
