[PCG] Add Get Texture Info node

#Synced-CL 52605476

diff --git a/Engine/Plugins/PCG/Source/PCG/Private/Data/PCGRenderTargetData.cpp b/Engine/Plugins/PCG/Source/PCG/Private/Data/PCGRenderTargetData.cpp
--- a/Engine/Plugins/PCG/Source/PCG/Private/Data/PCGRenderTargetData.cpp
+++ b/Engine/Plugins/PCG/Source/PCG/Private/Data/PCGRenderTargetData.cpp
@@ -25,6 +25,9 @@ void UPCGRenderTargetData::Initialize(UTextureRenderTarget2D* InRenderTarget, co
 		Height = RenderTarget->SizeY;
 		Format = RenderTarget->GetFormat();
 
+		const FTextureRHIRef InitRHI = GetTextureRHI();
+		NumMips = InitRHI ? InitRHI->GetDesc().NumMips : 0;
+
 		if (!bSkipReadbackToCPU)
 		{
 			if (FTextureRenderTargetResource* RTResource = RenderTarget->GameThread_GetRenderTargetResource())
diff --git a/Engine/Plugins/PCG/Source/PCG/Private/Data/PCGTexture2DArrayData.cpp b/Engine/Plugins/PCG/Source/PCG/Private/Data/PCGTexture2DArrayData.cpp
--- a/Engine/Plugins/PCG/Source/PCG/Private/Data/PCGTexture2DArrayData.cpp
+++ b/Engine/Plugins/PCG/Source/PCG/Private/Data/PCGTexture2DArrayData.cpp
@@ -69,6 +69,7 @@ bool UPCGTexture2DArrayData::InitializeInternal(const FPCGTexture2DArrayDataInit
 		const FRHITextureDesc& Desc = GetRHI()->GetDesc();
 		Resolution = Desc.Extent;
 		ArraySize = Desc.ArraySize;
+		NumMips = Desc.NumMips;
 		Format = Desc.Format;
 	}
 
@@ -100,6 +101,7 @@ void UPCGTexture2DArrayData::CopyTexture2DArrayData(UPCGTexture2DArrayData* NewT
 	NewTextureData->ResourceType = ResourceType;
 	NewTextureData->Resolution = Resolution;
 	NewTextureData->ArraySize = ArraySize;
+	NewTextureData->NumMips = NumMips;
 	NewTextureData->Format = Format;
 	NewTextureData->Filter = Filter;
 	NewTextureData->Bounds = Bounds;
diff --git a/Engine/Plugins/PCG/Source/PCG/Private/Data/PCGTexture2DBaseData.cpp b/Engine/Plugins/PCG/Source/PCG/Private/Data/PCGTexture2DBaseData.cpp
--- a/Engine/Plugins/PCG/Source/PCG/Private/Data/PCGTexture2DBaseData.cpp
+++ b/Engine/Plugins/PCG/Source/PCG/Private/Data/PCGTexture2DBaseData.cpp
@@ -29,6 +29,7 @@ void UPCGTexture2DBaseData::CopyTexture2DBaseData(UPCGTexture2DBaseData* NewText
 	CopyBaseSurfaceData(NewTextureData);
 
 	NewTextureData->Format = Format;
+	NewTextureData->NumMips = NumMips;
 	NewTextureData->Filter = Filter;
 	NewTextureData->TextureHandle = TextureHandle;
 	NewTextureData->Bounds = Bounds;
diff --git a/Engine/Plugins/PCG/Source/PCG/Private/Data/PCGTextureData.cpp b/Engine/Plugins/PCG/Source/PCG/Private/Data/PCGTextureData.cpp
--- a/Engine/Plugins/PCG/Source/PCG/Private/Data/PCGTextureData.cpp
+++ b/Engine/Plugins/PCG/Source/PCG/Private/Data/PCGTextureData.cpp
@@ -478,7 +478,9 @@ void UPCGTextureData::InitializeInternal(UTexture* InTexture, uint32 InTextureIn
 	Width = 0;
 	Height = 0;
 	bSkipReadbackToCPU = bInSkipReadbackToCPU;
-	Format = GetTextureRHI() ? GetTextureRHI()->GetDesc().Format : PF_Unknown;
+	const FTextureRHIRef InitRHI = GetTextureRHI();
+	Format = InitRHI ? InitRHI->GetDesc().Format : PF_Unknown;
+	NumMips = InitRHI ? InitRHI->GetDesc().NumMips : 0;
 
 	Bounds = FBox(EForceInit::ForceInit);
 	Bounds += FVector(-1.0f, -1.0f, 0.0f);
diff --git a/Engine/Plugins/PCG/Source/PCG/Private/Elements/PCGGetTextureInfo.cpp b/Engine/Plugins/PCG/Source/PCG/Private/Elements/PCGGetTextureInfo.cpp
new file mode 100644
--- /dev/null
+++ b/Engine/Plugins/PCG/Source/PCG/Private/Elements/PCGGetTextureInfo.cpp
@@ -0,0 +1,125 @@
+// Copyright Epic Games, Inc. All Rights Reserved.
+
+#include "Elements/PCGGetTextureInfo.h"
+
+#include "PCGContext.h"
+#include "PCGParamData.h"
+#include "PCGPin.h"
+#include "Compute/PCGComputeCommon.h"
+#include "Data/PCGTexture2DBaseData.h"
+
+#include "PixelFormat.h"
+
+#include UE_INLINE_GENERATED_CPP_BY_NAME(PCGGetTextureInfo)
+
+#define LOCTEXT_NAMESPACE "PCGGetTextureInfoElement"
+
+namespace PCGGetTextureInfoConstants
+{
+	const FName InputTextureLabel = TEXT("Texture");
+}
+
+#if WITH_EDITOR
+FName UPCGGetTextureInfoSettings::GetDefaultNodeName() const
+{
+	return TEXT("GetTextureInfo");
+}
+
+FText UPCGGetTextureInfoSettings::GetDefaultNodeTitle() const
+{
+	return LOCTEXT("NodeTitle", "Get Texture Info");
+}
+
+FText UPCGGetTextureInfoSettings::GetNodeTooltipText() const
+{
+	return LOCTEXT("NodeTooltip", "Returns texture metadata as attributes.");
+}
+#endif
+
+TArray<FPCGPinProperties> UPCGGetTextureInfoSettings::InputPinProperties() const
+{
+	TArray<FPCGPinProperties> PinProperties;
+
+	FPCGPinProperties& TexturePin = PinProperties.Emplace_GetRef(
+		PCGGetTextureInfoConstants::InputTextureLabel,
+		FPCGDataTypeInfoTexture2DBase::AsId(),
+		/*bAllowMultipleConnections=*/true,
+		/*bAllowMultipleData=*/true);
+	TexturePin.SetRequiredPin();
+
+	return PinProperties;
+}
+
+TArray<FPCGPinProperties> UPCGGetTextureInfoSettings::OutputPinProperties() const
+{
+	TArray<FPCGPinProperties> PinProperties;
+	PinProperties.Emplace(
+		PCGPinConstants::DefaultOutputLabel,
+		EPCGDataType::Param,
+		/*bAllowMultipleConnections=*/true,
+		/*bAllowMultipleData=*/true,
+		LOCTEXT("OutParamTooltip", "Attribute set containing the texture metadata"));
+
+	return PinProperties;
+}
+
+FPCGElementPtr UPCGGetTextureInfoSettings::CreateElement() const
+{
+	return MakeShared<FPCGGetTextureInfoElement>();
+}
+
+bool FPCGGetTextureInfoElement::ExecuteInternal(FPCGContext* Context) const
+{
+	TRACE_CPUPROFILER_EVENT_SCOPE(FPCGGetTextureInfoElement::Execute);
+
+	check(Context);
+
+	const UPCGGetTextureInfoSettings* Settings = Context->GetInputSettings<UPCGGetTextureInfoSettings>();
+	check(Settings);
+
+	const TArray<FPCGTaggedData> TextureInputs = Context->InputData.GetInputsByPin(PCGGetTextureInfoConstants::InputTextureLabel);
+
+	for (const FPCGTaggedData& Input : TextureInputs)
+	{
+		const UPCGTexture2DBaseData* TextureData = Cast<UPCGTexture2DBaseData>(Input.Data);
+		if (!TextureData)
+		{
+			PCGE_LOG(Error, GraphAndLog, LOCTEXT("UnsupportedTextureType", "Input data is not a supported texture type."));
+			continue;
+		}
+
+		const FIntPoint Resolution = TextureData->GetResolution();
+		const EPixelFormat PixelFormat = TextureData->GetFormat();
+		const EPCGRenderTargetFormat ComputeFormat = PCGComputeHelpers::GetPCGRenderTargetFormatFromPixelFormat(PixelFormat);
+
+		UPCGParamData* OutputParamData = FPCGContext::NewObject_AnyThread<UPCGParamData>(Context);
+		UPCGMetadata* Metadata = OutputParamData->Metadata;
+		check(Metadata);
+
+		auto CreateAttribute = [Metadata, Context](bool bInEnabled, FName InAttributeName, auto InValue)
+		{
+			if (bInEnabled && !Metadata->CreateAttribute(InAttributeName, InValue, /*bAllowInterpolation=*/false, /*bOverrideParent=*/false))
+			{
+				PCGLog::LogErrorOnGraph(FText::Format(LOCTEXT("FailedToCreateAttribute", "Failed to create attribute '{0}'."), FText::FromName(InAttributeName)), Context);
+			}
+		};
+
+		CreateAttribute(Settings->bOutputWidth, Settings->WidthAttributeName, Resolution.X);
+		CreateAttribute(Settings->bOutputHeight, Settings->HeightAttributeName, Resolution.Y);
+		CreateAttribute(Settings->bOutputNumMips, Settings->NumMipsAttributeName, static_cast<int32>(TextureData->GetNumMips()));
+		CreateAttribute(Settings->bOutputArraySize, Settings->ArraySizeAttributeName, static_cast<int32>(TextureData->GetArraySize()));
+		CreateAttribute(Settings->bOutputFormat, Settings->FormatAttributeName, FName(GetPixelFormatString(PixelFormat)));
+		CreateAttribute(Settings->bOutputFormatIndex, Settings->FormatIndexAttributeName, static_cast<int32>(PixelFormat));
+		CreateAttribute(Settings->bOutputComputeFormat, Settings->ComputeFormatAttributeName, FName(StaticEnum<EPCGRenderTargetFormat>()->GetNameStringByValue(static_cast<int64>(ComputeFormat))));
+		CreateAttribute(Settings->bOutputComputeFormatIndex, Settings->ComputeFormatIndexAttributeName, static_cast<int32>(ComputeFormat));
+
+		Metadata->AddEntry();
+
+		FPCGTaggedData& Output = Context->OutputData.TaggedData.Emplace_GetRef();
+		Output.Data = OutputParamData;
+	}
+
+	return true;
+}
+
+#undef LOCTEXT_NAMESPACE
diff --git a/Engine/Plugins/PCG/Source/PCG/Public/Data/PCGTexture2DBaseData.h b/Engine/Plugins/PCG/Source/PCG/Public/Data/PCGTexture2DBaseData.h
--- a/Engine/Plugins/PCG/Source/PCG/Public/Data/PCGTexture2DBaseData.h
+++ b/Engine/Plugins/PCG/Source/PCG/Public/Data/PCGTexture2DBaseData.h
@@ -44,6 +44,7 @@ public:
 	virtual EPixelFormat GetFormat() const { return Format; }
 	virtual EPCGTextureFilter GetFilter() const { return Filter; }
 	virtual uint16 GetArraySize() const { return 1; }
+	virtual uint8 GetNumMips() const { return NumMips; }
 	virtual FIntPoint GetResolution() const PURE_VIRTUAL(UPCGTexture2DBaseData::GetResolution, return FIntPoint::ZeroValue;);
 	virtual TRefCountPtr<IPooledRenderTarget> GetRefCountedTexture() const { return TextureHandle; }
 
@@ -57,6 +58,9 @@ protected:
 	UPROPERTY(BlueprintReadOnly, VisibleAnywhere, Category = SpatialData)
 	TEnumAsByte<EPixelFormat> Format = EPixelFormat::PF_Unknown;
 
+	UPROPERTY()
+	uint8 NumMips = 0;
+
 public:
 	UPROPERTY(BlueprintReadOnly, VisibleAnywhere, Category = SpatialData)
 	EPCGTextureFilter Filter = EPCGTextureFilter::Bilinear;
diff --git a/Engine/Plugins/PCG/Source/PCG/Public/Elements/PCGGetTextureInfo.h b/Engine/Plugins/PCG/Source/PCG/Public/Elements/PCGGetTextureInfo.h
new file mode 100644
--- /dev/null
+++ b/Engine/Plugins/PCG/Source/PCG/Public/Elements/PCGGetTextureInfo.h
@@ -0,0 +1,92 @@
+// Copyright Epic Games, Inc. All Rights Reserved.
+
+#pragma once
+
+#include "PCGSettings.h"
+
+#include "PCGGetTextureInfo.generated.h"
+
+/** Returns texture metadata as attributes. */
+UCLASS(MinimalAPI, BlueprintType, ClassGroup = (Procedural))
+class UPCGGetTextureInfoSettings : public UPCGSettings
+{
+	GENERATED_BODY()
+
+public:
+	//~Begin UPCGSettings interface
+#if WITH_EDITOR
+	virtual FName GetDefaultNodeName() const override;
+	virtual FText GetDefaultNodeTitle() const override;
+	virtual FText GetNodeTooltipText() const override;
+	virtual EPCGSettingsType GetType() const override { return EPCGSettingsType::Generic; }
+#endif
+
+protected:
+	virtual TArray<FPCGPinProperties> InputPinProperties() const override;
+	virtual TArray<FPCGPinProperties> OutputPinProperties() const override;
+	virtual FPCGElementPtr CreateElement() const override;
+	//~End UPCGSettings interface
+
+public:
+	UPROPERTY(BlueprintReadWrite, EditAnywhere, Category = Settings, meta = (PCG_Overridable, InlineEditConditionToggle))
+	bool bOutputWidth = true;
+
+	/** Output attribute name for the texture width in pixels. */
+	UPROPERTY(BlueprintReadWrite, EditAnywhere, Category = Settings, meta = (PCG_Overridable, EditCondition = "bOutputWidth"))
+	FName WidthAttributeName = TEXT("Width");
+
+	UPROPERTY(BlueprintReadWrite, EditAnywhere, Category = Settings, meta = (PCG_Overridable, InlineEditConditionToggle))
+	bool bOutputHeight = true;
+
+	/** Output attribute name for the texture height in pixels. */
+	UPROPERTY(BlueprintReadWrite, EditAnywhere, Category = Settings, meta = (PCG_Overridable, EditCondition = "bOutputHeight"))
+	FName HeightAttributeName = TEXT("Height");
+
+	UPROPERTY(BlueprintReadWrite, EditAnywhere, Category = Settings, meta = (PCG_Overridable, InlineEditConditionToggle))
+	bool bOutputNumMips = true;
+
+	/** Output attribute name for the number of mip levels. */
+	UPROPERTY(BlueprintReadWrite, EditAnywhere, Category = Settings, meta = (PCG_Overridable, EditCondition = "bOutputNumMips"))
+	FName NumMipsAttributeName = TEXT("NumMips");
+
+	UPROPERTY(BlueprintReadWrite, EditAnywhere, Category = Settings, meta = (PCG_Overridable, InlineEditConditionToggle))
+	bool bOutputArraySize = true;
+
+	/** Output attribute name for the texture array size (1 for non-array textures). */
+	UPROPERTY(BlueprintReadWrite, EditAnywhere, Category = Settings, meta = (PCG_Overridable, EditCondition = "bOutputArraySize"))
+	FName ArraySizeAttributeName = TEXT("ArraySize");
+
+	UPROPERTY(BlueprintReadWrite, EditAnywhere, Category = Settings, meta = (PCG_Overridable, InlineEditConditionToggle))
+	bool bOutputFormat = true;
+
+	/** Output attribute name for the pixel format as a human-readable string (e.g. "PF_B8G8R8A8"). */
+	UPROPERTY(BlueprintReadWrite, EditAnywhere, Category = Settings, meta = (PCG_Overridable, EditCondition = "bOutputFormat"))
+	FName FormatAttributeName = TEXT("Format");
+
+	UPROPERTY(BlueprintReadWrite, EditAnywhere, Category = Settings, meta = (PCG_Overridable, InlineEditConditionToggle))
+	bool bOutputFormatIndex = true;
+
+	/** Output attribute name for the pixel format as an integer (EPixelFormat enum value). */
+	UPROPERTY(BlueprintReadWrite, EditAnywhere, Category = Settings, meta = (PCG_Overridable, EditCondition = "bOutputFormatIndex"))
+	FName FormatIndexAttributeName = TEXT("FormatIndex");
+
+	UPROPERTY(BlueprintReadWrite, EditAnywhere, Category = Settings, meta = (PCG_Overridable, InlineEditConditionToggle))
+	bool bOutputComputeFormat = true;
+
+	/** Output attribute name for the PCG compute format as a human-readable string (e.g. "RGBA16f"). This is the pixel format narrowed to the subset of formats that support UAV access, used by PCG compute nodes. */
+	UPROPERTY(BlueprintReadWrite, EditAnywhere, Category = Settings, meta = (PCG_Overridable, EditCondition = "bOutputComputeFormat"))
+	FName ComputeFormatAttributeName = TEXT("ComputeFormat");
+
+	UPROPERTY(BlueprintReadWrite, EditAnywhere, Category = Settings, meta = (PCG_Overridable, InlineEditConditionToggle))
+	bool bOutputComputeFormatIndex = true;
+
+	/** Output attribute name for the PCG compute format as an integer (EPCGRenderTargetFormat enum value). This is the pixel format narrowed to the subset of formats that support UAV access. */
+	UPROPERTY(BlueprintReadWrite, EditAnywhere, Category = Settings, meta = (PCG_Overridable, EditCondition = "bOutputComputeFormatIndex"))
+	FName ComputeFormatIndexAttributeName = TEXT("ComputeFormatIndex");
+};
+
+class FPCGGetTextureInfoElement : public IPCGElement
+{
+protected:
+	virtual bool ExecuteInternal(FPCGContext* Context) const override;
+};
