[PCG] Fix PCGRayTrace failing when bindless is enabled for ray tracing

#rnx
#Synced-CL 52605476

diff --git a/Engine/Plugins/PCG/Source/PCGCompute/Private/PCGRayTrace.cpp b/Engine/Plugins/PCG/Source/PCGCompute/Private/PCGRayTrace.cpp
--- a/Engine/Plugins/PCG/Source/PCGCompute/Private/PCGRayTrace.cpp
+++ b/Engine/Plugins/PCG/Source/PCGCompute/Private/PCGRayTrace.cpp
@@ -50,7 +50,7 @@ class FPCGRayTraceCS : public FGlobalShader
 	static bool ShouldCompilePermutation(const FGlobalShaderPermutationParameters& Parameters)
 	{
 		return IsRayTracingEnabledForProject(Parameters.Platform)
-			&& IsBindlessEnabledForAnyGraphics(FShaderPlatformConfig::GetBindlessConfiguration(Parameters.Platform));
+			&& IsBindlessEnabledForRayTracing(FShaderPlatformConfig::GetBindlessConfiguration(Parameters.Platform));
 	}
 
 	static void ModifyCompilationEnvironment(const FGlobalShaderPermutationParameters& Parameters, FShaderCompilerEnvironment& OutEnvironment)
@@ -86,9 +86,9 @@ void PCGRayTrace::RenderPCGRayTraceInline(FRDGBuilder& GraphBuilder, const FPCGR
 		return;
 	}
 
-	if (!IsBindlessEnabledForAnyGraphics(FShaderPlatformConfig::GetBindlessConfiguration(GMaxRHIShaderPlatform)))
+	if (!IsBindlessEnabledForRayTracing(FShaderPlatformConfig::GetBindlessConfiguration(GMaxRHIShaderPlatform)))
 	{
-		UE_LOGF(LogPCGCompute, Error, "Failed to submit inline raytrace shader. Bindless rendering is not enabled for platform '%ls'.", *LexToString(GMaxRHIShaderPlatform));
+		UE_LOGF(LogPCGCompute, Error, "Failed to submit inline raytrace shader. Bindless rendering is not enabled for platform '%ls'. Check your project settings.", *LexToString(GMaxRHIShaderPlatform));
 		return;
 	}
 
