[PCG] Expose Raw Buffer data type in pin dropdowns, Custom HLSL declarations, and inspection

#Synced-CL 52605476

diff --git a/Engine/Plugins/PCG/Shaders/Private/BuiltInKernels/PCGCountUniqueAttributeValues.usf b/Engine/Plugins/PCG/Shaders/Private/BuiltInKernels/PCGCountUniqueAttributeValues.usf
--- a/Engine/Plugins/PCG/Shaders/Private/BuiltInKernels/PCGCountUniqueAttributeValues.usf
+++ b/Engine/Plugins/PCG/Shaders/Private/BuiltInKernels/PCGCountUniqueAttributeValues.usf
@@ -12,10 +12,10 @@ void InitializeValueKey(uint InDataIndex, uint InValueKey, uint InBucketIndex)
 {
 #if PCG_OUTPUT_TO_RAW_BUFFER
 	const uint FirstElementIndex = (InValueKey * CountUniqueValues_GetNumBucketsTotal() + InBucketIndex) * 2u;
-	if (FirstElementIndex < Out_GetNumElements(0))
+	if (FirstElementIndex < Out_GetNumElements())
 	{
 		// We currently have an output element per possible string key value. This may be compacted in the future.
-		Out_Store(InDataIndex, FirstElementIndex, InValueKey);
+		Out_Store(FirstElementIndex, InValueKey);
 	}
 #else
 	// Write value of attribute. We currently have an output element per possible string key value.
@@ -98,7 +98,7 @@ void PCGCountUniqueAttributeValuesCS(uint3 GroupId : SV_GroupID, uint GroupIndex
 
 #if PCG_OUTPUT_TO_RAW_BUFFER
 		const uint ElementIndex = (AttributeValue * NumBucketsTotal + BucketIndex) * 2u + 1;
-		Out_AtomicAdd(/*DataIndex=*/0, ElementIndex, 1u);
+		Out_AtomicAdd(ElementIndex, 1u);
 #else
 		// Always one bucket per attribute value.
 		const uint Out_ElementIndex = AttributeValue * NumBucketsTotal + BucketIndex;
diff --git a/Engine/Plugins/PCG/Shaders/Private/BuiltInKernels/PCGPostRayTrace.usf b/Engine/Plugins/PCG/Shaders/Private/BuiltInKernels/PCGPostRayTrace.usf
--- a/Engine/Plugins/PCG/Shaders/Private/BuiltInKernels/PCGPostRayTrace.usf
+++ b/Engine/Plugins/PCG/Shaders/Private/BuiltInKernels/PCGPostRayTrace.usf
@@ -34,19 +34,18 @@ void PCGPostRayTraceCS(uint3 GroupId : SV_GroupID, uint GroupIndex : SV_GroupInd
 		return;
 	}
 	
-	const uint RayTraceResults_DataIndex = 0;
 	const uint RayTraceResults_FirstElementIndex = ThreadIndex * PCG_RAY_TRACE_PACKED_BUFFER_STRIDE_UINTS;
 	const uint PackedPositionIndex = RayTraceResults_FirstElementIndex;
 	const uint PackedTexCoordsIndex = RayTraceResults_FirstElementIndex + 4;
 	const uint PackedNormalsIndex = RayTraceResults_FirstElementIndex + 8;
 
-	if (PackedPositionIndex >= RayTraceResults_GetNumElements(RayTraceResults_DataIndex))
+	if (PackedPositionIndex >= RayTraceResults_GetNumElements())
 	{
 		// Should not happen - invalid results data.
 		return;
 	}
 
-	const uint4 PackedPositionAndValid = RayTraceResults_Load4(RayTraceResults_DataIndex, PackedPositionIndex);
+	const uint4 PackedPositionAndValid = RayTraceResults_Load4(PackedPositionIndex);
 	const float3 HitPosition = asfloat(PackedPositionAndValid.xyz);
 	bool bMissedOrOutOfBounds = false;
 
@@ -77,7 +76,7 @@ void PCGPostRayTraceCS(uint3 GroupId : SV_GroupID, uint GroupIndex : SV_GroupInd
 	{
 		if (!bMissedOrOutOfBounds)
 		{
-			const uint4 PackedTexCoords = RayTraceResults_Load4(RayTraceResults_DataIndex, PackedTexCoordsIndex);
+			const uint4 PackedTexCoords = RayTraceResults_Load4(PackedTexCoordsIndex);
 			const float2 TexCoords = asfloat(PackedTexCoords.xy);
 			Out_SetFloat2(Out_DataIndex, Out_ElementIndex, TexCoordsAttributeId, TexCoords);
 		}
@@ -93,7 +92,7 @@ void PCGPostRayTraceCS(uint3 GroupId : SV_GroupID, uint GroupIndex : SV_GroupInd
 	{
 		if (!bMissedOrOutOfBounds)
 		{
-			const uint4 PackedNormals = RayTraceResults_Load4(RayTraceResults_DataIndex, PackedNormalsIndex);
+			const uint4 PackedNormals = RayTraceResults_Load4(PackedNormalsIndex);
 			const float3 Normals = asfloat(PackedNormals.xyz);
 			Out_SetFloat3(Out_DataIndex, Out_ElementIndex, NormalsAttributeId, Normals);
 		}
diff --git a/Engine/Plugins/PCG/Shaders/Private/BuiltInKernels/PCGPreRayTrace.usf b/Engine/Plugins/PCG/Shaders/Private/BuiltInKernels/PCGPreRayTrace.usf
--- a/Engine/Plugins/PCG/Shaders/Private/BuiltInKernels/PCGPreRayTrace.usf
+++ b/Engine/Plugins/PCG/Shaders/Private/BuiltInKernels/PCGPreRayTrace.usf
@@ -30,7 +30,7 @@ void PCGPreRayTraceCS(uint3 GroupId : SV_GroupID, uint GroupIndex : SV_GroupInde
 	// If the point is culled, the ray should also be culled.
 	if (In_IsPointRemoved(In_DataIndex, In_ElementIndex))
 	{
-		Out_Store4(Out_DataIndex, PackedPositionIndex, (uint4)PCG_RAY_CULLED);
+		Out_Store4(PackedPositionIndex, (uint4)PCG_RAY_CULLED);
 		return;
 	}
 
@@ -105,7 +105,7 @@ void PCGPreRayTraceCS(uint3 GroupId : SV_GroupID, uint GroupIndex : SV_GroupInde
 	// Normalize ray direction.
 	RayDirection = dot(RayDirection, RayDirection) != 0.0f ? normalize(RayDirection) : (float3)0.0f;
 
-	Out_Store4(Out_DataIndex, PackedPositionIndex, uint4(asuint(RayOrigin), 0u));
-	Out_Store4(Out_DataIndex, PackedDirectionIndex, uint4(asuint(RayDirection), asuint(RayLength)));
-	Out_Store4(Out_DataIndex, PackedNormalsIndex, (uint4)0u); // Initialize buffer for normal values that will be written later.
+	Out_Store4(PackedPositionIndex, uint4(asuint(RayOrigin), 0u));
+	Out_Store4(PackedDirectionIndex, uint4(asuint(RayDirection), asuint(RayLength)));
+	Out_Store4(PackedNormalsIndex, (uint4)0u); // Initialize buffer for normal values that will be written later.
 }
diff --git a/Engine/Plugins/PCG/Shaders/Private/PCGRawBufferDataInterface.ush b/Engine/Plugins/PCG/Shaders/Private/PCGRawBufferDataInterface.ush
--- a/Engine/Plugins/PCG/Shaders/Private/PCGRawBufferDataInterface.ush
+++ b/Engine/Plugins/PCG/Shaders/Private/PCGRawBufferDataInterface.ush
@@ -14,44 +14,37 @@ uint {DataInterfaceName}_SizeBytes;
 
 // #################### META INFO ##########################
 
-uint GetNumData_{DataInterfaceName}()
+uint GetNumElements_{DataInterfaceName}()
 {
-	// Multi-data not currently supported. Could be done with multiple buffers (each individual buffer should be a pure raw buffer).
-	return 1u;
-}
-
-uint GetNumElements_{DataInterfaceName}(uint InDataIndex)
-{
-	// InDataIndex unused for now as multi-data not currently supported.
 	return {DataInterfaceName}_SizeBytes >> 2;
 }
 
 // #################### ACCESSORS ##########################
 
-uint Load_{DataInterfaceName}(uint InDataIndex, uint InElementIndex)
+uint Load_{DataInterfaceName}(uint InElementIndex)
 {
 	return {DataInterfaceName}_Data.Load(InElementIndex << 2);
 }
 
-uint4 Load4_{DataInterfaceName}(uint InDataIndex, uint InFirstElementIndex)
+uint4 Load4_{DataInterfaceName}(uint InFirstElementIndex)
 {
 	return {DataInterfaceName}_Data.Load4(InFirstElementIndex << 2);
 }
 
-void Store_{DataInterfaceName}(uint InDataIndex, uint InElementIndex, uint InValue)
+void Store_{DataInterfaceName}(uint InElementIndex, uint InValue)
 {
 	{DataInterfaceName}_Data.Store(InElementIndex << 2, InValue);
 }
 
-void Store4_{DataInterfaceName}(uint InDataIndex, uint InFirstElementIndex, uint4 InValue)
+void Store4_{DataInterfaceName}(uint InFirstElementIndex, uint4 InValue)
 {
 	{DataInterfaceName}_Data.Store4(InFirstElementIndex << 2, InValue);
 }
 
 // #################### ATOMICS ##########################
 
-// Returns value of attribute before incrementing.
-uint AtomicAdd_{DataInterfaceName}(uint InDataIndex, uint InElementIndex, uint InValueToAdd)
+// Returns value before incrementing.
+uint AtomicAdd_{DataInterfaceName}(uint InElementIndex, uint InValueToAdd)
 {
 	uint OriginalValue;
 	{DataInterfaceName}_Data.InterlockedAdd(InElementIndex << 2, (uint)InValueToAdd, OriginalValue);
diff --git a/Engine/Plugins/PCG/Source/PCG/Private/Compute/Data/PCGRawBufferData.h b/Engine/Plugins/PCG/Source/PCG/Private/Compute/Data/PCGRawBufferData.h
deleted file mode 100644
diff --git a/Engine/Plugins/PCG/Source/PCG/Private/Compute/DataInterfaces/PCGRawBufferDataInterface.cpp b/Engine/Plugins/PCG/Source/PCG/Private/Compute/DataInterfaces/PCGRawBufferDataInterface.cpp
--- a/Engine/Plugins/PCG/Source/PCG/Private/Compute/DataInterfaces/PCGRawBufferDataInterface.cpp
+++ b/Engine/Plugins/PCG/Source/PCG/Private/Compute/DataInterfaces/PCGRawBufferDataInterface.cpp
@@ -1,6 +1,6 @@
 // Copyright Epic Games, Inc. All Rights Reserved.
 
-#include "PCGRawBufferDataInterface.h"
+#include "Compute/DataInterfaces/PCGRawBufferDataInterface.h"
 
 #include "PCGCommon.h"
 #include "PCGModule.h"
@@ -46,42 +46,32 @@ namespace PCGRawBufferDataInterface
 
 void UPCGRawBufferDataInterface::GetSupportedOutputs(TArray<FShaderFunctionDefinition>& OutFunctions) const
 {
-	OutFunctions.AddDefaulted_GetRef()
-		.SetName(TEXT("GetNumData"))
-		.AddReturnType(EShaderFundamentalType::Uint); // Num data
-
 	OutFunctions.AddDefaulted_GetRef()
 		.SetName(TEXT("GetNumElements"))
-		.AddParam(EShaderFundamentalType::Uint) // InDataIndex
-		.AddReturnType(EShaderFundamentalType::Uint); // Size in bytes
+		.AddReturnType(EShaderFundamentalType::Uint);
 
 	OutFunctions.AddDefaulted_GetRef()
 		.SetName(TEXT("Load"))
-		.AddParam(EShaderFundamentalType::Uint) // InDataIndex
 		.AddParam(EShaderFundamentalType::Uint) // InElementIndex
 		.AddReturnType(EShaderFundamentalType::Uint);
 
 	OutFunctions.AddDefaulted_GetRef()
 		.SetName(TEXT("Load4"))
-		.AddParam(EShaderFundamentalType::Uint) // InDataIndex
 		.AddParam(EShaderFundamentalType::Uint) // InFirstElementIndex
 		.AddReturnType(EShaderFundamentalType::Uint, 4);
 
 	OutFunctions.AddDefaulted_GetRef()
 		.SetName(TEXT("Store"))
-		.AddParam(EShaderFundamentalType::Uint) // InDataIndex
 		.AddParam(EShaderFundamentalType::Uint) // InElementIndex
 		.AddParam(EShaderFundamentalType::Uint); // InValue
 
 	OutFunctions.AddDefaulted_GetRef()
 		.SetName(TEXT("Store4"))
-		.AddParam(EShaderFundamentalType::Uint) // InDataIndex
 		.AddParam(EShaderFundamentalType::Uint) // InFirstElementIndex
 		.AddParam(EShaderFundamentalType::Uint, 4); // InValue
 
 	OutFunctions.AddDefaulted_GetRef()
 		.SetName(TEXT("AtomicAdd"))
-		.AddParam(EShaderFundamentalType::Uint) // InDataIndex
 		.AddParam(EShaderFundamentalType::Uint) // InElementIndex
 		.AddParam(EShaderFundamentalType::Uint) // InValueToAdd
 		.AddReturnType(EShaderFundamentalType::Uint); // Value before add
diff --git a/Engine/Plugins/PCG/Source/PCG/Private/Compute/DataInterfaces/PCGRawBufferDataInterface.h b/Engine/Plugins/PCG/Source/PCG/Private/Compute/DataInterfaces/PCGRawBufferDataInterface.h
deleted file mode 100644
diff --git a/Engine/Plugins/PCG/Source/PCG/Private/Compute/Elements/PCGCustomHLSL.cpp b/Engine/Plugins/PCG/Source/PCG/Private/Compute/Elements/PCGCustomHLSL.cpp
--- a/Engine/Plugins/PCG/Source/PCG/Private/Compute/Elements/PCGCustomHLSL.cpp
+++ b/Engine/Plugins/PCG/Source/PCG/Private/Compute/Elements/PCGCustomHLSL.cpp
@@ -11,6 +11,7 @@
 #include "Compute/PCGComputeSource.h"
 #include "Compute/PCGKernelHelpers.h"
 #include "Compute/Elements/PCGCustomHLSLKernel.h"
+#include "Compute/Data/PCGRawBufferData.h"
 #include "Data/PCGTexture2DArrayData.h"
 #include "Data/PCGLandscapeData.h"
 #include "Data/PCGPointData.h"
@@ -604,6 +605,11 @@ void UPCGCustomHLSLSettings::UpdateInputDeclarations()
 			VirtualTextureDataPins.Add(Pin.Label.ToString());
 		}
 
+		if (!!(Pin.AllowedTypes & FPCGDataTypeInfoRawBuffer::AsId()))
+		{
+			RawBufferDataPins.Add(Pin.Label.ToString());
+		}
+
 		if (!!(Pin.AllowedTypes & EPCGDataType::StaticMeshResource))
 		{
 			StaticMeshDataPins.Add(Pin.Label.ToString());
@@ -789,8 +795,9 @@ void UPCGCustomHLSLSettings::UpdateInputDeclarations()
 		}
 
 		InputDeclarations += FString::Format(TEXT(
-			"uint {0}_ReadNumValues();\n"
-			"uint {0}_ReadValue(uint Index);\n"),
+			"uint {0}_GetNumElements();\n"
+			"uint {0}_Load(uint ElementIndex);\n"
+			"uint4 {0}_Load4(uint FirstElementIndex);\n"),
 			{ bMultiPin ? PCGHLSLElement::PinDeclTemplateStr : RawBufferDataPins[0] });
 
 		InputDeclarations += TEXT("\n");
@@ -862,6 +869,11 @@ void UPCGCustomHLSLSettings::UpdateOutputDeclarations()
 		{
 			TextureArrayDataPins.Add(Pin.Label.ToString());
 		}
+
+		if (!!(Pin.AllowedTypes & FPCGDataTypeInfoRawBuffer::AsId()))
+		{
+			RawBufferDataPins.Add(Pin.Label.ToString());
+		}
 	}
 
 	if (!DataCollectionDataPins.IsEmpty())
@@ -987,7 +999,12 @@ void UPCGCustomHLSLSettings::UpdateOutputDeclarations()
 		}
 
 		OutputDeclarations += FString::Format(TEXT(
-			"uint {0}_WriteValue(uint Index, uint Value);\n"),
+			"uint {0}_GetNumElements();\n"
+			"uint {0}_Load(uint ElementIndex);\n"
+			"uint4 {0}_Load4(uint FirstElementIndex);\n"
+			"void {0}_Store(uint ElementIndex, uint Value);\n"
+			"void {0}_Store4(uint FirstElementIndex, uint4 Value);\n"
+			"uint {0}_AtomicAdd(uint ElementIndex, uint ValueToAdd); // Returns value before add.\n"),
 			{ bMultiPin ? PCGHLSLElement::PinDeclTemplateStr : RawBufferDataPins[0] });
 
 		OutputDeclarations += TEXT("\n");
@@ -1169,7 +1186,9 @@ void UPCGCustomHLSLSettings::UpdatePinSettings()
 		}
 
 		// Enforce single data on pin types that don't support a dynamic number of data.
-		if (Properties.AllowedTypes.IsChildOf(FPCGDataTypeInfoLandscape::AsId()) || Properties.AllowedTypes.IsChildOf(FPCGDataTypeInfoTexture2DArray::AsId()))
+		if (Properties.AllowedTypes.IsChildOf(FPCGDataTypeInfoLandscape::AsId())
+			|| Properties.AllowedTypes.IsChildOf(FPCGDataTypeInfoTexture2DArray::AsId())
+			|| Properties.AllowedTypes.IsChildOf(FPCGDataTypeInfoRawBuffer::AsId()))
 		{
 			Properties.bAllowMultipleData = false;
 			Properties.bAllowEditMultipleData = false;
@@ -1233,6 +1252,13 @@ void UPCGCustomHLSLSettings::UpdatePinSettings()
 			// Also show normal texture settings for texture arrays.
 			Properties.PropertiesGPU.bShowTexturePinSettings = true;
 		}
+		else if (Properties.AllowedTypes.IsChildOf(FPCGDataTypeInfoRawBuffer::AsId()))
+		{
+			Properties.bAllowMultipleData = false;
+			Properties.bAllowEditMultipleData = false;
+			Properties.PropertiesGPU.bShowTexturePinSettings = false;
+			Properties.PropertiesGPU.bShowTextureArrayPinSettings = false;
+		}
 		else
 		{
 			Properties.PropertiesGPU.bShowTexturePinSettings = false;
diff --git a/Engine/Plugins/PCG/Source/PCG/Private/Compute/Elements/PCGCustomHLSLKernel.cpp b/Engine/Plugins/PCG/Source/PCG/Private/Compute/Elements/PCGCustomHLSLKernel.cpp
--- a/Engine/Plugins/PCG/Source/PCG/Private/Compute/Elements/PCGCustomHLSLKernel.cpp
+++ b/Engine/Plugins/PCG/Source/PCG/Private/Compute/Elements/PCGCustomHLSLKernel.cpp
@@ -8,6 +8,7 @@
 #include "PCGPointPropertiesTraits.h"
 #include "Compute/PCGComputeGraph.h"
 #include "Compute/PCGDataBinding.h"
+#include "Compute/Data/PCGRawBufferData.h"
 #include "Compute/Elements/PCGCustomHLSL.h"
 #include "Compute/DataInterfaces/Elements/PCGCustomHLSLDataInterface.h"
 #include "Data/PCGTexture2DArrayData.h"
@@ -36,6 +37,7 @@ namespace PCGCustomHLSLKernel
 	constexpr TCHAR CopyElementFunctionKeyword[] = { TEXT("CopyElementFrom") };
 	constexpr TCHAR InitializeFunctionKeyword[] = { TEXT("InitializePoint") };
 	constexpr TCHAR StoreFunctionKeyword[] = { TEXT("Store") };
+	constexpr TCHAR Store4FunctionKeyword[] = { TEXT("Store4") };
 	constexpr TCHAR SetPositionKeyword[] = { TEXT("SetPosition") };
 	constexpr TCHAR SetRotationKeyword[] = { TEXT("SetRotation") };
 	constexpr TCHAR SetScaleKeyword[] = { TEXT("SetScale") };
@@ -936,6 +938,11 @@ void UPCGCustomHLSLKernel::ParseShaderSource()
 		{
 			InitializeKeywords.Add(PinStr + TEXT("_") + PCGCustomHLSLKernel::StoreFunctionKeyword);
 		}
+		else if (PinProps.AllowedTypes.IsChildOf(FPCGDataTypeInfoRawBuffer::AsId()))
+		{
+			InitializeKeywords.Add(PinStr + TEXT("_") + PCGCustomHLSLKernel::StoreFunctionKeyword);
+			InitializeKeywords.Add(PinStr + TEXT("_") + PCGCustomHLSLKernel::Store4FunctionKeyword);
+		}
 
 		if (!!(PinProps.AllowedTypes & EPCGDataType::Point))
 		{
diff --git a/Engine/Plugins/PCG/Source/PCG/Public/Compute/Data/PCGRawBufferData.h b/Engine/Plugins/PCG/Source/PCG/Public/Compute/Data/PCGRawBufferData.h
new file mode 100644
--- /dev/null
+++ b/Engine/Plugins/PCG/Source/PCG/Public/Compute/Data/PCGRawBufferData.h
@@ -0,0 +1,42 @@
+// Copyright Epic Games, Inc. All Rights Reserved.
+
+#pragma once
+
+#include "PCGData.h"
+
+#include "PCGRawBufferData.generated.h"
+
+USTRUCT()
+struct FPCGDataTypeInfoRawBuffer : public FPCGDataTypeInfo
+{
+	GENERATED_BODY()
+
+	PCG_DECLARE_TYPE_INFO(PCG_API)
+
+};
+
+/** A proxy for data residing on the GPU with functionality to read the data back to the CPU. */
+UCLASS(MinimalAPI, ClassGroup = (Procedural), DisplayName = "Raw Buffer")
+class UPCGRawBufferData : public UPCGData
+{
+	GENERATED_BODY()
+
+public:
+	//~Begin UPCGData interface
+	PCG_ASSIGN_TYPE_INFO(FPCGDataTypeInfoRawBuffer);
+
+	PCG_API virtual void AddToCrc(FArchiveCrc32& Ar, bool bFullDataCrc) const override;
+	PCG_API virtual void GetResourceSizeEx(FResourceSizeEx& CumulativeResourceSize) override;
+	PCG_API virtual void ReleaseTransientResources(const TCHAR* InRease = nullptr) override;
+	//~End UPCGData interface
+
+	void SetData(TArray<uint32>&& InData) { Data = MoveTemp(InData); }
+	int32 GetNumUint32s() const { return Data.Num(); }
+
+	TArray<uint32>& GetData() { return Data; }
+	const TArray<uint32>& GetConstData() const { return Data; }
+
+protected:
+	UPROPERTY()
+	TArray<uint32> Data;
+};
diff --git a/Engine/Plugins/PCG/Source/PCG/Public/Compute/DataInterfaces/PCGRawBufferDataInterface.h b/Engine/Plugins/PCG/Source/PCG/Public/Compute/DataInterfaces/PCGRawBufferDataInterface.h
new file mode 100644
--- /dev/null
+++ b/Engine/Plugins/PCG/Source/PCG/Public/Compute/DataInterfaces/PCGRawBufferDataInterface.h
@@ -0,0 +1,115 @@
+// Copyright Epic Games, Inc. All Rights Reserved.
+
+#pragma once
+
+#include "Compute/DataInterfaces/PCGComputeDataInterface.h"
+
+#include "RHIResources.h"
+
+#include "PCGRawBufferDataInterface.generated.h"
+
+class FPCGRawBufferDataInterfaceParameters;
+class UPCGRawBufferData;
+
+/** A data interface for a simple array of uint values. No data format header or attributes, just raw array access. */
+UCLASS(ClassGroup = (Procedural))
+class UPCGRawBufferDataInterface : public UPCGExportableDataInterface
+{
+	GENERATED_BODY()
+
+public:
+	//~ Begin UComputeDataInterface Interface
+	TCHAR const* GetClassName() const override { return TEXT("PCGRawBuffer"); }
+	void GetSupportedInputs(TArray<FShaderFunctionDefinition>& OutFunctions) const override { return GetSupportedOutputs(OutFunctions); }
+	PCG_API void GetSupportedOutputs(TArray<FShaderFunctionDefinition>& OutFunctions) const override;
+	PCG_API void GetShaderParameters(TCHAR const* UID, FShaderParametersMetadataBuilder& InOutBuilder, FShaderParametersMetadataAllocations& InOutAllocations) const override;
+	PCG_API TCHAR const* GetShaderVirtualPath() const override;
+	PCG_API void GetShaderHash(FString& InOutKey) const override;
+	PCG_API void GetHLSL(FString& OutHLSL, FString const& InDataInterfaceName) const override;
+	virtual bool GetRequiresReadback() const override { return true; }
+	PCG_API UComputeDataProvider* CreateDataProvider() const override;
+	//~ End UComputeDataInterface Interface
+
+	bool GetRequiresZeroInitialization() const { return bRequiresZeroInitialization; }
+	void SetRequiresZeroInitialization(bool bInZeroInit) { bRequiresZeroInitialization = bInZeroInit; }
+
+protected:
+	/** Whether to perform full 0-initialization of the buffer. */
+	UPROPERTY()
+	bool bRequiresZeroInitialization = false;
+};
+
+UCLASS()
+class UPCGRawBufferDataProvider : public UPCGExportableDataProvider
+{
+	GENERATED_BODY()
+
+public:
+	//~ Begin UComputeDataProvider Interface
+	PCG_API virtual void Initialize(const UComputeDataInterface* InDataInterface, UObject* InBinding, uint64 InInputMask, uint64 InOutputMask) override;
+	PCG_API virtual bool PrepareForExecute_GameThread(UPCGDataBinding* InBinding) override;
+	PCG_API virtual FComputeDataProviderRenderProxy* GetRenderProxy() override;
+	PCG_API virtual void Reset() override;
+	//~ End UComputeDataProvider Interface
+
+protected:
+	int32 SizeBytes = INDEX_NONE;
+
+	bool bZeroInitialize = false;
+
+	UPROPERTY()
+	TObjectPtr<const UPCGRawBufferData> DataToUpload;
+
+	struct FReadbackState
+	{
+		TArray<uint32> Data;
+	};
+	TSharedPtr<FReadbackState, ESPMode::ThreadSafe> ReadbackState;
+};
+
+class FPCGRawBufferDataProviderProxy : public FComputeDataProviderRenderProxy
+{
+public:
+	struct FParams
+	{
+		int32 SizeBytes = 0;
+		bool bZeroInitialize = false;
+		const UPCGRawBufferData* DataToUpload = nullptr;
+		EPCGExportMode ExportMode = EPCGExportMode::NoExport;
+		FName OutputPinLabel;
+		FName OutputPinLabelAlias;
+		TWeakObjectPtr<UPCGRawBufferDataProvider> DataProviderWeakPtr;
+		FReadbackCallback AsyncReadbackCallback_RenderThread;
+	};
+
+	PCG_API FPCGRawBufferDataProviderProxy(const FParams& InParams);
+
+	//~ Begin FComputeDataProviderRenderProxy Interface
+	PCG_API bool IsValid(FValidationData const& InValidationData) const override;
+	PCG_API void AllocateResources(FRDGBuilder& GraphBuilder, FAllocationData const& InAllocationData);
+	PCG_API void GatherDispatchData(FDispatchData const& InDispatchData) override;
+	PCG_API void GetReadbackData(TArray<FReadbackData>& OutReadbackData) const override;
+	//~ End FComputeDataProviderRenderProxy Interface
+
+protected:
+	using FParameters = FPCGRawBufferDataInterfaceParameters;
+
+	int32 SizeBytes = INDEX_NONE;
+
+	bool bZeroInitialize = false;
+
+	EPCGExportMode ExportMode = EPCGExportMode::NoExport;
+
+	FName OutputPinLabel;
+	FName OutputPinLabelAlias;
+
+	FRDGBufferRef Data = nullptr;
+	FRDGBufferUAVRef DataUAV = nullptr;
+
+	TArray<uint32> DataToUpload;
+
+	TWeakObjectPtr<UPCGRawBufferDataProvider> DataProviderWeakPtr;
+
+	/** Called from render thread when readback from GPU is complete. */
+	FReadbackCallback AsyncReadbackCallback_RenderThread;
+};
diff --git a/Engine/Plugins/PCG/Source/PCGEditor/Private/DataVisualizations/PCGRawBufferDataVisualization.cpp b/Engine/Plugins/PCG/Source/PCGEditor/Private/DataVisualizations/PCGRawBufferDataVisualization.cpp
new file mode 100644
--- /dev/null
+++ b/Engine/Plugins/PCG/Source/PCGEditor/Private/DataVisualizations/PCGRawBufferDataVisualization.cpp
@@ -0,0 +1,118 @@
+// Copyright Epic Games, Inc. All Rights Reserved.
+
+#include "DataVisualizations/PCGRawBufferDataVisualization.h"
+
+#include "Compute/Data/PCGRawBufferData.h"
+#include "DataVisualizations/PCGDataVisualizationHelpers.h"
+#include "Metadata/Accessors/IPCGAttributeAccessor.h"
+#include "Metadata/Accessors/PCGAttributeAccessorKeys.h"
+#include "Metadata/Accessors/PCGCustomAccessor.h"
+
+namespace PCGRawBufferVisualization
+{
+	/** Simple keys implementation that reports the number of uint32 values in the buffer. */
+	class FKeys : public IPCGAttributeAccessorKeys
+	{
+	public:
+		explicit FKeys(int32 InCount)
+			: IPCGAttributeAccessorKeys(/*bInReadOnly=*/ true)
+			, Count(InCount)
+		{}
+
+		virtual int32 GetNum() const override { return Count; }
+
+	private:
+		int32 Count;
+	};
+
+	/** Read-only accessor that reads uint32 values from the raw buffer, exposing them as int32. */
+	class FValueAccessor : public IPCGAttributeAccessorT<FValueAccessor>
+	{
+	public:
+		using Type = int32;
+		using Super = IPCGAttributeAccessorT<FValueAccessor>;
+
+		explicit FValueAccessor(const TArray<uint32>& InData)
+			: Super(/*bInReadOnly=*/ true)
+			, Data(InData)
+		{}
+
+		bool GetRangeImpl(TArrayView<int32> OutValues, int32 Index, const IPCGAttributeAccessorKeys& InKeys) const
+		{
+			const int32 NumKeys = InKeys.GetNum();
+			if (NumKeys == 0)
+			{
+				return false;
+			}
+
+			int32 Current = Index;
+			if (Current >= NumKeys)
+			{
+				Current %= NumKeys;
+			}
+
+			for (int32 i = 0; i < OutValues.Num(); ++i)
+			{
+				OutValues[i] = Data.IsValidIndex(Current) ? static_cast<int32>(Data[Current]) : 0;
+				if (++Current >= NumKeys)
+				{
+					Current = 0;
+				}
+			}
+
+			return true;
+		}
+
+		bool SetRangeImpl(TArrayView<const int32>, int32, IPCGAttributeAccessorKeys&, EPCGAttributeAccessorFlags)
+		{
+			return false;
+		}
+
+	private:
+		const TArray<uint32>& Data;
+	};
+}
+
+void FPCGRawBufferDataVisualization::ExecuteDebugDisplay(FPCGContext* Context, const UPCGSettingsInterface* SettingsInterface, const UPCGData* Data, AActor* TargetActor) const
+{
+	// No spatial debug display for raw buffer data.
+}
+
+FPCGTableVisualizerInfo FPCGRawBufferDataVisualization::GetTableVisualizerInfoWithDomain(const UPCGData* Data, const FPCGMetadataDomainID& DomainID) const
+{
+	FPCGTableVisualizerInfo Info;
+
+	const UPCGRawBufferData* RawBufferData = Cast<UPCGRawBufferData>(Data);
+	if (!RawBufferData)
+	{
+		return Info;
+	}
+
+	Info.Data = RawBufferData;
+
+	const int32 NumValues = RawBufferData->GetNumUint32s();
+	TSharedRef<PCGRawBufferVisualization::FKeys> Keys = MakeShared<PCGRawBufferVisualization::FKeys>(NumValues);
+
+	// Index column
+	{
+		FPCGTableVisualizerColumnInfo& Col = Info.ColumnInfos.Emplace_GetRef();
+		Col.Id = PCGDataVisualizationConstants::NAME_Index;
+		Col.Label = FText::FromString(TEXT("$Index"));
+		Col.Accessor = MakeShared<FPCGIndexAccessor>();
+		Col.AccessorKeys = Keys;
+	}
+
+	// Value column
+	{
+		FPCGTableVisualizerColumnInfo& Col = Info.ColumnInfos.Emplace_GetRef();
+		Col.Id = FName(TEXT("Value"));
+		Col.Label = FText::FromString(TEXT("Value"));
+		Col.Tooltip = FText::FromString(TEXT("Raw uint32 value"));
+		Col.Accessor = MakeShared<PCGRawBufferVisualization::FValueAccessor>(RawBufferData->GetConstData());
+		Col.AccessorKeys = Keys;
+	}
+
+	Info.SortingColumn = PCGDataVisualizationConstants::NAME_Index;
+
+	return Info;
+}
diff --git a/Engine/Plugins/PCG/Source/PCGEditor/Private/PCGEditorModule.cpp b/Engine/Plugins/PCG/Source/PCGEditor/Private/PCGEditorModule.cpp
--- a/Engine/Plugins/PCG/Source/PCGEditor/Private/PCGEditorModule.cpp
+++ b/Engine/Plugins/PCG/Source/PCGEditor/Private/PCGEditorModule.cpp
@@ -46,11 +46,13 @@
 #include "ChangeTracking/PCGLandscapeTracker.h"
 #include "ChangeTracking/PCGMiscTracker.h"
 #include "ChangeTracking/PCGRuntimeGenChangeHandler.h"
+#include "Compute/Data/PCGRawBufferData.h"
 #include "Data/DataView/PCGDataViewData.h"
 #include "DataVisualizations/PCGBaseTextureDataVisualization.h"
 #include "DataVisualizations/PCGDataViewVisualization.h"
 #include "DataVisualizations/PCGLandscapeDataVisualization.h"
 #include "DataVisualizations/PCGParamDataVisualization.h"
+#include "DataVisualizations/PCGRawBufferDataVisualization.h"
 #include "DataVisualizations/PCGPolygon2DDataVisualization.h"
 #include "DataVisualizations/PCGSpatialDataVisualization.h"
 #include "DataVisualizations/PCGSplineDataVisualization.h"
@@ -1050,6 +1052,7 @@ void FPCGEditorModule::RegisterPCGDataVisualizations()
 	DataVisRegistry.InternalRegistry.Add(UPCGPolygon2DData::StaticClass(), MakeUnique<const FPCGPolygon2DDataVisualization>());
 	DataVisRegistry.InternalRegistry.Add(UPCGPolygon2DInteriorSurfaceData::StaticClass(), MakeUnique<const FPCGPolygon2DDataVisualization>());
 	DataVisRegistry.InternalRegistry.Add(UPCGDataViewData::StaticClass(), MakeUnique<const FPCGDataViewVisualization>());
+	DataVisRegistry.InternalRegistry.Add(UPCGRawBufferData::StaticClass(), MakeUnique<const FPCGRawBufferDataVisualization>());
 }
 
 void FPCGEditorModule::UnregisterPCGDataVisualizations()
@@ -1156,6 +1159,7 @@ void FPCGEditorModule::RegisterPinColorAndIcons()
 	InRegistry.RegisterPinColorFunction(FPCGDataTypeInfoDynamicMesh::AsId(), [](const FPCGDataTypeIdentifier&) { return GetDefault<UPCGEditorSettings>()->DynamicMeshPinColor; });
 	InRegistry.RegisterPinColorFunction(FPCGDataTypeInfoSpatial::AsId(), [](const FPCGDataTypeIdentifier&) { return GetDefault<UPCGEditorSettings>()->SpatialDataPinColor; });
 	InRegistry.RegisterPinColorFunction(FPCGDataTypeInfoStaticMeshResource::AsId(), [](const FPCGDataTypeIdentifier&) { return GetDefault<UPCGEditorSettings>()->StaticMeshResourcePinColor; });
+	InRegistry.RegisterPinColorFunction(FPCGDataTypeInfoRawBuffer::AsId(), [](const FPCGDataTypeIdentifier&) { return GetDefault<UPCGEditorSettings>()->RawBufferPinColor; });
 	InRegistry.RegisterPinColorFunction(FPCGDataTypeInfoDataView::AsId(), [](const FPCGDataTypeIdentifier&) { return GetDefault<UPCGEditorSettings>()->DataViewPinColor; });
 	InRegistry.RegisterPinColorFunction(FPCGDataTypeInfoOther::AsId(), [](const FPCGDataTypeIdentifier&) { return GetDefault<UPCGEditorSettings>()->UnknownDataPinColor; });
 	InRegistry.RegisterPinColorFunction(FPCGDataTypeInfo::AsId(), [](const FPCGDataTypeIdentifier&) { return GetDefault<UPCGEditorSettings>()->DefaultPinColor; });
diff --git a/Engine/Plugins/PCG/Source/PCGEditor/Private/PCGEditorSettings.cpp b/Engine/Plugins/PCG/Source/PCGEditor/Private/PCGEditorSettings.cpp
--- a/Engine/Plugins/PCG/Source/PCGEditor/Private/PCGEditorSettings.cpp
+++ b/Engine/Plugins/PCG/Source/PCGEditor/Private/PCGEditorSettings.cpp
@@ -57,6 +57,7 @@ UPCGEditorSettings::UPCGEditorSettings(const FObjectInitializer& ObjectInitializ
 	PrimitiveDataPinColor = FLinearColor(0.22f, 0.05f, 1.0f);
 	DynamicMeshPinColor = FLinearColor(0.42f, 0.25f, 1.0f);
 	StaticMeshResourcePinColor = FLinearColor(0.82f, 0.25f, 0.1f);
+	RawBufferPinColor = FLinearColor(0.6f, 0.7f, 1.0f);
 
 	ParamDataPinColor = FLinearColor(1.0f, 0.38f, 0.02f);
 	UnknownDataPinColor = FLinearColor(0.3f, 0.3f, 0.3f);
diff --git a/Engine/Plugins/PCG/Source/PCGEditor/Public/DataVisualizations/PCGRawBufferDataVisualization.h b/Engine/Plugins/PCG/Source/PCGEditor/Public/DataVisualizations/PCGRawBufferDataVisualization.h
new file mode 100644
--- /dev/null
+++ b/Engine/Plugins/PCG/Source/PCGEditor/Public/DataVisualizations/PCGRawBufferDataVisualization.h
@@ -0,0 +1,14 @@
+// Copyright Epic Games, Inc. All Rights Reserved.
+
+#pragma once
+
+#include "PCGDataVisualization.h"
+
+class UPCGData;
+
+class FPCGRawBufferDataVisualization : public IPCGDataVisualization
+{
+public:
+	PCGEDITOR_API virtual void ExecuteDebugDisplay(FPCGContext* Context, const UPCGSettingsInterface* SettingsInterface, const UPCGData* Data, AActor* TargetActor) const override;
+	PCGEDITOR_API virtual FPCGTableVisualizerInfo GetTableVisualizerInfoWithDomain(const UPCGData* Data, const FPCGMetadataDomainID& DomainID) const override;
+};
diff --git a/Engine/Plugins/PCG/Source/PCGEditor/Public/PCGEditorSettings.h b/Engine/Plugins/PCG/Source/PCGEditor/Public/PCGEditorSettings.h
--- a/Engine/Plugins/PCG/Source/PCGEditor/Public/PCGEditorSettings.h
+++ b/Engine/Plugins/PCG/Source/PCGEditor/Public/PCGEditorSettings.h
@@ -249,6 +249,10 @@ public:
 	UPROPERTY(EditAnywhere, config, Category = Node, meta = (HideAlphaChannel))
 	FLinearColor StaticMeshResourcePinColor;
 
+	/** Color used for data pins of type Raw Buffer */
+	UPROPERTY(EditAnywhere, config, Category = Node, meta = (HideAlphaChannel))
+	FLinearColor RawBufferPinColor;
+
 	/** Color used for other/unknown data types */
 	UPROPERTY(EditAnywhere, config, Category = Node, meta = (HideAlphaChannel))
 	FLinearColor UnknownDataPinColor;
