[PCG] Fix dynamic culling for GPU nodes + tests.

Input-less wires were not being created on the compute graph tasks, and pin dependency info was not being propagated across compute graph tasks.

Add support in the GPU test common code for properly detecting culled nodes.

[RN] [PCG] Fixed culling of GPU nodes so that their execution is subject to control flow such as Branch nodes.
#Synced-CL 52685145

diff --git a/Engine/Plugins/PCG/Source/PCG/Private/Graph/PCGGraphCompiler.cpp b/Engine/Plugins/PCG/Source/PCG/Private/Graph/PCGGraphCompiler.cpp
--- a/Engine/Plugins/PCG/Source/PCG/Private/Graph/PCGGraphCompiler.cpp
+++ b/Engine/Plugins/PCG/Source/PCG/Private/Graph/PCGGraphCompiler.cpp
@@ -1135,17 +1135,31 @@ void FPCGGraphCompiler::CalculateDynamicActivePinDependencies(FPCGTaskId InTaskI
 
 	for (const FPCGGraphTaskInput& Input : InOutCompiledTasks[InTaskId].Inputs)
 	{
-		if (!Node || !Input.DownstreamPin.IsSet())
+		if (!Input.DownstreamPin.IsSet())
 		{
 			continue;
 		}
 
-		const UPCGPin* InputPin = Node->GetInputPin(Input.DownstreamPin.GetValue().Label);
+		if (Node)
+		{
+			const UPCGPin* InputPin = Node->GetInputPin(Input.DownstreamPin.GetValue().Label);
 
-		// Consider only primary input pins in this pass.
-		if (!InputPin || !Node->IsInputPinRequiredByExecution(InputPin))
+			// Consider only primary input pins in this pass.
+			if (!InputPin || !Node->IsInputPinRequiredByExecution(InputPin))
+			{
+				continue;
+			}
+		}
+		else
 		{
-			continue;
+			// For ComputeGraph tasks (Node=nullptr), the node's settings are unavailable so we cannot call
+			// IsInputPinRequiredByExecution. Fall back to the downstream pin's own PinStatus: pins explicitly
+			// marked Required (e.g. exec dep pins with bExecutionDependencyRequired=true) are treated as
+			// required by execution. Normal/Advanced pins fall through to the disjunction pass below.
+			if (!Input.DownstreamPin.GetValue().IsRequiredPin())
+			{
+				continue;
+			}
 		}
 
 		const UPCGNode* UpstreamNode = InOutCompiledTasks[Input.TaskId].Node;
@@ -1209,25 +1223,30 @@ void FPCGGraphCompiler::CalculateDynamicActivePinDependencies(FPCGTaskId InTaskI
 				const UPCGNode* UpstreamNode = InOutCompiledTasks[Input.TaskId].Node;
 				if (!UpstreamNode)
 				{
-					continue;
+					// Upstream is a ComputeGraph task (Node=nullptr). We cannot resolve a pin index from its virtual output pins, so fall back
+					// to a pin-less dependency. This ensures that if the ComputeGraph task is dynamically culled, the culling propagates to
+					// this task as well (CullInactiveDownstreamNodes adds NodeIdToPinId for culled tasks).
+					PinDependency.AddPinDependency(PCGPinIdHelpers::NodeIdToPinId(Input.TaskId));
 				}
-
-				if (UpstreamNode->GetOutputPins().Num() > PCGPinIdHelpers::MaxOutputPins)
+				else
 				{
-					LogTooManyOutputPinsWarning(UpstreamNode);
-					continue;
-				}
+					if (UpstreamNode->GetOutputPins().Num() > PCGPinIdHelpers::MaxOutputPins)
+					{
+						LogTooManyOutputPinsWarning(UpstreamNode);
+						continue;
+					}
 
-				const int PinIndex = UpstreamNode->GetOutputPins().IndexOfByPredicate([&Input](const UPCGPin* InPin)
-				{
-					return InPin->Properties == Input.UpstreamPin.GetValue();
-				});
+					const int PinIndex = UpstreamNode->GetOutputPins().IndexOfByPredicate([&Input](const UPCGPin* InPin)
+					{
+						return InPin->Properties.Label == Input.UpstreamPin.GetValue().Label;
+					});
 
-				if (PinIndex != INDEX_NONE)
-				{
-					check(PinIndex < PCGPinIdHelpers::MaxOutputPins);
+					if (PinIndex != INDEX_NONE)
+					{
+						check(PinIndex < PCGPinIdHelpers::MaxOutputPins);
 
-					PinDependency.AddPinDependency(PCGPinIdHelpers::NodeIdAndPinIndexToPinId(Input.TaskId, PinIndex));
+						PinDependency.AddPinDependency(PCGPinIdHelpers::NodeIdAndPinIndexToPinId(Input.TaskId, PinIndex));
+					}
 				}
 			}
 			else
diff --git a/Engine/Plugins/PCG/Source/PCG/Private/Graph/PCGGraphCompilerGPU.cpp b/Engine/Plugins/PCG/Source/PCG/Private/Graph/PCGGraphCompilerGPU.cpp
--- a/Engine/Plugins/PCG/Source/PCG/Private/Graph/PCGGraphCompilerGPU.cpp
+++ b/Engine/Plugins/PCG/Source/PCG/Private/Graph/PCGGraphCompilerGPU.cpp
@@ -522,8 +522,7 @@ void FPCGGraphCompilerGPU::ExpandGPUNodeKernelsToTasks(
 						FPCGGraphTaskInput& NewTaskInput = NewKernelTask.Inputs.Add_GetRef(*OriginalTaskInput);
 
 						// Resolve dynamic pin type here to simplify later logic.
-						// todo_pcg: Possibly the graph compiler should resolve dynamic pins during compilation. In a quick test this broke dynamic culling logic
-						// due to mismatching pin properties.
+						// todo_pcg: Possibly the graph compiler should resolve dynamic pins during compilation.
 						if (NewTaskInput.UpstreamPin)
 						{
 							const UPCGSettings* ProducerSettings = InOutCompiledTasks[NewTaskInput.TaskId].Node ? InOutCompiledTasks[NewTaskInput.TaskId].Node->GetSettings() : nullptr;
@@ -543,6 +542,16 @@ void FPCGGraphCompilerGPU::ExpandGPUNodeKernelsToTasks(
 				}
 			}
 
+			// Preserve dependency-only inputs (e.g. execution dependency pins) from the original task. These are PCG control-flow inputs
+			// (bProvideData=false) that have no kernel equivalent but are required for dynamic culling to work correctly at execution time.
+			for (const FPCGGraphTaskInput& OriginalInput : InOutCompiledTasks[TaskId].Inputs)
+			{
+				if (!OriginalInput.bProvideData && OriginalInput.DownstreamPin.IsSet())
+				{
+					NewKernelTask.Inputs.Add(OriginalInput);
+				}
+			}
+
 			NodeKernels.RemoveAt(ReadyKernelIndex);
 		}
 
diff --git a/Engine/Plugins/PCG/Source/PCG/Private/Graph/PCGGraphExecutor.cpp b/Engine/Plugins/PCG/Source/PCG/Private/Graph/PCGGraphExecutor.cpp
--- a/Engine/Plugins/PCG/Source/PCG/Private/Graph/PCGGraphExecutor.cpp
+++ b/Engine/Plugins/PCG/Source/PCG/Private/Graph/PCGGraphExecutor.cpp
@@ -3118,6 +3118,16 @@ FPCGTaskId FPCGGraphExecutor::ScheduleDebugWithTaskCallback(UPCGComponent* InCom
 		// Schedule the output capture hooks
 		FPCGTaskId CaptureTaskId = ScheduleGeneric([this, TaskCompleteCallback, CompiledTask]
 		{
+			// If the task was dynamically culled it produced no output — skip the callback.
+			{
+				PCG::TSharedScopeLock ReadLock(TaskOutputsLock);
+				const FOutputDataInfo* OutputDataInfo = TaskOutputs.Find(CompiledTask.NodeId);
+				if (OutputDataInfo && OutputDataInfo->bCulled)
+				{
+					return true;
+				}
+			}
+
 			FPCGDataCollection TaskOutputData;
 			if (CompiledTask.Node && GetOutputData(CompiledTask.NodeId, TaskOutputData))
 			{
diff --git a/Engine/Plugins/PCG/Source/PCG/Private/Tests/Compute/PCGGPUTestCommon.cpp b/Engine/Plugins/PCG/Source/PCG/Private/Tests/Compute/PCGGPUTestCommon.cpp
--- a/Engine/Plugins/PCG/Source/PCG/Private/Tests/Compute/PCGGPUTestCommon.cpp
+++ b/Engine/Plugins/PCG/Source/PCG/Private/Tests/Compute/PCGGPUTestCommon.cpp
@@ -126,6 +126,8 @@ bool FPCGGPUGraphTestRunner::Execute(const FPCGGPUGraphTestRunnerParams& Params,
 	CapturedOutputs.SetNum(Params.CollectOutputDataForNodes.Num());
 	TArray<bool> NodeOutputCaptured;
 	NodeOutputCaptured.Init(false, Params.CollectOutputDataForNodes.Num());
+	TArray<bool> CulledNodeExecuted;
+	CulledNodeExecuted.Init(false, Params.ExpectCulledNodes.Num());
 
 	// This will execute the graph and run the lambda on each node, which we use to capture output data.
 	const FPCGTaskId FinalTaskId = Executor->ScheduleDebugWithTaskCallback(
@@ -141,6 +143,15 @@ bool FPCGGPUGraphTestRunner::Execute(const FPCGGPUGraphTestRunnerParams& Params,
 					break;
 				}
 			}
+
+			for (int32 i = 0; i < Params.ExpectCulledNodes.Num(); ++i)
+			{
+				if (Node == Params.ExpectCulledNodes[i])
+				{
+					CulledNodeExecuted[i] = true;
+					break;
+				}
+			}
 		}
 	);
 
@@ -194,6 +205,16 @@ bool FPCGGPUGraphTestRunner::Execute(const FPCGGPUGraphTestRunnerParams& Params,
 		return false;
 	}
 
+	bool bAllCulledNodesWereCulled = true;
+	for (int32 i = 0; i < Params.ExpectCulledNodes.Num(); ++i)
+	{
+		bAllCulledNodesWereCulled &= Params.Test->TestFalse(*FString::Printf(TEXT("Node '%s' expected to be culled was not executed"), *Params.ExpectCulledNodes[i]->GetNodeTitle(EPCGNodeTitleType::ListView).ToString()), CulledNodeExecuted[i]);
+	}
+	if (!bAllCulledNodesWereCulled)
+	{
+		return false;
+	}
+
 	if (Params.bReadbackGPUData)
 	{
 		FlushRenderingCommands();
diff --git a/Engine/Plugins/PCG/Source/PCG/Private/Tests/Compute/PCGGPUTestCommon.h b/Engine/Plugins/PCG/Source/PCG/Private/Tests/Compute/PCGGPUTestCommon.h
--- a/Engine/Plugins/PCG/Source/PCG/Private/Tests/Compute/PCGGPUTestCommon.h
+++ b/Engine/Plugins/PCG/Source/PCG/Private/Tests/Compute/PCGGPUTestCommon.h
@@ -76,6 +76,10 @@ struct FPCGGPUGraphTestRunnerParams
 
 	/** If true (default), UPCGProxyForGPUData outputs are read back to CPU before being added to OutNodeData. */
 	bool bReadbackGPUData = true;
+
+	/** Nodes that should be culled (statically or dynamically) and therefore never execute. Execute() asserts that the debug callback was
+	* NOT fired for each of these nodes. */
+	TArray<UPCGNode*> ExpectCulledNodes;
 };
 
 /**
diff --git a/Engine/Plugins/PCG/Source/PCG/Private/Tests/Elements/PCGBranchTest.cpp b/Engine/Plugins/PCG/Source/PCG/Private/Tests/Elements/PCGBranchTest.cpp
--- a/Engine/Plugins/PCG/Source/PCG/Private/Tests/Elements/PCGBranchTest.cpp
+++ b/Engine/Plugins/PCG/Source/PCG/Private/Tests/Elements/PCGBranchTest.cpp
@@ -1,11 +1,19 @@
 // Copyright Epic Games, Inc. All Rights Reserved.
 
+#if WITH_EDITOR
+
+#include "PCGCommon.h"
 #include "PCGContext.h"
+#include "PCGGraph.h"
 #include "Data/PCGPointData.h"
+#include "Data/PCGBasePointData.h"
+#include "Elements/PCGCreateAttribute.h"
+#include "Elements/PCGCreatePoints.h"
+#include "Elements/PCGGather.h"
+#include "Elements/PCGTransformPoints.h"
 #include "Elements/ControlFlow/PCGBranch.h"
 #include "Tests/PCGTestsCommon.h"
-
-#if WITH_EDITOR
+#include "Tests/Compute/PCGGPUTestCommon.h"
 
 class FPCGBranchTest : public FPCGTestBaseClass
 {
@@ -60,4 +68,495 @@ bool FPCGBranchTest_Basic::RunTest(const FString& Parameters)
 	return true;
 }
 
-#endif // WITH_EDITOR
\ No newline at end of file
+/**
+ * Shared GPU graph builder for the branch tests.
+ *
+ * Builds:
+ *   CreatePoints (1 point) -> Branch -> "Output A" -> TransformPointsA (+100 on X) -\
+ *                                    -> "Output B" -> TransformPointsB (+200 on Y) --> Gather
+ *
+ * The caller is responsible for wiring the branch condition (static or dynamic).
+ */
+namespace PCGBranchTestGraphGPU
+{
+	// Expected point location when routed through Output A (TransformPoints +100 on X).
+	static const FVector OutputALocation = FVector(100.0, 0.0, 0.0);
+	// Expected point location when routed through Output B (TransformPoints +200 on Y).
+	static const FVector OutputBLocation = FVector(0.0, 200.0, 0.0);
+
+	struct FNodes
+	{
+		UPCGNode* BranchNode = nullptr;
+		UPCGNode* GatherNode = nullptr;
+	};
+
+	static FNodes Build(UPCGGraph* Graph, bool bOutputToB = false)
+	{
+		// CreatePoints: one default point at the origin.
+		UPCGCreatePointsSettings* CreatePointsSettings = nullptr;
+		UPCGNode* CreatePointsNode = Graph->AddNodeOfType<UPCGCreatePointsSettings>(CreatePointsSettings);
+		CreatePointsSettings->PointsToCreate.Empty();
+		CreatePointsSettings->PointsToCreate.AddDefaulted();
+		CreatePointsSettings->CoordinateSpace = EPCGCoordinateSpace::World;
+
+		// Branch: condition wired by the caller.
+		UPCGBranchSettings* BranchSettings = nullptr;
+		UPCGNode* BranchNode = Graph->AddNodeOfType<UPCGBranchSettings>(BranchSettings);
+		BranchSettings->bOutputToB = bOutputToB;
+
+		// TransformPoints A: translate +100 on X.
+		UPCGTransformPointsSettings* TransformASettings = nullptr;
+		UPCGNode* TransformANode = Graph->AddNodeOfType<UPCGTransformPointsSettings>(TransformASettings);
+		TransformASettings->SetExecuteOnGPU(true);
+		TransformASettings->OffsetMin = OutputALocation;
+		TransformASettings->OffsetMax = OutputALocation;
+
+		// TransformPoints B: translate +200 on Y.
+		UPCGTransformPointsSettings* TransformBSettings = nullptr;
+		UPCGNode* TransformBNode = Graph->AddNodeOfType<UPCGTransformPointsSettings>(TransformBSettings);
+		TransformBSettings->SetExecuteOnGPU(true);
+		TransformBSettings->OffsetMin = OutputBLocation;
+		TransformBSettings->OffsetMax = OutputBLocation;
+
+		// Gather: collects from both transform nodes.
+		UPCGGatherSettings* GatherSettings = nullptr;
+		UPCGNode* GatherNode = Graph->AddNodeOfType<UPCGGatherSettings>(GatherSettings);
+
+		Graph->AddEdge(CreatePointsNode, PCGPinConstants::DefaultOutputLabel, BranchNode,     PCGPinConstants::DefaultInputLabel);
+		Graph->AddEdge(BranchNode,       PCGBranchConstants::OutputLabelA,    TransformANode, PCGPinConstants::DefaultInputLabel);
+		Graph->AddEdge(BranchNode,       PCGBranchConstants::OutputLabelB,    TransformBNode, PCGPinConstants::DefaultInputLabel);
+		Graph->AddEdge(TransformANode,   PCGPinConstants::DefaultOutputLabel, GatherNode,     PCGPinConstants::DefaultInputLabel);
+		Graph->AddEdge(TransformBNode,   PCGPinConstants::DefaultOutputLabel, GatherNode,     PCGPinConstants::DefaultInputLabel);
+
+		return { BranchNode, GatherNode };
+	}
+
+	/** Validates a single tagged data item: expects exactly 1 point at the given location. */
+	static bool ValidateOutput(FAutomationTestBase* Test, const FPCGTaggedData& TaggedData, const FVector& ExpectedLocation)
+	{
+		const UPCGBasePointData* OutPoints = Cast<UPCGBasePointData>(TaggedData.Data);
+		if (!Test->TestNotNull(TEXT("Readback data is point data"), OutPoints))
+		{
+			return false;
+		}
+
+		if (!Test->TestEqual(TEXT("Exactly one point"), OutPoints->GetNumPoints(), 1))
+		{
+			return false;
+		}
+
+		const FVector ActualLocation = OutPoints->GetTransform(0).GetLocation();
+
+		return Test->TestNearlyEqual(*FString::Printf(TEXT("Point location (expected [%.0f, %.0f, %.0f], got [%.1f, %.1f, %.1f])"),
+				ExpectedLocation.X, ExpectedLocation.Y, ExpectedLocation.Z, ActualLocation.X, ActualLocation.Y, ActualLocation.Z), ActualLocation, ExpectedLocation, 1.0f);
+	}
+}
+
+// -----------------------------------------------------------------------------------------
+// Static branch tests
+//
+// bOutputToB is baked directly onto the branch settings at graph construction time.
+// IsPinStaticallyActive() returns false for the inactive pin, so the inactive GPU node is
+// culled at graph-compile time before any dispatch occurs.
+// -----------------------------------------------------------------------------------------
+
+/**
+ * Static branch (GPU): bOutputToB=false routes through Output A.
+ * TransformPointsB is statically culled. Gather produces 1 point at (100, 0, 0).
+ */
+IMPLEMENT_CUSTOM_SIMPLE_AUTOMATION_TEST(
+	FPCGBranchStaticTest_RouteToA_GPU,
+	FPCGTestBaseClass,
+	"Plugins.PCG.Branch.Static.RouteToA.GPU",
+	PCGTestsCommon::TestFlags | EAutomationTestFlags::NonNullRHI)
+
+bool FPCGBranchStaticTest_RouteToA_GPU::RunTest(const FString& Parameters)
+{
+	FPCGGPUGraphTestRunner Runner;
+	UPCGNode* GatherNode = nullptr;
+	GatherNode = PCGBranchTestGraphGPU::Build(Runner.Graph, /*bOutputToB=*/false).GatherNode;
+
+	TArray<FPCGDataCollection> OutNodeData;
+	if (!Runner.Execute({ this, { GatherNode } }, OutNodeData))
+	{
+		return false;
+	}
+
+	return PCGTestsCommon::ValidateSingleTaggedData(this, OutNodeData[0],
+		[&](const FPCGTaggedData& TaggedData) { return PCGBranchTestGraphGPU::ValidateOutput(this, TaggedData, PCGBranchTestGraphGPU::OutputALocation); });
+}
+
+/**
+ * Static branch (GPU): bOutputToB=true routes through Output B.
+ * TransformPointsA is statically culled. Gather produces 1 point at (0, 200, 0).
+ */
+IMPLEMENT_CUSTOM_SIMPLE_AUTOMATION_TEST(
+	FPCGBranchStaticTest_RouteToB_GPU,
+	FPCGTestBaseClass,
+	"Plugins.PCG.Branch.Static.RouteToB.GPU",
+	PCGTestsCommon::TestFlags | EAutomationTestFlags::NonNullRHI)
+
+bool FPCGBranchStaticTest_RouteToB_GPU::RunTest(const FString& Parameters)
+{
+	FPCGGPUGraphTestRunner Runner;
+	UPCGNode* GatherNode = nullptr;
+	GatherNode = PCGBranchTestGraphGPU::Build(Runner.Graph, /*bOutputToB=*/true).GatherNode;
+
+	TArray<FPCGDataCollection> OutNodeData;
+	if (!Runner.Execute({ this, { GatherNode } }, OutNodeData))
+	{
+		return false;
+	}
+
+	return PCGTestsCommon::ValidateSingleTaggedData(this, OutNodeData[0],
+		[&](const FPCGTaggedData& TaggedData) { return PCGBranchTestGraphGPU::ValidateOutput(this, TaggedData, PCGBranchTestGraphGPU::OutputBLocation); });
+}
+
+// -----------------------------------------------------------------------------------------
+// Dynamic branch tests
+//
+// bOutputToB is driven at runtime by a bool attribute fed to the branch's dedicated
+// bOutputToB property override pin. Using the dedicated pin (rather than DefaultParamsLabel)
+// ensures IsPropertyOverriddenByPin() returns true at compile time, keeping both output
+// pins statically active so both GPU nodes compile into the task graph.
+// IsPinStaticallyActive() returns true for both output pins (the inactive one is unknown
+// until execution), so culling happens dynamically after the branch executes.
+// -----------------------------------------------------------------------------------------
+
+/**
+ * Dynamic branch (GPU): bOutputToB attribute=false routes through Output A.
+ * TransformPointsB is dynamically culled at runtime. Gather produces 1 point at (100, 0, 0).
+ */
+IMPLEMENT_CUSTOM_SIMPLE_AUTOMATION_TEST(
+	FPCGBranchDynamicTest_RouteToA_GPU,
+	FPCGTestBaseClass,
+	"Plugins.PCG.Branch.Dynamic.RouteToA.GPU",
+	PCGTestsCommon::TestFlags | EAutomationTestFlags::NonNullRHI)
+
+bool FPCGBranchDynamicTest_RouteToA_GPU::RunTest(const FString& Parameters)
+{
+	FPCGGPUGraphTestRunner Runner;
+	UPCGNode* BranchNode = nullptr;
+	UPCGNode* GatherNode = nullptr;
+	const PCGBranchTestGraphGPU::FNodes Nodes = PCGBranchTestGraphGPU::Build(Runner.Graph);
+	BranchNode = Nodes.BranchNode;
+	GatherNode = Nodes.GatherNode;
+
+	// CreateAttributeSet: bool attribute "bOutputToB" = false, fed to the branch's dedicated bOutputToB override pin.
+	UPCGCreateAttributeSetSettings* CreateAttrSettings = nullptr;
+	UPCGNode* CreateAttrNode = Runner.Graph->AddNodeOfType<UPCGCreateAttributeSetSettings>(CreateAttrSettings);
+	CreateAttrSettings->AttributeTypes.Type = EPCGMetadataTypes::Boolean;
+	CreateAttrSettings->AttributeTypes.BoolValue = false;
+	CreateAttrSettings->OutputTarget.SetAttributeName(GET_MEMBER_NAME_CHECKED(UPCGBranchSettings, bOutputToB));
+
+	Runner.Graph->AddEdge(CreateAttrNode, PCGPinConstants::DefaultOutputLabel, BranchNode, GET_MEMBER_NAME_CHECKED(UPCGBranchSettings, bOutputToB));
+
+	TArray<FPCGDataCollection> OutNodeData;
+	if (!Runner.Execute({ this, { GatherNode } }, OutNodeData))
+	{
+		return false;
+	}
+
+	return PCGTestsCommon::ValidateSingleTaggedData(this, OutNodeData[0],
+		[&](const FPCGTaggedData& TaggedData) { return PCGBranchTestGraphGPU::ValidateOutput(this, TaggedData, PCGBranchTestGraphGPU::OutputALocation); });
+}
+
+/**
+ * Dynamic branch (GPU): bOutputToB attribute=true routes through Output B.
+ * TransformPointsA is dynamically culled at runtime. Gather produces 1 point at (0, 200, 0).
+ */
+IMPLEMENT_CUSTOM_SIMPLE_AUTOMATION_TEST(
+	FPCGBranchDynamicTest_RouteToB_GPU,
+	FPCGTestBaseClass,
+	"Plugins.PCG.Branch.Dynamic.RouteToB.GPU",
+	PCGTestsCommon::TestFlags | EAutomationTestFlags::NonNullRHI)
+
+bool FPCGBranchDynamicTest_RouteToB_GPU::RunTest(const FString& Parameters)
+{
+	FPCGGPUGraphTestRunner Runner;
+	UPCGNode* GatherNode = nullptr;
+	const PCGBranchTestGraphGPU::FNodes Nodes = PCGBranchTestGraphGPU::Build(Runner.Graph);
+	GatherNode = Nodes.GatherNode;
+
+	// CreateAttributeSet: bool attribute "bOutputToB" = true, fed to the branch's dedicated bOutputToB override pin.
+	UPCGCreateAttributeSetSettings* CreateAttrSettings = nullptr;
+	UPCGNode* CreateAttrNode = Runner.Graph->AddNodeOfType<UPCGCreateAttributeSetSettings>(CreateAttrSettings);
+	CreateAttrSettings->AttributeTypes.Type = EPCGMetadataTypes::Boolean;
+	CreateAttrSettings->AttributeTypes.BoolValue = true;
+	CreateAttrSettings->OutputTarget.SetAttributeName(GET_MEMBER_NAME_CHECKED(UPCGBranchSettings, bOutputToB));
+
+	Runner.Graph->AddEdge(CreateAttrNode, PCGPinConstants::DefaultOutputLabel, Nodes.BranchNode, GET_MEMBER_NAME_CHECKED(UPCGBranchSettings, bOutputToB));
+
+	TArray<FPCGDataCollection> OutNodeData;
+	if (!Runner.Execute({ this, { GatherNode } }, OutNodeData))
+	{
+		return false;
+	}
+
+	return PCGTestsCommon::ValidateSingleTaggedData(this, OutNodeData[0],
+		[&](const FPCGTaggedData& TaggedData) { return PCGBranchTestGraphGPU::ValidateOutput(this, TaggedData, PCGBranchTestGraphGPU::OutputBLocation); });
+}
+
+// -----------------------------------------------------------------------------------------
+// Execution dependency culling tests
+//
+// A Branch node's Output A is connected to the execution dependency pin of a CustomHLSL
+// PointProcessor. With bExecutionDependencyRequired=true on the PointProcessor and
+// bOutputToB=true on the Branch, Output A is culled, which transitively culls the
+// PointProcessor. The Gather node (fed only by the PointProcessor) produces 0 data items.
+//
+// Graph:
+//   CreatePoints (1 point) ------------------------------------> CustomHLSL (X += 50) -> Gather
+//                                                              ^ ExecDep (Required)
+//   CreateAttributeSet -> Branch (bOutputToB=true) -> OutputA -+
+//
+// Static test: bOutputToB=true is baked at construction time.
+//   IsPinStaticallyActive(OutputA)=false -> PointProcessor's required execution dependency
+//   pin has no active upstream -> PointProcessor statically culled before any GPU dispatch.
+//
+// Dynamic test: bOutputToB=true is fed at runtime via the dedicated bOutputToB override pin.
+//   IsPinStaticallyActive(OutputA)=true at compile time (branch direction unknown), so the
+//   PointProcessor is compiled into the task graph. At runtime Branch deactivates OutputA,
+//   triggering dynamic culling of the PointProcessor via the inactive execution dependency.
+// -----------------------------------------------------------------------------------------
+
+namespace PCGBranchExecDepTestGraphGPU
+{
+	struct FNodes
+	{
+		UPCGNode* BranchNode = nullptr;
+		UPCGNode* HLSLNode   = nullptr;
+		UPCGNode* GatherNode = nullptr;
+	};
+
+	static FNodes Build(UPCGGraph* Graph, bool bOutputToB = false)
+	{
+		// CreatePoints: one default point at the origin.
+		UPCGCreatePointsSettings* CreatePointsSettings = nullptr;
+		UPCGNode* CreatePointsNode = Graph->AddNodeOfType<UPCGCreatePointsSettings>(CreatePointsSettings);
+		CreatePointsSettings->PointsToCreate.Empty();
+		CreatePointsSettings->PointsToCreate.AddDefaulted();
+		CreatePointsSettings->CoordinateSpace = EPCGCoordinateSpace::World;
+
+		// CustomHLSL PointProcessor: translate X += 50.
+		// bExecutionDependencyRequired=true makes the execution dependency pin Required (not
+		// Advanced), so the static culling analysis considers it and culls this node when
+		// Branch.OutputA is statically inactive.
+		UPCGCustomHLSLSettings* HLSLSettings = nullptr;
+		UPCGNode* HLSLNode = Graph->AddNodeOfType<UPCGCustomHLSLSettings>(HLSLSettings);
+		HLSLSettings->KernelType = EPCGKernelType::PointProcessor;
+		HLSLSettings->SetSourceText(
+			TEXT("float3 Pos = In_GetPosition(In_DataIndex, ElementIndex);\n")
+			TEXT("Pos.x += 50.0f;\n")
+			TEXT("Out_SetPosition(Out_DataIndex, ElementIndex, Pos);\n")
+		);
+		HLSLSettings->bExecutionDependencyRequired = true;
+		// Pin status (Advanced vs Required) is baked during UpdatePins(), which ran when AddNodeOfType
+		// was called above (before bExecutionDependencyRequired was set). Refresh now so the ExecDep
+		// input pin carries EPCGPinStatus::Required when the graph is compiled.
+		HLSLNode->UpdateAfterSettingsChangeDuringCreation();
+
+		// Branch: condition wired by the caller.
+		UPCGBranchSettings* BranchSettings = nullptr;
+		UPCGNode* BranchNode = Graph->AddNodeOfType<UPCGBranchSettings>(BranchSettings);
+		BranchSettings->bOutputToB = bOutputToB;
+
+		// Gather: reads the PointProcessor output.
+		UPCGGatherSettings* GatherSettings = nullptr;
+		UPCGNode* GatherNode = Graph->AddNodeOfType<UPCGGatherSettings>(GatherSettings);
+
+		// CreateAttributeSet: provides data to Branch so it is not culled by the unwired-input check.
+		UPCGCreateAttributeSetSettings* BranchFeedSettings = nullptr;
+		UPCGNode* BranchFeedNode = Graph->AddNodeOfType<UPCGCreateAttributeSetSettings>(BranchFeedSettings);
+		BranchFeedSettings->AttributeTypes.Type = EPCGMetadataTypes::Boolean;
+
+		// Data path: CreatePoints -> CustomHLSL -> Gather.
+		Graph->AddEdge(CreatePointsNode, PCGPinConstants::DefaultOutputLabel, HLSLNode,   PCGPinConstants::DefaultInputLabel);
+		Graph->AddEdge(HLSLNode,         PCGPinConstants::DefaultOutputLabel, GatherNode, PCGPinConstants::DefaultInputLabel);
+
+		// Control path: Branch.OutputA -> CustomHLSL execution dependency pin.
+		Graph->AddEdge(BranchFeedNode, PCGPinConstants::DefaultOutputLabel, BranchNode, PCGPinConstants::DefaultInputLabel);
+		Graph->AddEdge(BranchNode,     PCGBranchConstants::OutputLabelA,    HLSLNode,   PCGPinConstants::DefaultExecutionDependencyLabel);
+
+		return { BranchNode, HLSLNode, GatherNode };
+	}
+}
+
+/**
+ * Static branch (GPU): bOutputToB=true culls OutputA at compile time, statically culling the
+ * PointProcessor via its required execution dependency pin. The PointProcessor never executes.
+ */
+IMPLEMENT_CUSTOM_SIMPLE_AUTOMATION_TEST(
+	FPCGBranchExecDepStaticTest_Culled_GPU,
+	FPCGTestBaseClass,
+	"Plugins.PCG.Branch.ExecDep.Static.Culled.GPU",
+	PCGTestsCommon::TestFlags | EAutomationTestFlags::NonNullRHI)
+
+bool FPCGBranchExecDepStaticTest_Culled_GPU::RunTest(const FString& Parameters)
+{
+	FPCGGPUGraphTestRunner Runner;
+	UPCGNode* HLSLNode = nullptr;
+	HLSLNode = PCGBranchExecDepTestGraphGPU::Build(Runner.Graph, /*bOutputToB=*/true).HLSLNode;
+
+	// Branch.OutputA is statically inactive (bOutputToB=true), so the PointProcessor's required
+	// execution dependency pin has no active upstream -- it must be statically culled.
+	TArray<FPCGDataCollection> OutNodeData;
+	return Runner.Execute({ this, {}, /*bReadbackGPUData=*/false, { HLSLNode } }, OutNodeData);
+}
+
+/**
+ * Dynamic branch (GPU): bOutputToB=true (via dedicated override pin) culls OutputA at runtime,
+ * dynamically culling the PointProcessor via its execution dependency pin.
+ * Gather produces 0 data items.
+ */
+IMPLEMENT_CUSTOM_SIMPLE_AUTOMATION_TEST(
+	FPCGBranchExecDepDynamicTest_Culled_GPU,
+	FPCGTestBaseClass,
+	"Plugins.PCG.Branch.ExecDep.Dynamic.Culled.GPU",
+	PCGTestsCommon::TestFlags | EAutomationTestFlags::NonNullRHI)
+
+bool FPCGBranchExecDepDynamicTest_Culled_GPU::RunTest(const FString& Parameters)
+{
+	FPCGGPUGraphTestRunner Runner;
+	UPCGNode* GatherNode = nullptr;
+	const PCGBranchExecDepTestGraphGPU::FNodes Nodes = PCGBranchExecDepTestGraphGPU::Build(Runner.Graph);
+	GatherNode = Nodes.GatherNode;
+
+	// CreateAttributeSet: bool attribute "bOutputToB" = true, fed to Branch's dedicated bOutputToB override pin.
+	UPCGCreateAttributeSetSettings* CreateAttrSettings = nullptr;
+	UPCGNode* CreateAttrNode = Runner.Graph->AddNodeOfType<UPCGCreateAttributeSetSettings>(CreateAttrSettings);
+	CreateAttrSettings->AttributeTypes.Type = EPCGMetadataTypes::Boolean;
+	CreateAttrSettings->AttributeTypes.BoolValue = true;
+	CreateAttrSettings->OutputTarget.SetAttributeName(GET_MEMBER_NAME_CHECKED(UPCGBranchSettings, bOutputToB));
+
+	Runner.Graph->AddEdge(CreateAttrNode, PCGPinConstants::DefaultOutputLabel, Nodes.BranchNode, GET_MEMBER_NAME_CHECKED(UPCGBranchSettings, bOutputToB));
+
+	// Branch.OutputA is dynamically deactivated (bOutputToB=true at runtime), which culls the
+	// PointProcessor via its required execution dependency, which should transitively cull Gather.
+	// ExpectCulledNodes asserts the debug callback was NOT fired for Gather -- i.e. it never executed.
+	TArray<FPCGDataCollection> OutNodeData;
+	return Runner.Execute({ this, {}, /*bReadbackGPUData=*/false, { GatherNode } }, OutNodeData);
+}
+
+// -----------------------------------------------------------------------------------------
+// Data pin culling tests
+//
+// A Branch node sits in the data path. Its Output A feeds the default In pin of a
+// CustomHLSL PointProcessor. With bOutputToB=true, Output A is deactivated, leaving the
+// PointProcessor's required In pin without an active upstream -- which culls it.
+// Gather (fed only by the PointProcessor) is transitively culled.
+//
+// Graph:
+//   CreatePoints (1 point) -> Branch (bOutputToB=true) -> OutputA (inactive) -> CustomHLSL.In
+//                                                                                CustomHLSL.Out -> Gather
+//
+// Static test: bOutputToB=true baked at construction time.
+//   IsPinStaticallyActive(OutputA)=false -> PointProcessor's required In pin has no active
+//   upstream -> PointProcessor statically culled before any GPU dispatch.
+//
+// Dynamic test: bOutputToB=true fed at runtime via the dedicated bOutputToB override pin.
+//   IsPinStaticallyActive(OutputA)=true at compile time (branch direction unknown), so the
+//   PointProcessor is compiled in. At runtime Branch deactivates OutputA, triggering dynamic
+//   culling of the PointProcessor via its inactive In pin.
+// -----------------------------------------------------------------------------------------
+
+namespace PCGBranchDataPinTestGraphGPU
+{
+    struct FNodes
+    {
+        UPCGNode* BranchNode    = nullptr;
+        UPCGNode* HLSLNode      = nullptr;
+        UPCGNode* GatherNode    = nullptr;
+    };
+
+    static FNodes Build(UPCGGraph* Graph, bool bOutputToB = false)
+    {
+        // CreatePoints: one default point.
+        UPCGCreatePointsSettings* CreatePointsSettings = nullptr;
+        UPCGNode* CreatePointsNode = Graph->AddNodeOfType<UPCGCreatePointsSettings>(CreatePointsSettings);
+        CreatePointsSettings->PointsToCreate.Empty();
+        CreatePointsSettings->PointsToCreate.AddDefaulted();
+        CreatePointsSettings->CoordinateSpace = EPCGCoordinateSpace::World;
+
+        // Branch: sits in the data path. bOutputToB routes away from OutputA.
+        UPCGBranchSettings* BranchSettings = nullptr;
+        UPCGNode* BranchNode = Graph->AddNodeOfType<UPCGBranchSettings>(BranchSettings);
+        BranchSettings->bOutputToB = bOutputToB;
+
+        // CustomHLSL PointProcessor: translate X += 50.
+        // No exec dep pin -- culling is driven purely by the inactive In data pin.
+        UPCGCustomHLSLSettings* HLSLSettings = nullptr;
+        UPCGNode* HLSLNode = Graph->AddNodeOfType<UPCGCustomHLSLSettings>(HLSLSettings);
+        HLSLSettings->SetKernelType(EPCGKernelType::PointProcessor);
+        HLSLSettings->SetSourceText(
+            TEXT("float3 Pos = In_GetPosition(In_DataIndex, ElementIndex);\n")
+            TEXT("Pos.x += 50.0f;\n")
+            TEXT("Out_SetPosition(Out_DataIndex, ElementIndex, Pos);\n")
+        );
+
+        // Gather: reads the PointProcessor output.
+        UPCGGatherSettings* GatherSettings = nullptr;
+        UPCGNode* GatherNode = Graph->AddNodeOfType<UPCGGatherSettings>(GatherSettings);
+
+        // Data path: CreatePoints -> Branch -> OutputA -> CustomHLSL.In -> Gather.
+        Graph->AddEdge(CreatePointsNode, PCGPinConstants::DefaultOutputLabel, BranchNode, PCGPinConstants::DefaultInputLabel);
+        Graph->AddEdge(BranchNode,       PCGBranchConstants::OutputLabelA,    HLSLNode,   PCGPinConstants::DefaultInputLabel);
+        Graph->AddEdge(HLSLNode,         PCGPinConstants::DefaultOutputLabel, GatherNode, PCGPinConstants::DefaultInputLabel);
+
+        return { BranchNode, HLSLNode, GatherNode };
+    }
+}
+
+/**
+ * Static branch (GPU): bOutputToB=true culls OutputA at compile time, statically culling the
+ * PointProcessor via its required In data pin (no exec dep pin involved).
+ */
+IMPLEMENT_CUSTOM_SIMPLE_AUTOMATION_TEST(
+    FPCGBranchDataPinStaticTest_Culled_GPU,
+    FPCGTestBaseClass,
+    "Plugins.PCG.Branch.DataPin.Static.Culled.GPU",
+    PCGTestsCommon::TestFlags | EAutomationTestFlags::NonNullRHI)
+
+bool FPCGBranchDataPinStaticTest_Culled_GPU::RunTest(const FString& Parameters)
+{
+    FPCGGPUGraphTestRunner Runner;
+    const PCGBranchDataPinTestGraphGPU::FNodes Nodes = PCGBranchDataPinTestGraphGPU::Build(Runner.Graph, /*bOutputToB=*/true);
+
+    // Branch.OutputA is statically inactive, so the PointProcessor's required In data pin has
+    // no active upstream -- it must be statically culled.
+    TArray<FPCGDataCollection> OutNodeData;
+    return Runner.Execute({ this, {}, /*bReadbackGPUData=*/false, { Nodes.HLSLNode, Nodes.GatherNode } }, OutNodeData);
+}
+
+/**
+ * Dynamic branch (GPU): bOutputToB=true (via dedicated override pin) culls OutputA at runtime,
+ * dynamically culling the PointProcessor via its required In data pin (no exec dep pin involved).
+ * Gather is transitively culled.
+ */
+IMPLEMENT_CUSTOM_SIMPLE_AUTOMATION_TEST(
+    FPCGBranchDataPinDynamicTest_Culled_GPU,
+    FPCGTestBaseClass,
+    "Plugins.PCG.Branch.DataPin.Dynamic.Culled.GPU",
+    PCGTestsCommon::TestFlags | EAutomationTestFlags::NonNullRHI)
+
+bool FPCGBranchDataPinDynamicTest_Culled_GPU::RunTest(const FString& Parameters)
+{
+    FPCGGPUGraphTestRunner Runner;
+    const PCGBranchDataPinTestGraphGPU::FNodes Nodes = PCGBranchDataPinTestGraphGPU::Build(Runner.Graph);
+
+    // CreateAttributeSet: bool attribute "bOutputToB" = true, fed to Branch's dedicated override pin.
+    UPCGCreateAttributeSetSettings* CreateAttrSettings = nullptr;
+    UPCGNode* CreateAttrNode = Runner.Graph->AddNodeOfType<UPCGCreateAttributeSetSettings>(CreateAttrSettings);
+    CreateAttrSettings->AttributeTypes.Type = EPCGMetadataTypes::Boolean;
+    CreateAttrSettings->AttributeTypes.BoolValue = true;
+    CreateAttrSettings->OutputTarget.SetAttributeName(GET_MEMBER_NAME_CHECKED(UPCGBranchSettings, bOutputToB));
+
+    Runner.Graph->AddEdge(CreateAttrNode, PCGPinConstants::DefaultOutputLabel, Nodes.BranchNode, GET_MEMBER_NAME_CHECKED(UPCGBranchSettings, bOutputToB));
+
+    // Branch.OutputA is dynamically deactivated, leaving the PointProcessor's required In data
+    // pin without an active upstream -- triggering dynamic culling of both PointProcessor and Gather.
+    TArray<FPCGDataCollection> OutNodeData;
+    return Runner.Execute({ this, {}, /*bReadbackGPUData=*/false, { Nodes.HLSLNode, Nodes.GatherNode } }, OutNodeData);
+}
+
+#endif // WITH_EDITOR
diff --git a/Engine/Plugins/PCG/Source/PCG/Public/Elements/ControlFlow/PCGBranch.h b/Engine/Plugins/PCG/Source/PCG/Public/Elements/ControlFlow/PCGBranch.h
--- a/Engine/Plugins/PCG/Source/PCG/Public/Elements/ControlFlow/PCGBranch.h
+++ b/Engine/Plugins/PCG/Source/PCG/Public/Elements/ControlFlow/PCGBranch.h
@@ -5,6 +5,12 @@
 
 #include "PCGBranch.generated.h"
 
+namespace PCGBranchConstants
+{
+	extern PCG_API const FName OutputLabelA;
+	extern PCG_API const FName OutputLabelB;
+}
+
 /**
  * Routes input data to one of two outputs, based on a boolean condition.
  */
