Fix use after free when FPCGComputeGraphContext gets destroyed

#rnx
[RN] rnx
#Synced-CL 52718886

diff --git a/Engine/Plugins/PCG/Source/PCG/Private/Compute/Elements/PCGComputeGraphElement.cpp b/Engine/Plugins/PCG/Source/PCG/Private/Compute/Elements/PCGComputeGraphElement.cpp
--- a/Engine/Plugins/PCG/Source/PCG/Private/Compute/Elements/PCGComputeGraphElement.cpp
+++ b/Engine/Plugins/PCG/Source/PCG/Private/Compute/Elements/PCGComputeGraphElement.cpp
@@ -25,6 +25,7 @@
 #include "ComputeFramework/ComputeKernelCompileResult.h"
 #include "Engine/World.h"
 #include "Logging/LogVerbosity.h"
+#include "UObject/GCObject.h"
 
 #if WITH_EDITOR
 #include "Editor.h"
@@ -80,19 +81,52 @@ static TAutoConsoleVariable<bool> CVarReadbackTextureDataOnInspect(
 	TEXT("pcg.GPU.ReadbackTextureDataOnInspect"),
 	true,
 	TEXT("Reads texture data back to the CPU when inspected."));
+
 #else // WITH_EDITOR
 #define PCG_COMPUTE_GRAPH_BREAKPOINT()
 #endif // WITH_EDITOR
 
+static TAutoConsoleVariable<bool> CVarPCGFixProviderGCRaceOnDestroy(
+	TEXT("pcg.GPU.FixProviderGCRaceOnDestroy"),
+	true,
+	TEXT("When enabled, pins UComputeDataProvider objects in GC via an FGCObject while the deferred "
+	     "game-thread cleanup lambda is in flight. Prevents a use-after-free when GC runs between "
+	     "~FPCGComputeGraphContext (worker thread) and the lambda executing (game thread)."));
+
 FPCGComputeGraphContext::~FPCGComputeGraphContext()
 {
 	const FPCGStack* StackPtr = GetStack();
 	TWeakObjectPtr<const UPCGGraph> GraphWeak = StackPtr ? StackPtr->GetNearestNonInlinedGraphForCurrentFrame() : nullptr;
 	TWeakObjectPtr<UPCGDataBinding> DataBindingWeak = MakeWeakObjectPtr(DataBinding);
-	
+
+	// Holds UComputeDataProvider objects alive in GC for the duration of the deferred lambda.
+	// Without this, a GC cycle between the destructor (worker thread) and the lambda (game thread)
+	// can collect the providers — making Provider.Get() return a dangling pointer.
+	struct FProviderGCPinner : public FGCObject
+	{
+		TArray<TObjectPtr<UComputeDataProvider>> Providers;
+
+		virtual void AddReferencedObjects(FReferenceCollector& Collector) override
+		{
+			Collector.AddReferencedObjects(Providers);
+		}
+
+		virtual FString GetReferencerName() const override
+		{
+			return TEXT("FPCGComputeGraphContext::FProviderGCPinner");
+		}
+	};
+
+	TSharedPtr<FProviderGCPinner> GCPinner;
+	if (CVarPCGFixProviderGCRaceOnDestroy.GetValueOnAnyThread() && ComputeGraphInstance.IsValid() && !IsInGameThread())
+	{
+		GCPinner = MakeShared<FProviderGCPinner>();
+		GCPinner->Providers = ComputeGraphInstance->GetDataProviders();
+	}
+
 	// Can be here on a worker thread, after unpinning the context in one of the async tasks.
 	// Will execute immediately if on game thread, otherwise on next game thread frame.
-	PCGHelpers::ExecuteOnGameThread(TEXT("~FPCGComputeGraphContext"), [DataBindingWeak, GraphInstance=ComputeGraphInstance, bInstanceInitialized=bComputeGraphInstanceInitialized, GraphWeak, GridSize=GenerationGridSize, GraphIndex=ComputeGraphIndex]
+	PCGHelpers::ExecuteOnGameThread(TEXT("~FPCGComputeGraphContext"), [DataBindingWeak, GraphInstance=ComputeGraphInstance, bInstanceInitialized=bComputeGraphInstanceInitialized, GraphWeak, GridSize=GenerationGridSize, GraphIndex=ComputeGraphIndex, GCPinner]
 	{
 		LLM_SCOPE_BYTAG(PCG);
 
@@ -123,6 +157,9 @@ FPCGComputeGraphContext::~FPCGComputeGraphContext()
 				}
 			}
 		}
+
+		// GCPinner goes out of scope here - FGCObject unregisters and providers
+		// are released back to GC naturally after this point.
 	});
 }
 
