[PCG] Fix data domain being allowed for multi entries on all spatial data that doesn't support elements

* The default domain was created with support multi entry to true
* Now when we setup our domain, and it is the default, we update the cached value

#rnx
#Synced-CL 52758271

diff --git a/Engine/Plugins/PCG/Source/PCG/Private/Metadata/PCGMetadata.cpp b/Engine/Plugins/PCG/Source/PCG/Private/Metadata/PCGMetadata.cpp
--- a/Engine/Plugins/PCG/Source/PCG/Private/Metadata/PCGMetadata.cpp
+++ b/Engine/Plugins/PCG/Source/PCG/Private/Metadata/PCGMetadata.cpp
@@ -317,6 +317,8 @@ void UPCGMetadata::SetupDomain(FPCGMetadataDomainID DomainID, bool bIsDefault)
 		// Map the default domain to the DefaultID domain and update domain to the not-default domain ID.
 		TSharedPtr<FPCGMetadataDomain> DefaultIDLayer = MetadataDomains[PCGMetadataDomainID::Default];
 		MetadataDomains.Emplace(DomainID, DefaultIDLayer)->DomainID = DomainID;
+		// Make sure to update the cached value for allow for multiple entries, as it might not have been setup correctly.
+		DefaultIDLayer->bSupportMultiEntries = MetadataDomainSupportsMultiEntries(DomainID);
 	}
 	else
 	{
diff --git a/Engine/Plugins/PCG/Source/PCG/Public/Metadata/PCGMetadataDomain.h b/Engine/Plugins/PCG/Source/PCG/Public/Metadata/PCGMetadataDomain.h
--- a/Engine/Plugins/PCG/Source/PCG/Public/Metadata/PCGMetadataDomain.h
+++ b/Engine/Plugins/PCG/Source/PCG/Public/Metadata/PCGMetadataDomain.h
@@ -321,7 +321,7 @@ protected:
 	const FPCGMetadataDomain* Parent = nullptr;
 
 	// Cached value on construction to know if we support multi entries
-	const bool bSupportMultiEntries = true;
+	bool bSupportMultiEntries = true;
 
 	TMap<FName, FPCGMetadataAttributeBase*> Attributes;
 	PCGMetadataAttributeKey NextAttributeId = 0;
